GeometryCore: create new GeometryCore Module. Move GeometricObjects code from GeometryProcessing plugin to this new module, as well as core FDynamicMesh3 classes (attributes, aabbtree/octree, change t

diff --git a/Engine/Plugins/Experimental/GeometryProcessing/Source/DynamicMesh/Public/MeshIndexMappings.h b/Engine/Plugins/Experimental/GeometryProcessing/Source/DynamicMesh/Public/MeshIndexMappings.h
--- a/Engine/Plugins/Experimental/GeometryProcessing/Source/DynamicMesh/Public/MeshIndexMappings.h
+++ b/Engine/Plugins/Experimental/GeometryProcessing/Source/DynamicMesh/Public/MeshIndexMappings.h
@@ -2,91 +2,4 @@
 
 #pragma once
 
-#include "DynamicMesh3.h"
-
-namespace UE
-{
-namespace Geometry
-{
-
-
-/**
- * FMeshIndexMappings stores a set of integer IndexMaps for a mesh
- * This is a convenient object to have, to avoid passing around large numbers of separate maps.
- * The individual maps are not necessarily all filled by every operation.
- */
-struct DYNAMICMESH_API FMeshIndexMappings
-{
-protected:
-	FIndexMapi VertexMap;
-	FIndexMapi TriangleMap;
-	FIndexMapi GroupMap;
-
-	FIndexMapi ColorMap;
-	TArray<FIndexMapi> UVMaps;
-	TArray<FIndexMapi> NormalMaps;
-
-public:
-	/** Size internal arrays-of-maps to be suitable for this Mesh */
-	void Initialize(FDynamicMesh3* Mesh);
-
-	/** @return the value used to indicate "invalid" in the mapping */
-	constexpr int InvalidID() const { return VertexMap.UnmappedID(); }
-
-	void Reset()
-	{
-		VertexMap.Reset();
-		TriangleMap.Reset();
-		GroupMap.Reset();
-		ColorMap.Reset();
-		for (FIndexMapi& UVMap : UVMaps)
-		{
-			UVMap.Reset();
-		}
-		for (FIndexMapi& NormalMap : NormalMaps)
-		{
-			NormalMap.Reset();
-		}
-	}
-
-	void ResetTriangleMap()
-	{
-		TriangleMap.Reset();
-	}
-
-	FIndexMapi& GetVertexMap() { return VertexMap; }
-	const FIndexMapi& GetVertexMap() const { return VertexMap; }
-	inline void SetVertex(int FromID, int ToID) { VertexMap.Add(FromID, ToID); }
-	inline int GetNewVertex(int FromID) const { return VertexMap.GetTo(FromID); }
-	inline bool ContainsVertex(int FromID) const { return VertexMap.ContainsFrom(FromID); }
-
-	FIndexMapi& GetTriangleMap() { return TriangleMap; }
-	const FIndexMapi& GetTriangleMap() const { return TriangleMap; }
-	void SetTriangle(int FromID, int ToID) { TriangleMap.Add(FromID, ToID); }
-	int GetNewTriangle(int FromID) const { return TriangleMap.GetTo(FromID); }
-	inline bool ContainsTriangle(int FromID) const { return TriangleMap.ContainsFrom(FromID); }
-
-	FIndexMapi& GetGroupMap() { return GroupMap; }
-	void SetGroup(int FromID, int ToID) { GroupMap.Add(FromID, ToID); }
-	int GetNewGroup(int FromID) const { return GroupMap.GetTo(FromID); }
-	inline bool ContainsGroup(int FromID) const { return GroupMap.ContainsFrom(FromID); }
-
-	FIndexMapi& GetUVMap(int UVLayer) { return UVMaps[UVLayer]; }
-	void SetUV(int UVLayer, int FromID, int ToID) { UVMaps[UVLayer].Add(FromID, ToID); }
-	int GetNewUV(int UVLayer, int FromID) const { return UVMaps[UVLayer].GetTo(FromID); }
-	inline bool ContainsUV(int UVLayer, int FromID) const { return UVMaps[UVLayer].ContainsFrom(FromID); }
-
-	FIndexMapi& GetNormalMap(int NormalLayer) { return NormalMaps[NormalLayer]; }
-	void SetNormal(int NormalLayer, int FromID, int ToID) { NormalMaps[NormalLayer].Add(FromID, ToID); }
-	int GetNewNormal(int NormalLayer, int FromID) const { return NormalMaps[NormalLayer].GetTo(FromID); }
-	inline bool ContainsNormal(int NormalLayer, int FromID) const { return NormalMaps[NormalLayer].ContainsFrom(FromID); }
-
-	FIndexMapi& GetColorMap() { return ColorMap; }
-	void SetColor( int FromID, int ToID) { ColorMap.Add(FromID, ToID); }
-	int GetNewColor(int FromID) const { return ColorMap.GetTo(FromID); }
-	inline bool ContainsColor(int FromID) const { return ColorMap.ContainsFrom(FromID); }
-
-};
-
-} // end namespace UE::Geometry
-} // end namespace UE
\ No newline at end of file
+#include "DynamicMesh/MeshIndexMappings.h"
\ No newline at end of file
diff --git a/Engine/Plugins/Experimental/GeometryProcessing/Source/GeometricObjects/GeometricObjects.Build.cs b/Engine/Plugins/Experimental/GeometryProcessing/Source/GeometricObjects/GeometricObjects.Build.cs
--- a/Engine/Plugins/Experimental/GeometryProcessing/Source/GeometricObjects/GeometricObjects.Build.cs
+++ b/Engine/Plugins/Experimental/GeometryProcessing/Source/GeometricObjects/GeometricObjects.Build.cs
@@ -14,7 +14,8 @@ public class GeometricObjects : ModuleRules
 		// (which won't have UObject garbage collection).
 		PublicDependencyModuleNames.AddRange(
 			new string[] {
-				"Core"
+				"Core",
+				"GeometryCore"
 			}
 			);
 	}
diff --git a/Engine/Plugins/Experimental/GeometryProcessing/Source/GeometryAlgorithms/GeometryAlgorithms.Build.cs b/Engine/Plugins/Experimental/GeometryProcessing/Source/GeometryAlgorithms/GeometryAlgorithms.Build.cs
--- a/Engine/Plugins/Experimental/GeometryProcessing/Source/GeometryAlgorithms/GeometryAlgorithms.Build.cs
+++ b/Engine/Plugins/Experimental/GeometryProcessing/Source/GeometryAlgorithms/GeometryAlgorithms.Build.cs
@@ -15,7 +15,7 @@ public class GeometryAlgorithms : ModuleRules
         PublicDependencyModuleNames.AddRange(
 			new string[] {
 				"Core",
-				"GeometricObjects"
+				"GeometryCore"
 			}
 			);
 	}
diff --git a/Engine/Plugins/Experimental/UVEditor/Source/UVEditorTools/UVEditorTools.Build.cs b/Engine/Plugins/Experimental/UVEditor/Source/UVEditorTools/UVEditorTools.Build.cs
--- a/Engine/Plugins/Experimental/UVEditor/Source/UVEditorTools/UVEditorTools.Build.cs
+++ b/Engine/Plugins/Experimental/UVEditor/Source/UVEditorTools/UVEditorTools.Build.cs
@@ -47,6 +47,7 @@ public class UVEditorTools : ModuleRules
 					
 				"MeshConversion",
 				"DynamicMesh",
+				"GeometryCore",
 				"ModelingComponents",
 				// ... add private dependencies that you statically link with here ...	
 			}
diff --git a/Engine/Source/Runtime/GeometryCore/GeometryCore.Build.cs b/Engine/Source/Runtime/GeometryCore/GeometryCore.Build.cs
new file mode 100644
--- /dev/null
+++ b/Engine/Source/Runtime/GeometryCore/GeometryCore.Build.cs
@@ -0,0 +1,26 @@
+// Copyright Epic Games, Inc. All Rights Reserved.
+
+using UnrealBuildTool;
+
+public class GeometryCore : ModuleRules
+{
+	public GeometryCore(ReadOnlyTargetRules Target) : base(Target)
+	{
+        PCHUsage = ModuleRules.PCHUsageMode.UseExplicitOrSharedPCHs;
+
+		// For GPUSkinPublicDefs.h
+		PublicIncludePaths.Add("Runtime/Engine/Public");
+
+		PublicIncludePathModuleNames.AddRange(
+			new string[] {
+				"AnimationCore",			// For the BoneWeights.h include
+			}
+		);
+
+		PublicDependencyModuleNames.AddRange(
+			new string[] {
+				"Core"
+			}
+		);
+	}
+}
diff --git a/Engine/Source/Runtime/GeometryCore/Private/CompGeom/PolygonTriangulation.cpp b/Engine/Source/Runtime/GeometryCore/Private/CompGeom/PolygonTriangulation.cpp
new file mode 100644
--- /dev/null
+++ b/Engine/Source/Runtime/GeometryCore/Private/CompGeom/PolygonTriangulation.cpp
@@ -0,0 +1,372 @@
+// Copyright Epic Games, Inc. All Rights Reserved.
+
+#include "CompGeom/PolygonTriangulation.h"
+#include "TriangleTypes.h"
+
+#include "ExplicitUseGeometryMathTypes.h"		// using UE::Geometry::(math types)
+using namespace UE::Geometry;
+
+// explicit instantiations
+namespace PolygonTriangulation
+{
+	template GEOMETRYCORE_API void TriangulateSimplePolygon<float>(const TArray<FVector2<float>>& VertexPositions, TArray<FIndex3i>& OutTriangles);
+	template GEOMETRYCORE_API void TriangulateSimplePolygon<double>(const TArray<FVector2<double>>& VertexPositions, TArray<FIndex3i>& OutTriangles);
+
+	template GEOMETRYCORE_API void ComputePolygonPlane<float>(const TArray<FVector3<float>>& VertexPositions, FVector3<float>& NormalOut, FVector3<float>& PlanePointOut);
+	template GEOMETRYCORE_API void ComputePolygonPlane<double>(const TArray<FVector3<double>>& VertexPositions, FVector3<double>& NormalOut, FVector3<double>& PlanePointOut);
+
+	template GEOMETRYCORE_API void TriangulateSimplePolygon<float>(const TArray<FVector3<float>>& VertexPositions, TArray<FIndex3i>& OutTriangles);
+	template GEOMETRYCORE_API void TriangulateSimplePolygon<double>(const TArray<FVector3<double>>& VertexPositions, TArray<FIndex3i>& OutTriangles);
+}
+
+
+
+
+//
+// Triangulate using ear clipping
+// This is based on the 3D triangulation code from MeshDescription.cpp, simplified for 2D polygons
+// 
+template<typename T>
+void PolygonTriangulation::TriangulateSimplePolygon(const TArray<FVector2<T>>& VertexPositions, TArray<FIndex3i>& OutTriangles)
+{
+	struct Local
+	{
+		static inline bool IsTriangleFlipped(T OrientationSign, const FVector2<T>& VertexPositionA, const FVector2<T>& VertexPositionB, const FVector2<T>& VertexPositionC)
+		{
+			T TriSignedArea = TTriangle2<T>::SignedArea(VertexPositionA, VertexPositionB, VertexPositionC);
+			return TriSignedArea * OrientationSign < 0;
+		}
+	};
+
+	// Polygon must have at least three vertices/edges
+	int32 PolygonVertexCount = VertexPositions.Num();
+	check(PolygonVertexCount >= 3);
+
+
+	// compute signed area of polygon
+	double PolySignedArea = 0;
+	for (int i = 0; i < PolygonVertexCount; ++i)
+	{
+		const FVector2<T>& v1 = VertexPositions[i];
+		const FVector2<T>& v2 = VertexPositions[(i + 1) % PolygonVertexCount];
+		PolySignedArea += v1.X*v2.Y - v1.Y*v2.X;
+	}
+	PolySignedArea *= 0.5;
+	bool bIsClockwise = PolySignedArea < 0;
+	double OrientationSign = (bIsClockwise) ? -1.0 : 1.0;
+
+
+	OutTriangles.Reset();
+
+
+	// If perimeter has 3 vertices, just copy content of perimeter out 
+	if (PolygonVertexCount == 3)
+	{
+		OutTriangles.Add(FIndex3i(0, 1, 2));
+		return;
+	}
+
+	// Make a simple linked list array of the previous and next vertex numbers, for each vertex number
+	// in the polygon.  This will just save us having to iterate later on.
+	TArray<int32> PrevVertexNumbers, NextVertexNumbers;
+
+	PrevVertexNumbers.SetNumUninitialized(PolygonVertexCount, false);
+	NextVertexNumbers.SetNumUninitialized(PolygonVertexCount, false);
+
+	for (int32 VertexNumber = 0; VertexNumber < PolygonVertexCount; ++VertexNumber)
+	{
+		PrevVertexNumbers[VertexNumber] = VertexNumber - 1;
+		NextVertexNumbers[VertexNumber] = VertexNumber + 1;
+	}
+	PrevVertexNumbers[0] = PolygonVertexCount - 1;
+	NextVertexNumbers[PolygonVertexCount - 1] = 0;
+
+
+	int32 EarVertexNumber = 0;
+	int32 EarTestCount = 0;
+	for (int32 RemainingVertexCount = PolygonVertexCount; RemainingVertexCount >= 3; )
+	{
+		bool bIsEar = true;
+
+		// If we're down to only a triangle, just treat it as an ear.  Also, if we've tried every possible candidate
+		// vertex looking for an ear, go ahead and just treat the current vertex as an ear.  This can happen when 
+		// vertices are collinear or other degenerate cases.
+		if (RemainingVertexCount > 3 && EarTestCount < RemainingVertexCount)
+		{
+			const FVector2<T>& PrevVertexPosition = VertexPositions[PrevVertexNumbers[EarVertexNumber]];
+			const FVector2<T>& EarVertexPosition = VertexPositions[EarVertexNumber];
+			const FVector2<T>& NextVertexPosition = VertexPositions[NextVertexNumbers[EarVertexNumber]];
+
+			// Figure out whether the potential ear triangle is facing the same direction as the polygon
+			// itself.  If it's facing the opposite direction, then we're dealing with a concave triangle
+			// and we'll skip it for now.
+			if (!Local::IsTriangleFlipped(
+				OrientationSign, PrevVertexPosition, EarVertexPosition, NextVertexPosition))
+			{
+				int32 TestVertexNumber = NextVertexNumbers[NextVertexNumbers[EarVertexNumber]];
+
+				do
+				{
+					// Test every other remaining vertex to make sure that it doesn't lie inside our potential ear
+					// triangle.  If we find a vertex that's inside the triangle, then it cannot actually be an ear.
+					const FVector2<T>& TestVertexPosition = VertexPositions[TestVertexNumber];
+					if (TTriangle2<T>::IsInside(PrevVertexPosition, EarVertexPosition, NextVertexPosition, TestVertexPosition))
+					{
+						bIsEar = false;
+						break;
+					}
+
+					TestVertexNumber = NextVertexNumbers[TestVertexNumber];
+				} while (TestVertexNumber != PrevVertexNumbers[EarVertexNumber]);
+			}
+			else
+			{
+				bIsEar = false;
+			}
+		}
+
+		if (bIsEar)
+		{
+			// OK, we found an ear!  Let's save this triangle in our output buffer.
+			{
+				OutTriangles.Emplace();
+				FIndex3i& Triangle = OutTriangles.Last();
+				Triangle.A = PrevVertexNumbers[EarVertexNumber];
+				Triangle.B = EarVertexNumber;
+				Triangle.C = NextVertexNumbers[EarVertexNumber];
+			}
+
+			// Update our linked list.  We're effectively cutting off the ear by pointing the ear vertex's neighbors to
+			// point at their next sequential neighbor, and reducing the remaining vertex count by one.
+			{
+				NextVertexNumbers[PrevVertexNumbers[EarVertexNumber]] = NextVertexNumbers[EarVertexNumber];
+				PrevVertexNumbers[NextVertexNumbers[EarVertexNumber]] = PrevVertexNumbers[EarVertexNumber];
+				--RemainingVertexCount;
+			}
+
+			// Move on to the previous vertex in the list, now that this vertex was cut
+			EarVertexNumber = PrevVertexNumbers[EarVertexNumber];
+
+			EarTestCount = 0;
+		}
+		else
+		{
+			// The vertex is not the ear vertex, because it formed a triangle that either had a normal which pointed in the opposite direction
+			// of the polygon, or at least one of the other polygon vertices was found to be inside the triangle.  Move on to the next vertex.
+			EarVertexNumber = NextVertexNumbers[EarVertexNumber];
+
+			// Keep track of how many ear vertices we've tested, so that if we exhaust all remaining vertices, we can
+			// fall back to clipping the triangle and adding it to our mesh anyway.  This is important for degenerate cases.
+			++EarTestCount;
+		}
+	}
+
+	check(OutTriangles.Num() > 0);
+}
+
+
+
+
+
+template<typename T>
+void PolygonTriangulation::ComputePolygonPlane(const TArray<FVector3<T>>& VertexPositions, FVector3<T>& PlaneNormalOut, FVector3<T>& PlanePointOut)
+{
+	// NOTE: This polygon plane computation code is partially based on the implementation of "Newell's method" from Real-Time
+	//       Collision Detection by Christer Ericson, published by Morgan Kaufmann Publishers, (c) 2005 Elsevier Inc
+
+	PlaneNormalOut = FVector3<T>::Zero();
+	PlanePointOut = FVector3<T>::Zero();
+
+	int32 NumVertices = VertexPositions.Num();
+	if (NumVertices == 3)
+	{
+		PlaneNormalOut = VectorUtil::Normal(VertexPositions[0], VertexPositions[1], VertexPositions[2]);
+		PlanePointOut = (VertexPositions[0] + VertexPositions[1] + VertexPositions[2]) / (T)3;
+	}
+
+	// Use 'Newell's Method' to compute a robust 'best fit' plane from the vertices of this polygon
+	for (int32 VertexNumberI = NumVertices - 1, VertexNumberJ = 0; VertexNumberJ < NumVertices; VertexNumberI = VertexNumberJ, VertexNumberJ++)
+	{
+		const FVector3<T>& PositionI = VertexPositions[VertexNumberI];
+		const FVector3<T>& PositionJ = VertexPositions[VertexNumberJ];
+		PlanePointOut += PositionJ;
+		PlaneNormalOut.X += (PositionJ.Y - PositionI.Y) * (PositionI.Z + PositionJ.Z);
+		PlaneNormalOut.Y += (PositionJ.Z - PositionI.Z) * (PositionI.X + PositionJ.X);
+		PlaneNormalOut.Z += (PositionJ.X - PositionI.X) * (PositionI.Y + PositionJ.Y);
+	}
+	Normalize(PlaneNormalOut);
+	PlanePointOut /= (T)NumVertices;
+}
+
+
+
+
+
+
+/** Given three direction vectors, indicates if A and B are on the same 'side' of Vec. */
+static bool VectorsOnSameSide(const FVector& Vec, const FVector& A, const FVector& B, const float SameSideDotProductEpsilon)
+{
+	const FVector CrossA = FVector::CrossProduct(Vec, A);
+	const FVector CrossB = FVector::CrossProduct(Vec, B);
+	float DotWithEpsilon = SameSideDotProductEpsilon + FVector::DotProduct(CrossA, CrossB);
+	return !FMath::IsNegativeFloat(DotWithEpsilon);
+}
+
+
+/** Util to see if P lies within triangle created by A, B and C. */
+static bool PointInTriangle(const FVector& A, const FVector& B, const FVector& C, const FVector& P, const float InsideTriangleDotProductEpsilon)
+{
+	// Cross product indicates which 'side' of the vector the point is on
+	// If its on the same side as the remaining vert for all edges, then its inside.	
+	return (VectorsOnSameSide(B - A, P - A, C - A, InsideTriangleDotProductEpsilon) &&
+		VectorsOnSameSide(C - B, P - B, A - B, InsideTriangleDotProductEpsilon) &&
+		VectorsOnSameSide(A - C, P - C, B - C, InsideTriangleDotProductEpsilon));
+}
+
+
+//
+// Triangulate using ear clipping
+// This is based on the 3D triangulation code from MeshDescription.cpp, simplified for 2D polygons
+// 
+template<typename T>
+void PolygonTriangulation::TriangulateSimplePolygon(const TArray<FVector3<T>>& VertexPositions, TArray<FIndex3i>& OutTriangles)
+{
+	struct Local3
+	{
+		static bool IsTriangleFlipped(const FVector3<T>& ReferenceNormal, const FVector3<T>& VertexPositionA, const FVector3<T>& VertexPositionB, const FVector3<T>& VertexPositionC)
+		{
+			FVector3<T> TriangleNormal = VectorUtil::Normal(VertexPositionA, VertexPositionB, VertexPositionC);
+			return TriangleNormal.Dot(ReferenceNormal) <= (T)0;
+		}
+
+		static bool VectorsOnSameSide(const FVector3<T>& Vec, const FVector3<T>& A, const FVector3<T>& B, T SameSideDotProductEpsilon)
+		{
+			FVector3<T> CrossA = Vec.Cross(A);
+			FVector3<T> CrossB = Vec.Cross(B);
+			T DotWithEpsilon = SameSideDotProductEpsilon + CrossA.Dot(CrossB);
+			return DotWithEpsilon >= 0; // !FMath::IsNegativeFloat(DotWithEpsilon);
+		}
+
+		static bool PointInTriangle(const FVector3<T>& A, const FVector3<T>& B, const FVector3<T>& C, const FVector3<T>& P, const T InsideTriangleDotProductEpsilon)
+		{
+			return (VectorsOnSameSide(B - A, P - A, C - A, InsideTriangleDotProductEpsilon) &&
+				VectorsOnSameSide(C - B, P - B, A - B, InsideTriangleDotProductEpsilon) &&
+				VectorsOnSameSide(A - C, P - C, B - C, InsideTriangleDotProductEpsilon));
+		}
+	};
+
+	// Polygon must have at least three vertices/edges
+	int32 PolygonVertexCount = VertexPositions.Num();
+	check(PolygonVertexCount >= 3);
+
+	OutTriangles.Reset();
+
+	// If perimeter has 3 vertices, just copy content of perimeter out 
+	if (PolygonVertexCount == 3)
+	{
+		OutTriangles.Add(FIndex3i(0, 1, 2));
+		return;
+	}
+
+	// First figure out the polygon normal.  We need this to determine which triangles are convex, so that
+	// we can figure out which ears to clip
+	FVector3<T> PolygonNormal, PolygonCentroid;
+	ComputePolygonPlane(VertexPositions, PolygonNormal, PolygonCentroid);
+
+	// Make a simple linked list array of the previous and next vertex numbers, for each vertex number
+	// in the polygon.  This will just save us having to iterate later on.
+	TArray<int32> PrevVertexNumbers, NextVertexNumbers;
+
+	PrevVertexNumbers.SetNumUninitialized(PolygonVertexCount, false);
+	NextVertexNumbers.SetNumUninitialized(PolygonVertexCount, false);
+
+	for (int32 VertexNumber = 0; VertexNumber < PolygonVertexCount; ++VertexNumber)
+	{
+		PrevVertexNumbers[VertexNumber] = VertexNumber - 1;
+		NextVertexNumbers[VertexNumber] = VertexNumber + 1;
+	}
+	PrevVertexNumbers[0] = PolygonVertexCount - 1;
+	NextVertexNumbers[PolygonVertexCount - 1] = 0;
+
+
+	int32 EarVertexNumber = 0;
+	int32 EarTestCount = 0;
+	for (int32 RemainingVertexCount = PolygonVertexCount; RemainingVertexCount >= 3; )
+	{
+		bool bIsEar = true;
+
+		// If we're down to only a triangle, just treat it as an ear.  Also, if we've tried every possible candidate
+		// vertex looking for an ear, go ahead and just treat the current vertex as an ear.  This can happen when 
+		// vertices are collinear or other degenerate cases.
+		if (RemainingVertexCount > 3 && EarTestCount < RemainingVertexCount)
+		{
+			const FVector3<T>& PrevVertexPosition = VertexPositions[PrevVertexNumbers[EarVertexNumber]];
+			const FVector3<T>& EarVertexPosition = VertexPositions[EarVertexNumber];
+			const FVector3<T>& NextVertexPosition = VertexPositions[NextVertexNumbers[EarVertexNumber]];
+
+			// Figure out whether the potential ear triangle is facing the same direction as the polygon
+			// itself.  If it's facing the opposite direction, then we're dealing with a concave triangle
+			// and we'll skip it for now.
+			if (!Local3::IsTriangleFlipped(
+				PolygonNormal, PrevVertexPosition, EarVertexPosition, NextVertexPosition))
+			{
+				int32 TestVertexNumber = NextVertexNumbers[NextVertexNumbers[EarVertexNumber]];
+
+				do
+				{
+					// Test every other remaining vertex to make sure that it doesn't lie inside our potential ear
+					// triangle.  If we find a vertex that's inside the triangle, then it cannot actually be an ear.
+					const FVector3<T>& TestVertexPosition = VertexPositions[TestVertexNumber];
+					if (Local3::PointInTriangle(PrevVertexPosition, EarVertexPosition, NextVertexPosition, TestVertexPosition, SMALL_NUMBER))
+					{
+						bIsEar = false;
+						break;
+					}
+
+					TestVertexNumber = NextVertexNumbers[TestVertexNumber];
+				} while (TestVertexNumber != PrevVertexNumbers[EarVertexNumber]);
+			}
+			else
+			{
+				bIsEar = false;
+			}
+		}
+
+		if (bIsEar)
+		{
+			// OK, we found an ear!  Let's save this triangle in our output buffer.
+			{
+				OutTriangles.Emplace();
+				FIndex3i& Triangle = OutTriangles.Last();
+				Triangle.A = PrevVertexNumbers[EarVertexNumber];
+				Triangle.B = EarVertexNumber;
+				Triangle.C = NextVertexNumbers[EarVertexNumber];
+			}
+
+			// Update our linked list.  We're effectively cutting off the ear by pointing the ear vertex's neighbors to
+			// point at their next sequential neighbor, and reducing the remaining vertex count by one.
+			{
+				NextVertexNumbers[PrevVertexNumbers[EarVertexNumber]] = NextVertexNumbers[EarVertexNumber];
+				PrevVertexNumbers[NextVertexNumbers[EarVertexNumber]] = PrevVertexNumbers[EarVertexNumber];
+				--RemainingVertexCount;
+			}
+
+			// Move on to the previous vertex in the list, now that this vertex was cut
+			EarVertexNumber = PrevVertexNumbers[EarVertexNumber];
+
+			EarTestCount = 0;
+		}
+		else
+		{
+			// The vertex is not the ear vertex, because it formed a triangle that either had a normal which pointed in the opposite direction
+			// of the polygon, or at least one of the other polygon vertices was found to be inside the triangle.  Move on to the next vertex.
+			EarVertexNumber = NextVertexNumbers[EarVertexNumber];
+
+			// Keep track of how many ear vertices we've tested, so that if we exhaust all remaining vertices, we can
+			// fall back to clipping the triangle and adding it to our mesh anyway.  This is important for degenerate cases.
+			++EarTestCount;
+		}
+	}
+
+	check(OutTriangles.Num() > 0);
+}
\ No newline at end of file
diff --git a/Engine/Source/Runtime/GeometryCore/Public/CompGeom/PolygonTriangulation.h b/Engine/Source/Runtime/GeometryCore/Public/CompGeom/PolygonTriangulation.h
new file mode 100644
--- /dev/null
+++ b/Engine/Source/Runtime/GeometryCore/Public/CompGeom/PolygonTriangulation.h
@@ -0,0 +1,33 @@
+// Copyright Epic Games, Inc. All Rights Reserved.
+
+#pragma once
+
+#include "VectorTypes.h"
+#include "IndexTypes.h"
+
+namespace PolygonTriangulation
+{
+	using namespace UE::Geometry;
+
+	/**
+	 * Compute triangulation of simple 2D polygon using ear-clipping
+	 * @param VertexPositions ordered vertices of 2D polygon
+	 * @param OutTriangles computed triangulation. Each triangle is a tuple of indices into VertexPositions.
+	 */
+	template<typename RealType>
+	void GEOMETRYCORE_API TriangulateSimplePolygon(const TArray<FVector2<RealType>>& VertexPositions, TArray<FIndex3i>& OutTriangles);
+
+
+	template<typename RealType>
+	void GEOMETRYCORE_API ComputePolygonPlane(const TArray<FVector3<RealType>>& VertexPositions, FVector3<RealType>& PlaneNormalOut, FVector3<RealType>& PlanePointOut );
+
+
+	/**
+	 * Compute triangulation of 3D simple polygon using ear-clipping
+	 * @param VertexPositions ordered vertices of 3D polygon
+	 * @param OutTriangles computed triangulation. Each triangle is a tuple of indices into VertexPositions.
+	 */
+	template<typename RealType>
+	void GEOMETRYCORE_API TriangulateSimplePolygon(const TArray<FVector3<RealType>>& VertexPositions, TArray<FIndex3i>& OutTriangles);
+
+}
\ No newline at end of file
diff --git a/Engine/Source/Runtime/GeometryCore/Public/ExplicitUseGeometryMathTypes.h b/Engine/Source/Runtime/GeometryCore/Public/ExplicitUseGeometryMathTypes.h
new file mode 100644
--- /dev/null
+++ b/Engine/Source/Runtime/GeometryCore/Public/ExplicitUseGeometryMathTypes.h
@@ -0,0 +1,124 @@
+// Copyright Epic Games, Inc. All Rights Reserved.
+
+#pragma once
+
+//
+// This header exists to simplify usage of the core UE::Geometry:: math types, such as FVector3d/f, FTransform3d/f, etc. 
+// The GeometryProcessing plugin used these type names for several years before the UE Core math library was extended with
+// double-precision support, at which point it was decided to rename FVector as FVector3f/FVector3d, etc, creating a conflict. 
+// Unfortunately the UE Core versions are not a direct replacement for the GeometryProcessing variants, because:
+//   1) they have templated base classes, and significant other template code takes advantage of this. The UE Core Math types currently do not.
+//   2) they use more reliable (but potentially slower) standard-library math functions
+//   3) they have function naming schemes compatible with other standard vector math libraries (which simplifies porting/etc)
+//   4) they have internal member differences in some case (eg UE::Geometry::FTransform3d does not use SSE, Matrix2/3/4 stores as Vectors, etc)
+//
+// As a result the global ::FVector3d and UE::Geometry::FVector3d names are ambiguous. To avoid requiring explicit UE::Geometry:: disambiguation 
+// in all cases, which would make code harder to read and vastly complicate UE4/UE5 merging, in code that primarily uses GeometryProcessing
+// types, we prefer to explicitly specify - via C++ 'using' declarations - that we are using the UE::Geometry:: types at cpp file scope.
+// Rather than add a standard block of declarations into each file, this header can be included **IN THE CPP ONLY** to do so.
+//
+// Note that the header does not contain 'using namespace UE::Geometry;' - we leave that to the cpp file, so the standard boilerplate is:
+//
+//		#include "ExplicitUseGeometryMathTypes.h"		// using UE::Geometry::(math types)
+//		using namespace UE::Geometry;
+//
+// This header **MUST NOT BE INCLUDED IN OTHER HEADERS** as this could result in ambiguous definitions for downstream includers.
+// In headers, the GeometryProcessing core math types can be explicitly 'used' in other namespaces, or in class definitions, but not in global scope.
+//
+
+
+#include "IndexTypes.h"
+#include "VectorTypes.h"
+#include "IntVectorTypes.h"
+#include "MatrixTypes.h"
+#include "Quaternion.h"
+#include "BoxTypes.h"
+#include "IntBoxTypes.h"
+#include "PlaneTypes.h"
+#include "HalfspaceTypes.h"
+#include "LineTypes.h"
+#include "SegmentTypes.h"
+#include "RayTypes.h"
+#include "SphereTypes.h"
+#include "CapsuleTypes.h"
+#include "OrientedBoxTypes.h"
+#include "TriangleTypes.h"
+#include "CircleTypes.h"
+#include "FrameTypes.h"
+#include "TransformTypes.h"
+
+using UE::Geometry::FIndex2i;
+using UE::Geometry::FIndex3i;
+using UE::Geometry::FIndex4i;
+
+//using UE::Geometry::FVector2f;
+//using UE::Geometry::FVector2d;
+using UE::Geometry::FVector2i;
+//using UE::Geometry::FVector3f;
+//using UE::Geometry::FVector3d;
+using UE::Geometry::FVector3i;
+//using UE::Geometry::FVector4f;
+//using UE::Geometry::FVector4d;
+
+using UE::Geometry::FMatrix3f;
+using UE::Geometry::FMatrix3d;
+using UE::Geometry::FMatrix2f;
+using UE::Geometry::FMatrix2d;
+
+using UE::Geometry::FQuaternionf;
+using UE::Geometry::FQuaterniond;
+
+using UE::Geometry::FInterval1f;
+using UE::Geometry::FInterval1d;
+using UE::Geometry::FInterval1i;
+using UE::Geometry::FAxisAlignedBox2f;
+using UE::Geometry::FAxisAlignedBox2d;
+using UE::Geometry::FAxisAlignedBox2i;
+using UE::Geometry::FAxisAlignedBox3f;
+using UE::Geometry::FAxisAlignedBox3d;
+using UE::Geometry::FAxisAlignedBox3i;
+
+using UE::Geometry::FHalfspace3d;
+using UE::Geometry::FHalfspace3f;
+
+//using UE::Geometry::FPlane3d;
+//using UE::Geometry::FPlane3f;
+
+using UE::Geometry::FLine2f;
+using UE::Geometry::FLine2d;
+using UE::Geometry::FLine3f;
+using UE::Geometry::FLine3d;
+
+using UE::Geometry::FSegment2f;
+using UE::Geometry::FSegment2d;
+using UE::Geometry::FSegment3f;
+using UE::Geometry::FSegment3d;
+
+using UE::Geometry::FRay3f;
+using UE::Geometry::FRay3d;
+
+using UE::Geometry::FSphere3d;
+using UE::Geometry::FSphere3f;
+
+using UE::Geometry::FCapsule3d;
+using UE::Geometry::FCapsule3f;
+
+using UE::Geometry::FOrientedBox3d;
+using UE::Geometry::FOrientedBox3f;
+
+using UE::Geometry::FTriangle2d;
+using UE::Geometry::FTriangle2f;
+using UE::Geometry::FTriangle3d;
+using UE::Geometry::FTriangle3f;
+
+using UE::Geometry::FCircle2d;
+using UE::Geometry::FCircle2f;
+using UE::Geometry::FCircle3d;
+using UE::Geometry::FCircle3f;
+
+//using UE::Geometry::FTransform3d;
+//using UE::Geometry::FTransform3f;
+
+using UE::Geometry::FFrame3d;
+using UE::Geometry::FFrame3f;
+
diff --git a/Engine/Source/Runtime/GeometryCore/Public/Generators/LineSegmentGenerators.h b/Engine/Source/Runtime/GeometryCore/Public/Generators/LineSegmentGenerators.h
new file mode 100644
--- /dev/null
+++ b/Engine/Source/Runtime/GeometryCore/Public/Generators/LineSegmentGenerators.h
@@ -0,0 +1,173 @@
+// Copyright Epic Games, Inc. All Rights Reserved.
+
+#pragma once
+
+#include "CoreMinimal.h"
+#include "VectorTypes.h"
+#include "FrameTypes.h"
+#include "TransformTypes.h"
+
+namespace UE
+{
+	namespace Geometry
+	{
+		/**
+		 * call EmitLineFunc for the 12 line segments defined by the given box parameters
+		 * @param HalfDimensions X/Y/Z dimensions of box
+		 * @param Center center of box
+		 * @param AxisX X axis of box
+		 * @param AxisY Y axis of box
+		 * @param AxisZ Z axis of box
+		 * @param Transform transform applied to points of box
+		 */
+		template<typename RealType>
+		void GenerateBoxSegments(
+			const FVector3<RealType>& HalfDimensions, 
+			const FVector3<RealType>& Center, 
+			const FVector3<RealType>& AxisX, 
+			const FVector3<RealType>& AxisY, 
+			const FVector3<RealType>& AxisZ,
+			const TTransform3<RealType>& Transform,
+			TFunctionRef<void(const FVector3<RealType>& A, const FVector3<RealType>& B)> EmitLineFunc);
+
+		/**
+		 * call EmitLineFunc for the line segments that make up the Circle defined by the given parameters and sampled with NumSteps vertices
+		 * @param NumSteps number of vertices used to approximate circle
+		 * @param Radius radius of circle
+		 * @param Center center of circle
+		 * @param AxisX first axis of circle
+		 * @param AxisY second axis of circle
+		 * @param Transform transform applied to each point
+		 */
+		template<typename RealType>
+		void GenerateCircleSegments(
+			int32 NumSteps, 
+			RealType Radius, 
+			const FVector3<RealType>& Center,
+			const FVector3<RealType>& AxisX,
+			const FVector3<RealType>& AxisY,
+			const TTransform3<RealType>& Transform,
+			TFunctionRef<void(const FVector3<RealType>& A, const FVector3<RealType>& B)> EmitLineFunc);
+
+		/**
+		 * call EmitLineFunc for the line segments that make up the Circular Arc defined by the given parameters and sampled with NumSteps vertices
+		 * @param NumSteps number of vertices used to approximate arc
+		 * @param Radius radius of arc
+		 * @param StartAngle arc starts at this angle
+		 * @param EndAngle arc ends at this angle
+		 * @param AxisX first axis of arc
+		 * @param AxisY second axis of arc
+		 * @param Transform transform applied to each point
+		 */
+		template<typename RealType>
+		void GenerateArcSegments(
+			int32 NumSteps, 
+			RealType Radius,
+			RealType StartAngle, 
+			RealType EndAngle,
+			const FVector3<RealType>& Center,
+			const FVector3<RealType>& AxisX,
+			const FVector3<RealType>& AxisY,
+			const TTransform3<RealType>& Transform,
+			TFunctionRef<void(const FVector3<RealType>& A, const FVector3<RealType>& B)> EmitLineFunc);
+	}
+}
+
+
+
+
+
+template<typename RealType>
+void UE::Geometry::GenerateBoxSegments(
+	const FVector3<RealType>& HalfDimensions,
+	const FVector3<RealType>& Center,
+	const FVector3<RealType>& AxisX,
+	const FVector3<RealType>& AxisY,
+	const FVector3<RealType>& AxisZ,
+	const TTransform3<RealType>& Transform,
+	TFunctionRef<void(const FVector3<RealType>& A, const FVector3<RealType>& B)> EmitLineFunc)
+{
+	// B is box max/min, P and Q are used to store the start and endpoints of the segments we create
+	FVector3<RealType>	B[2], P, Q;
+	B[0] = HalfDimensions; // max
+	B[1] = -B[0]; // min
+
+	// Iterate across the four corners of top/side/front and create segments along the Z axis
+	for (int32 i = 0; i < 2; i++)
+	{
+		for (int32 j = 0; j < 2; j++)
+		{
+			P.X = B[i].X; Q.X = B[i].X;
+			P.Y = B[j].Y; Q.Y = B[j].Y;
+			P.Z = B[0].Z; Q.Z = B[1].Z;
+			EmitLineFunc(Transform.TransformPosition(P), Transform.TransformPosition(Q));
+
+			P.Y = B[i].Y; Q.Y = B[i].Y;
+			P.Z = B[j].Z; Q.Z = B[j].Z;
+			P.X = B[0].X; Q.X = B[1].X;
+			EmitLineFunc(Transform.TransformPosition(P), Transform.TransformPosition(Q));
+
+			P.Z = B[i].Z; Q.Z = B[i].Z;
+			P.X = B[j].X; Q.X = B[j].X;
+			P.Y = B[0].Y; Q.Y = B[1].Y;
+			EmitLineFunc(Transform.TransformPosition(P), Transform.TransformPosition(Q));
+		}
+	}
+}
+
+
+
+template<typename RealType>
+void UE::Geometry::GenerateCircleSegments(
+	int32 NumSteps, RealType Radius,
+	const FVector3<RealType>& Center,
+	const FVector3<RealType>& AxisX,
+	const FVector3<RealType>& AxisY,
+	const TTransform3<RealType>& Transform,
+	TFunctionRef<void(const FVector3<RealType>& A, const FVector3<RealType>& B)> EmitLineFunc)
+{
+	FVector3<RealType> PrevPos = FVector3<RealType>::Zero();
+	for (int32 i = 0; i <= NumSteps; ++i)
+	{
+		RealType t = (RealType)i / (RealType)NumSteps;
+		RealType Angle = (RealType)2.0 * TMathUtil<RealType>::Pi * t;
+		RealType PlaneX = Radius * TMathUtil<RealType>::Cos(Angle);
+		RealType PlaneY = Radius * TMathUtil<RealType>::Sin(Angle);
+		FVector3<RealType> CurPos = Transform.TransformPosition(Center + PlaneX * AxisX + PlaneY * AxisY);
+		if (i > 0)
+		{
+			EmitLineFunc(PrevPos, CurPos);
+		}
+		PrevPos = CurPos;
+	}
+}
+
+
+
+template<typename RealType>
+void UE::Geometry::GenerateArcSegments(
+	int32 NumSteps,
+	RealType Radius,
+	RealType StartAngle,
+	RealType EndAngle,
+	const FVector3<RealType>& Center,
+	const FVector3<RealType>& AxisX,
+	const FVector3<RealType>& AxisY,
+	const TTransform3<RealType>& Transform,
+	TFunctionRef<void(const FVector3<RealType>& A, const FVector3<RealType>& B)> EmitLineFunc)
+{
+	FVector3<RealType> PrevPos = FVector3<RealType>::Zero();
+	for (int32 i = 0; i <= NumSteps; ++i)
+	{
+		RealType t = (RealType)i / (RealType)NumSteps;
+		RealType Angle = ((RealType)1 - t) * StartAngle + (t)*EndAngle;
+		RealType PlaneX = Radius * TMathUtil<RealType>::Cos(Angle);
+		RealType PlaneY = Radius * TMathUtil<RealType>::Sin(Angle);
+		FVector3<RealType> CurPos = Transform.TransformPosition(Center + PlaneX * AxisX + PlaneY * AxisY);
+		if (i > 0)
+		{
+			EmitLineFunc(PrevPos, CurPos);
+		}
+		PrevPos = CurPos;
+	}
+}
\ No newline at end of file
diff --git a/Engine/Source/Runtime/GeometryCore/Public/Generators/MarchingCubes.h b/Engine/Source/Runtime/GeometryCore/Public/Generators/MarchingCubes.h
new file mode 100644
--- /dev/null
+++ b/Engine/Source/Runtime/GeometryCore/Public/Generators/MarchingCubes.h
@@ -0,0 +1,1375 @@
+// Copyright Epic Games, Inc. All Rights Reserved.
+
+// Port of geometry3Sharp MarchingCubesPro
+
+#pragma once
+
+#include "Async/ParallelFor.h"
+#include "Misc/ScopeLock.h"
+
+#include "Math/UnrealMathUtility.h"
+#include "MeshShapeGenerator.h"
+#include "Misc/AssertionMacros.h"
+
+#include "Spatial/DenseGrid3.h"
+
+#include "CompGeom/PolygonTriangulation.h"
+
+#include "Util/IndexUtil.h"
+
+namespace UE
+{
+namespace Geometry
+{
+
+enum class /*GEOMETRYCORE_API*/ ERootfindingModes
+{
+	SingleLerp,
+	LerpSteps,
+	Bisection
+};
+
+class /*GEOMETRYCORE_API*/ FMarchingCubes : public FMeshShapeGenerator
+{
+public:
+	/**
+	*  this is the function we will evaluate
+	*/
+	TFunction<double(FVector3<double>)> Implicit;
+
+	/**
+	*  mesh surface will be at this isovalue. Normally 0 unless you want
+	*  offset surface or field is not a distance-field.
+	*/
+	double IsoValue = 0;
+
+	/** bounding-box we will mesh inside of. We use the min-corner and
+	 *  the width/height/depth, but do not clamp vertices to stay within max-corner,
+	 *  we may spill one cell over
+	 */
+	TAxisAlignedBox3<double> Bounds;
+
+	/**
+	 *  Length of edges of cubes that are marching.
+	 *  currently, # of cells along axis = (int)(bounds_dimension / CellSize) + 1
+	 */
+	double CubeSize = 0.1;
+
+	/**
+	 *  Use multi-threading? Generally a good idea unless problem is very small or
+	 *  you are multi-threading at a higher level (which may be more efficient)
+	 */
+	bool bParallelCompute = true;
+
+	/**
+	 * If true, code will assume that Implicit() is expensive enough that it is worth it to cache
+	 * evaluations when possible. For something simple like evaluation of an SDF defined by a discrete
+	 * grid, this is generally not worth the overhead.
+	 */
+	bool bEnableValueCaching = true;
+
+	/**
+	 * Max number of cells on any dimension; if exceeded, CubeSize will be automatically increased to fix
+	 */
+	int SafetyMaxDimension = 4096;
+
+	/**
+	 *  Which rootfinding method will be used to converge on surface along edges
+	 */
+	ERootfindingModes RootMode = ERootfindingModes::SingleLerp;
+
+	/**
+	 *  number of iterations of rootfinding method (ignored for SingleLerp)
+	 */
+	int RootModeSteps = 5;
+
+
+	/** if this function returns true, we should abort calculation */
+	TFunction<bool(void)> CancelF = []() { return false; };
+
+	/*
+	 * Outputs
+	 */
+
+	// cube indices range from [Origin,CellDimensions)   
+	FVector3i CellDimensions;
+
+
+	FMarchingCubes()
+	{
+		Bounds = TAxisAlignedBox3<double>(FVector3<double>::Zero(), 8);
+		CubeSize = 0.25;
+	}
+
+	virtual ~FMarchingCubes()
+	{
+	}
+
+	bool Validate()
+	{
+		return CubeSize > 0 && FMath::IsFinite(CubeSize) && !Bounds.IsEmpty() && FMath::IsFinite(Bounds.MaxDim());
+	}
+
+	/**
+	*  Run MC algorithm and generate Output mesh
+	*/
+	FMeshShapeGenerator& Generate() override
+	{
+		TRACE_CPUPROFILER_EVENT_SCOPE(Geometry_MCMesh_Generate);
+
+		if (!ensure(Validate()))
+		{
+			return *this;
+		}
+
+		SetDimensions();
+		GridBounds = FAxisAlignedBox3i(FVector3i::Zero(), CellDimensions - FVector3i(1,1,1)); // grid bounds are inclusive
+
+		if (bEnableValueCaching)
+		{
+			corner_values_grid = FDenseGrid3f(CellDimensions.X + 1, CellDimensions.Y + 1, CellDimensions.Z + 1, FMathf::MaxReal);
+		}
+		corner_values.Reset();
+		InitHashTables();
+		ResetMesh();
+
+		if (bParallelCompute) 
+		{
+			generate_parallel();
+		} 
+		else 
+		{
+			generate_basic();
+		}
+
+		// finalize mesh
+		BuildMesh();
+
+		return *this;
+	}
+
+
+	FMeshShapeGenerator& GenerateContinuation(TArrayView<const FVector3<double>> Seeds)
+	{
+		TRACE_CPUPROFILER_EVENT_SCOPE(Geometry_MCMesh_GenerateContinuation);
+
+		if (!ensure(Validate()))
+		{
+			return *this;
+		}
+
+		SetDimensions();
+		GridBounds = FAxisAlignedBox3i(FVector3i::Zero(), CellDimensions - FVector3i(1,1,1)); // grid bounds are inclusive
+
+		InitHashTables();
+		ResetMesh();
+		corner_values.Reset();
+
+		if (LastGridBounds != GridBounds)
+		{
+			if (bEnableValueCaching)
+			{
+				corner_values_grid = FDenseGrid3f(CellDimensions.X + 1, CellDimensions.Y + 1, CellDimensions.Z + 1, FMathf::MaxReal);
+			}
+			if (bParallelCompute)
+			{
+				done_cells = FDenseGrid3i(CellDimensions.X, CellDimensions.Y, CellDimensions.Z, 0);
+			}
+		}
+		else
+		{
+			if (bEnableValueCaching)
+			{
+				corner_values_grid.Assign(FMathf::MaxReal);
+			}
+			if (bParallelCompute)
+			{
+				done_cells.Assign(0);
+			}
+		}
+
+		if (bParallelCompute) 
+		{
+			generate_continuation_parallel(Seeds);
+		} 
+		else 
+		{
+			generate_continuation(Seeds);
+		}
+
+		// finalize mesh
+		BuildMesh();
+
+		LastGridBounds = GridBounds;
+
+		return *this;
+	}
+
+
+protected:
+
+
+	FAxisAlignedBox3i GridBounds;
+	FAxisAlignedBox3i LastGridBounds;
+
+
+	// we pass Cells around, this makes code cleaner
+	struct FGridCell
+	{
+		// TODO we do not actually need to store i, we just need the min-corner!
+		FVector3i i[8];    // indices of corners of cell
+		double f[8];      // field values at corners
+	};
+
+	void SetDimensions()
+	{
+		int NX = (int)(Bounds.Width() / CubeSize) + 1;
+		int NY = (int)(Bounds.Height() / CubeSize) + 1;
+		int NZ = (int)(Bounds.Depth() / CubeSize) + 1;
+		int MaxDim = FMath::Max3(NX, NY, NZ);
+		if (!ensure(MaxDim <= SafetyMaxDimension))
+		{
+			CubeSize = Bounds.MaxDim() / double(SafetyMaxDimension - 1);
+			NX = (int)(Bounds.Width() / CubeSize) + 1;
+			NY = (int)(Bounds.Height() / CubeSize) + 1;
+			NZ = (int)(Bounds.Depth() / CubeSize) + 1;
+		}
+		CellDimensions = FVector3i(NX, NY, NZ);
+	}
+
+	void corner_pos(const FVector3i& IJK, FVector3<double>& Pos)
+	{
+		Pos.X = Bounds.Min.X + CubeSize * IJK.X;
+		Pos.Y = Bounds.Min.Y + CubeSize * IJK.Y;
+		Pos.Z = Bounds.Min.Z + CubeSize * IJK.Z;
+	}
+	FVector3<double> corner_pos(const FVector3i& IJK)
+	{
+		return FVector3<double>(Bounds.Min.X + CubeSize * IJK.X,
+			Bounds.Min.Y + CubeSize * IJK.Y,
+			Bounds.Min.Z + CubeSize * IJK.Z);
+	}
+	FVector3i cell_index(const FVector3<double>& Pos)
+	{
+		checkSlow(Bounds.Contains(Pos));
+		return FVector3i(
+			(int)((Pos.X - Bounds.Min.X) / CubeSize),
+			(int)((Pos.Y - Bounds.Min.Y) / CubeSize),
+			(int)((Pos.Z - Bounds.Min.Z) / CubeSize));
+	}
+
+
+
+	//
+	// corner and edge hash functions, these pack the coordinate
+	// integers into 16-bits, so max of 65536 in any dimension.
+	//
+
+
+	int64 corner_hash(const FVector3i& Idx)
+	{
+		return ((int64)Idx.X&0xFFFF) | (((int64)Idx.Y&0xFFFF) << 16) | (((int64)Idx.Z&0xFFFF) << 32);
+	}
+	int64 corner_hash(int X, int Y, int Z)
+	{
+		return ((int64)X & 0xFFFF) | (((int64)Y & 0xFFFF) << 16) | (((int64)Z & 0xFFFF) << 32);
+	}
+
+	const int64 EDGE_X = int64(1) << 60;
+	const int64 EDGE_Y = int64(1) << 61;
+	const int64 EDGE_Z = int64(1) << 62;
+
+	int64 edge_hash(const FVector3i& Idx1, const FVector3i& Idx2)
+	{
+		if ( Idx1.X != Idx2.X )
+		{
+			int xlo = FMath::Min(Idx1.X, Idx2.X);
+			return corner_hash(xlo, Idx1.Y, Idx1.Z) | EDGE_X;
+		}
+		else if ( Idx1.Y != Idx2.Y )
+		{
+			int ylo = FMath::Min(Idx1.Y, Idx2.Y);
+			return corner_hash(Idx1.X, ylo, Idx1.Z) | EDGE_Y;
+		}
+		else
+		{
+			int zlo = FMath::Min(Idx1.Z, Idx2.Z);
+			return corner_hash(Idx1.X, Idx1.Y, zlo) | EDGE_Z;
+		}
+	}
+
+
+
+	//
+	// Hash table for edge vertices
+	//
+
+	int64 NumEdgeVertexSections = 64;
+	TArray<TMap<int64, int>> EdgeVertexSections;
+	TArray<FCriticalSection> EdgeVertexSectionLocks;
+	
+	int FindVertexID(int64 hash)
+	{
+		int32 SectionIndex = (int32)(hash % (NumEdgeVertexSections - 1));
+		FScopeLock Lock(&EdgeVertexSectionLocks[SectionIndex]);
+		int* Found = EdgeVertexSections[SectionIndex].Find(hash);
+		return (Found != nullptr) ? *Found : IndexConstants::InvalidID;
+	}
+
+	int AppendOrFindVertexID(int64 hash, FVector3<double> Pos)
+	{
+		int32 SectionIndex = (int32)(hash % (NumEdgeVertexSections - 1));
+		FScopeLock Lock(&EdgeVertexSectionLocks[SectionIndex]);
+		int* FoundVID = EdgeVertexSections[SectionIndex].Find(hash);
+		if (FoundVID != nullptr)
+		{
+			return *FoundVID;
+		}
+		int NewVID = append_vertex(Pos, hash);
+		EdgeVertexSections[SectionIndex].Add(hash, NewVID);
+		return NewVID;
+	}
+
+
+	int edge_vertex_id(const FVector3i& Idx1, const FVector3i& Idx2, double F1, double F2)
+	{
+		int64 hash = edge_hash(Idx1, Idx2);
+
+		int foundvid = FindVertexID(hash);
+		if (foundvid != IndexConstants::InvalidID)
+		{
+			return foundvid;
+		}
+
+		// ok this is a bit messy. We do not want to lock the entire hash table 
+		// while we do find_iso. However it is possible that during this time we
+		// are unlocked we have re-entered with the same edge. So when we
+		// re-acquire the lock we need to check again that we have not already
+		// computed this edge, otherwise we will end up with duplicate vertices!
+
+		FVector3<double> pa = FVector3<double>::Zero(), pb = FVector3<double>::Zero();
+		corner_pos(Idx1, pa);
+		corner_pos(Idx2, pb);
+		FVector3<double> Pos = FVector3<double>::Zero();
+		find_iso(pa, pb, F1, F2, Pos);
+
+		return AppendOrFindVertexID(hash, Pos);
+	}
+
+
+
+
+
+
+	//
+	// Store corner values in hash table. This doesn't make
+	// sense if we are evaluating entire grid, way too slow.
+	//
+
+	TMap<int64, double> corner_values;
+	FCriticalSection corner_values_lock;
+
+	double corner_value(const FVector3i& Idx)
+	{
+		int64 hash = corner_hash(Idx);
+		double value = 0;
+		double* findval = corner_values.Find(hash);
+		if (findval == nullptr)
+		{
+			FVector3<double> V = corner_pos(Idx);
+			value = Implicit(V);
+			corner_values.Add(hash, value);
+		} 
+		else
+		{
+			value = *findval;
+		}
+		return value;
+	}
+	void initialize_cell_values(FGridCell& Cell, bool Shift)
+	{
+		FScopeLock Lock(&corner_values_lock);
+
+		if ( Shift )
+		{
+			Cell.f[1] = corner_value(Cell.i[1]);
+			Cell.f[2] = corner_value(Cell.i[2]);
+			Cell.f[5] = corner_value(Cell.i[5]);
+			Cell.f[6] = corner_value(Cell.i[6]);
+		}
+		else
+		{
+			for (int i = 0; i < 8; ++i)
+			{
+				Cell.f[i] = corner_value(Cell.i[i]);
+			}
+		}
+	}
+
+
+
+	//
+	// store corner values in pre-allocated grid that has
+	// FMathf::MaxReal as sentinel. 
+	// (note this is float grid, not double...)
+	//
+
+	FDenseGrid3f corner_values_grid;
+
+	int64 NumCornerValueSections = 64;
+	TArray<FCriticalSection> CornerValueSectionLocks;
+	int32 GetCornerValueSectionIndex(int64 hash)
+	{
+		return (int32)(hash % (NumCornerValueSections - 1));
+	}
+
+	double corner_value_grid_parallel(const FVector3i& Idx)
+	{
+		int32 SectionIndex = GetCornerValueSectionIndex(corner_hash(Idx));
+		float val = 0.0;
+		{
+			FScopeLock Lock(&CornerValueSectionLocks[SectionIndex]);
+			val = corner_values_grid[Idx];
+		}
+		if (val != FMathf::MaxReal)
+		{
+			return (double)val;
+		}
+
+		FVector3<double> V = corner_pos(Idx);
+		val = (float)Implicit(V);
+
+		{
+			FScopeLock Lock(&CornerValueSectionLocks[SectionIndex]);
+			corner_values_grid[Idx] = val;
+		}
+
+		return (double)val;
+	}
+	double corner_value_grid(const FVector3i& Idx)
+	{
+		if (bParallelCompute)
+		{
+			return corner_value_grid_parallel(Idx);
+		}
+
+		float val = corner_values_grid[Idx];
+		if (val != FMathf::MaxReal)
+		{
+			return (double)val;
+		}
+
+		FVector3<double> V = corner_pos(Idx);
+		val = (float)Implicit(V);
+		corner_values_grid[Idx] = val;
+		return (double)val;
+	}
+
+	void initialize_cell_values_grid(FGridCell& Cell, bool Shift)
+	{
+		if (Shift)
+		{
+			Cell.f[1] = corner_value_grid(Cell.i[1]);
+			Cell.f[2] = corner_value_grid(Cell.i[2]);
+			Cell.f[5] = corner_value_grid(Cell.i[5]);
+			Cell.f[6] = corner_value_grid(Cell.i[6]);
+		}
+		else
+		{
+			for (int i = 0; i < 8; ++i)
+			{
+				Cell.f[i] = corner_value_grid(Cell.i[i]);
+			}
+		}
+	}
+
+
+
+	//
+	// explicitly compute corner values as necessary
+	//
+	//
+
+	double corner_value_nohash(const FVector3i& Idx) 
+	{
+		FVector3<double> V = corner_pos(Idx);
+		return Implicit(V);
+	}
+	void initialize_cell_values_nohash(FGridCell& Cell, bool Shift)
+	{
+		if (Shift)
+		{
+			Cell.f[1] = corner_value_nohash(Cell.i[1]);
+			Cell.f[2] = corner_value_nohash(Cell.i[2]);
+			Cell.f[5] = corner_value_nohash(Cell.i[5]);
+			Cell.f[6] = corner_value_nohash(Cell.i[6]);
+		}
+		else
+		{
+			for (int i = 0; i < 8; ++i)
+			{
+				Cell.f[i] = corner_value_nohash(Cell.i[i]);
+			}
+		}
+	}
+
+
+
+	/**
+	*  compute 3D corner-positions and field values for cell at index
+	*/
+	void initialize_cell(FGridCell& Cell, const FVector3i& Idx)
+	{
+		Cell.i[0] = FVector3i(Idx.X + 0, Idx.Y + 0, Idx.Z + 0);
+		Cell.i[1] = FVector3i(Idx.X + 1, Idx.Y + 0, Idx.Z + 0);
+		Cell.i[2] = FVector3i(Idx.X + 1, Idx.Y + 0, Idx.Z + 1);
+		Cell.i[3] = FVector3i(Idx.X + 0, Idx.Y + 0, Idx.Z + 1);
+		Cell.i[4] = FVector3i(Idx.X + 0, Idx.Y + 1, Idx.Z + 0);
+		Cell.i[5] = FVector3i(Idx.X + 1, Idx.Y + 1, Idx.Z + 0);
+		Cell.i[6] = FVector3i(Idx.X + 1, Idx.Y + 1, Idx.Z + 1);
+		Cell.i[7] = FVector3i(Idx.X + 0, Idx.Y + 1, Idx.Z + 1);
+
+		//initialize_cell_values(Cell, false);
+		if (bEnableValueCaching)
+		{
+			initialize_cell_values_grid(Cell, false);
+		}
+		else
+		{
+			initialize_cell_values_nohash(Cell, false);
+		}
+	}
+
+
+	// assume we just want to slide cell at XIdx-1 to cell at XIdx, while keeping
+	// yi and ZIdx constant. Then only x-coords change, and we have already 
+	// computed half the values
+	void shift_cell_x(FGridCell& Cell, int XIdx)
+	{
+		Cell.f[0] = Cell.f[1];
+		Cell.f[3] = Cell.f[2];
+		Cell.f[4] = Cell.f[5];
+		Cell.f[7] = Cell.f[6];
+
+		Cell.i[0].X = XIdx; Cell.i[1].X = XIdx+1; Cell.i[2].X = XIdx+1; Cell.i[3].X = XIdx;
+		Cell.i[4].X = XIdx; Cell.i[5].X = XIdx+1; Cell.i[6].X = XIdx+1; Cell.i[7].X = XIdx;
+
+		//initialize_cell_values(Cell, true);
+		if (bEnableValueCaching)
+		{
+			initialize_cell_values_grid(Cell, true);
+		}
+		else
+		{
+			initialize_cell_values_nohash(Cell, true);
+		}
+	}
+
+
+	void InitHashTables()
+	{
+		EdgeVertexSections.Reset();
+		EdgeVertexSections.SetNum(NumEdgeVertexSections);
+		EdgeVertexSectionLocks.Reset();
+		EdgeVertexSectionLocks.SetNum(NumEdgeVertexSections);
+
+		CornerValueSectionLocks.Reset();
+		CornerValueSectionLocks.SetNum(NumCornerValueSections);
+
+		DoneCellSectionLocks.Reset();
+		DoneCellSectionLocks.SetNum(NumDoneCellSections);
+	}
+
+
+	bool parallel_mesh_access = false;
+
+
+	/**
+	*  processing z-slabs of cells in parallel
+	*/
+	void generate_parallel()
+	{
+		parallel_mesh_access = true;
+
+		// [TODO] maybe shouldn't alway use Z axis here?
+		ParallelFor(CellDimensions.Z, [this](int32 ZIdx)
+		{
+			FGridCell Cell;
+			int vertTArray[12];
+			for (int yi = 0; yi < CellDimensions.Y; ++yi)
+			{
+				if (CancelF())
+				{
+					return;
+				}
+				// compute full cell at x=0, then slide along x row, which saves half of value computes
+				FVector3i Idx(0, yi, ZIdx);
+				initialize_cell(Cell, Idx);
+				polygonize_cell(Cell, vertTArray);
+				for (int XIdx = 1; XIdx < CellDimensions.X; ++XIdx)
+				{
+					shift_cell_x(Cell, XIdx);
+					polygonize_cell(Cell, vertTArray);
+				}
+			}
+		});
+
+
+		parallel_mesh_access = false;
+	}
+
+
+
+
+	/**
+	*  fully sequential version, no threading
+	*/
+	void generate_basic()
+	{
+		FGridCell Cell;
+		int vertTArray[12];
+
+		for (int ZIdx = 0; ZIdx < CellDimensions.Z; ++ZIdx)
+		{
+			for (int yi = 0; yi < CellDimensions.Y; ++yi)
+			{
+				if (CancelF())
+				{
+					return;
+				}
+				// compute full Cell at x=0, then slide along x row, which saves half of value computes
+				FVector3i Idx(0, yi, ZIdx);
+				initialize_cell(Cell, Idx);
+				polygonize_cell(Cell, vertTArray);
+				for (int XIdx = 1; XIdx < CellDimensions.X; ++XIdx)
+				{
+					shift_cell_x(Cell, XIdx);
+					polygonize_cell(Cell, vertTArray);
+				}
+
+			}
+		}
+	}
+
+
+
+
+	/**
+	*  fully sequential version, no threading
+	*/
+	void generate_continuation(TArrayView<const FVector3<double>> Seeds)
+	{
+		FGridCell Cell;
+		int vertTArray[12];
+
+		done_cells = FDenseGrid3i(CellDimensions.X, CellDimensions.Y, CellDimensions.Z, 0);
+
+		TArray<FVector3i> stack;
+
+		for (FVector3<double> seed : Seeds)
+		{
+			FVector3i seed_idx = cell_index(seed);
+			if (!done_cells.IsValidIndex(seed_idx) || done_cells[seed_idx] == 1)
+			{
+				continue;
+			}
+			stack.Add(seed_idx);
+			done_cells[seed_idx] = 1;
+
+			while ( stack.Num() > 0 )
+			{
+				FVector3i Idx = stack[stack.Num()-1]; 
+				stack.RemoveAt(stack.Num()-1);
+				if (CancelF())
+				{
+					return;
+				}
+
+				initialize_cell(Cell, Idx);
+				if ( polygonize_cell(Cell, vertTArray) )
+				{     // found crossing
+					for ( FVector3i o : IndexUtil::GridOffsets6 )
+					{
+						FVector3i nbr_idx = Idx + o;
+						if (GridBounds.Contains(nbr_idx) && done_cells[nbr_idx] == 0)
+						{
+							stack.Add(nbr_idx);
+							done_cells[nbr_idx] = 1;
+						}
+					}
+				}
+			}
+		}
+	}
+
+
+
+
+	/**
+	*  parallel seed evaluation
+	*/
+	void generate_continuation_parallel(TArrayView<const FVector3<double>> Seeds)
+	{
+		parallel_mesh_access = true;
+
+		ParallelFor(Seeds.Num(), [this, &Seeds](int32 Index)
+		{
+			FVector3<double> Seed = Seeds[Index];
+			FVector3i seed_idx = cell_index(Seed);
+			if (!done_cells.IsValidIndex(seed_idx) || set_cell_if_not_done(seed_idx) == false)
+			{
+				return;
+			}
+
+			FGridCell Cell;
+			int vertTArray[12];
+
+			TArray<FVector3i> stack;
+			stack.Add(seed_idx);
+
+			while (stack.Num() > 0)
+			{
+				FVector3i Idx = stack[stack.Num() - 1];
+				stack.RemoveAt(stack.Num() - 1);
+				if (CancelF())
+				{
+					return;
+				}
+
+				initialize_cell(Cell, Idx);
+				if (polygonize_cell(Cell, vertTArray))
+				{     // found crossing
+					for (FVector3i o : IndexUtil::GridOffsets6)
+					{
+						FVector3i nbr_idx = Idx + o;
+						if (GridBounds.Contains(nbr_idx))
+						{
+							if (set_cell_if_not_done(nbr_idx) == true)
+							{ 
+								stack.Add(nbr_idx);
+							}
+						}
+					}
+				}
+			}
+		});
+
+		parallel_mesh_access = false;
+	}
+
+
+
+	FDenseGrid3i done_cells;
+
+	int64 NumDoneCellSections = 64;
+	TArray<FCriticalSection> DoneCellSectionLocks;
+	int GetDoneCellSectionIndex(int64 hash)
+	{
+		return (int32)(hash % (NumDoneCellSections - 1));
+	}
+
+	bool set_cell_if_not_done(const FVector3i& Idx)
+	{
+		int32 SectionIndex = GetDoneCellSectionIndex(corner_hash(Idx));
+		bool was_set = false;
+		{
+			FScopeLock Lock(&DoneCellSectionLocks[SectionIndex]);
+			if (done_cells[Idx] == 0)
+			{
+				done_cells[Idx] = 1;
+				was_set = true;
+			}
+		}
+		return was_set;
+	}
+
+
+
+
+
+
+
+
+
+
+	/**
+	*  find edge crossings and generate triangles for this cell
+	*/
+	bool polygonize_cell(FGridCell& Cell, int VertIndexArray[8])
+	{
+		// construct bits of index into edge table, where bit for each
+		// corner is 1 if that value is < isovalue.
+		// This tell us which edges have sign-crossings, and the int value
+		// of the bitmap is an index into the edge and triangle tables
+		int cubeindex = 0, Shift = 1;
+		for (int i = 0; i < 8; ++i)
+		{
+			if (Cell.f[i] < IsoValue)
+			{
+				cubeindex |= Shift;
+			}
+			Shift <<= 1;
+		}
+
+		// no crossings!
+		if (EdgeTable[cubeindex] == 0)
+		{
+			return false;
+		}
+
+		// check each bit of value in edge table. If it is 1, we
+		// have a crossing on that edge. Look up the indices of this
+		// edge and find the intersection point along it
+		Shift = 1;
+		FVector3<double> pa = FVector3<double>::Zero(), pb = FVector3<double>::Zero();
+		for (int i = 0; i <= 11; i++)
+		{
+			if ((EdgeTable[cubeindex] & Shift) != 0)
+			{
+				int a = EdgeIndices[i][0], b = EdgeIndices[i][1];
+				VertIndexArray[i] = edge_vertex_id(Cell.i[a], Cell.i[b], Cell.f[a], Cell.f[b]);
+			}
+			Shift <<= 1;
+		}
+
+		int64 CellHash = corner_hash(Cell.i[0]);
+
+		// now iterate through the set of triangles in TriTable for this cube,
+		// and emit triangles using the vertices we found.
+		int tri_count = 0;
+		for (int i = 0; TriTable[cubeindex][i] != -1; i += 3)
+		{
+			int ta = TriTable[cubeindex][i];
+			int tb = TriTable[cubeindex][i + 1];
+			int tc = TriTable[cubeindex][i + 2];
+			int a = VertIndexArray[ta], b = VertIndexArray[tb], c = VertIndexArray[tc];
+
+			// if a corner is within tolerance of isovalue, then some triangles
+			// will be degenerate, and we can skip them w/o resulting in cracks (right?)
+			// !! this should never happen anymore...artifact of old hashtable impl
+			if (!ensure(a != b && a != c && b != c))
+			{
+				continue;
+			}
+
+			append_triangle(a, b, c, CellHash);
+			tri_count++;
+		}
+
+		return (tri_count > 0);
+	}
+
+
+	struct FIndexedVertex
+	{
+		int32 Index;
+		FVector3d Position;
+	};
+	std::atomic<int32> VertexCounter;
+
+	int64 NumVertexSections = 64;
+	TArray<FCriticalSection> VertexSectionLocks;
+	TArray<TArray<FIndexedVertex>> VertexSectionLists;
+	int GetVertexSectionIndex(int64 hash)
+	{
+		return (int32)(hash % (NumVertexSections - 1));
+	}
+
+	/**
+	*  add vertex to mesh, with locking if we are computing in parallel
+	*/
+	int append_vertex(FVector3<double> V, int64 CellHash)
+	{
+		int SectionIndex = GetVertexSectionIndex(CellHash);
+		int32 NewIndex = VertexCounter++;
+
+		if (parallel_mesh_access)
+		{
+			FScopeLock Lock(&VertexSectionLocks[SectionIndex]);
+			VertexSectionLists[SectionIndex].Add(FIndexedVertex{ NewIndex, V });
+		}
+		else
+		{
+			VertexSectionLists[SectionIndex].Add(FIndexedVertex{ NewIndex, V });
+		}
+
+		return NewIndex;
+	}
+
+
+	int64 NumTriangleSections = 64;
+	TArray<FCriticalSection> TriangleSectionLocks;
+	TArray<TArray<FIndex3i>> TriangleSectionLists;
+	int GetTriangleSectionIndex(int64 hash)
+	{
+		return (int32)(hash % (NumTriangleSections - 1));
+	}
+
+	/**
+	*  add triangle to mesh, with locking if we are computing in parallel
+	*/
+	void append_triangle(int A, int B, int C, int64 CellHash)
+	{
+		int SectionIndex = GetTriangleSectionIndex(CellHash);
+		if (parallel_mesh_access)
+		{
+			FScopeLock Lock(&TriangleSectionLocks[SectionIndex]);
+			TriangleSectionLists[SectionIndex].Add(FIndex3i(A, B, C));
+		}
+		else
+		{
+			TriangleSectionLists[SectionIndex].Add(FIndex3i(A, B, C));
+		}
+	}
+
+
+	/**
+	 * Reset internal mesh-assembly data structures
+	 */
+	void ResetMesh()
+	{
+		VertexSectionLocks.SetNum(NumVertexSections);
+		VertexSectionLists.Reset();
+		VertexSectionLists.SetNum(NumVertexSections);
+		VertexCounter = 0;
+
+		TriangleSectionLocks.SetNum(NumTriangleSections);
+		TriangleSectionLists.Reset();
+		TriangleSectionLists.SetNum(NumTriangleSections);
+	}
+
+	/**
+	 * Populate FMeshShapeGenerator data structures from accumulated
+	 * vertex/triangle sets
+	 */
+	void BuildMesh()
+	{
+		TRACE_CPUPROFILER_EVENT_SCOPE(Geometry_MCMesh_BuildMesh);
+
+		int32 NumVertices = VertexCounter;
+		TArray<FVector3d> VertexBuffer;
+		VertexBuffer.SetNum(NumVertices);
+		for (const TArray<FIndexedVertex>& VertexList : VertexSectionLists)
+		{
+			for (FIndexedVertex Vtx : VertexList)
+			{
+				VertexBuffer[Vtx.Index] = Vtx.Position;
+			}
+		}
+		for (int32 k = 0; k < NumVertices; ++k)
+		{
+			int32 vid = AppendVertex(VertexBuffer[k]);
+		}
+
+		for (const TArray<FIndex3i>& TriangleList : TriangleSectionLists)
+		{
+			for (FIndex3i Tri : TriangleList)
+			{
+				AppendTriangle(Tri.A, Tri.B, Tri.C);
+			}
+		}
+	}
+
+
+
+	/**
+	*  root-find the intersection along edge from f(P1)=ValP1 to f(P2)=ValP2
+	*/
+	void find_iso(const FVector3<double>& P1, const FVector3<double>& P2, double ValP1, double ValP2, FVector3<double>& PIso)
+	{
+		// Ok, this is a bit hacky but seems to work? If both isovalues
+		// are the same, we just return the midpoint. If one is nearly zero, we can
+		// but assume that's where the surface is. *However* if we return that point exactly,
+		// we can get nonmanifold vertices, because multiple fans may connect there. 
+		// Since FDynamicMesh3 disallows that, it results in holes. So we pull 
+		// slightly towards the other point along this edge. This means we will get
+		// repeated nearly-coincident vertices, but the mesh will be manifold.
+		const double dt = 0.999999;
+		if (FMath::Abs(ValP1 - ValP2) < 0.00001)
+		{
+			PIso = (P1 + P2) * 0.5;
+			return;
+		}
+		if (FMath::Abs(IsoValue - ValP1) < 0.00001)
+		{
+			PIso = dt * P1 + (1.0 - dt) * P2;
+			return;
+		}
+		if (FMath::Abs(IsoValue - ValP2) < 0.00001)
+		{
+			PIso = (dt) * P2 + (1.0 - dt) * P1;
+			return;
+		}
+
+		// Note: if we don't maintain min/max order here, then numerical error means
+		//   that hashing on point x/y/z doesn't work
+		FVector3<double> a = P1, b = P2;
+		double fa = ValP1, fb = ValP2;
+		if (ValP2 < ValP1)
+		{
+			a = P2; b = P1;
+			fb = ValP1; fa = ValP2;
+		}
+
+		// converge on root
+		if (RootMode == ERootfindingModes::Bisection)
+		{
+			for (int k = 0; k < RootModeSteps; ++k)
+			{
+				PIso.X = (a.X + b.X) * 0.5; PIso.Y = (a.Y + b.Y) * 0.5; PIso.Z = (a.Z + b.Z) * 0.5;
+				double mid_f = Implicit(PIso);
+				if (mid_f < IsoValue)
+				{
+					a = PIso; fa = mid_f;
+				}
+				else
+				{
+					b = PIso; fb = mid_f;
+				}
+			}
+			PIso = Lerp(a, b, 0.5);
+
+		}
+		else
+		{
+			double mu = 0;
+			if (RootMode == ERootfindingModes::LerpSteps)
+			{
+				for (int k = 0; k < RootModeSteps; ++k)
+				{
+					mu = FMathd::Clamp((IsoValue - fa) / (fb - fa), 0.0, 1.0);
+					PIso.X = a.X + mu * (b.X - a.X);
+					PIso.Y = a.Y + mu * (b.Y - a.Y);
+					PIso.Z = a.Z + mu * (b.Z - a.Z);
+					double mid_f = Implicit(PIso);
+					if (mid_f < IsoValue)
+					{
+						a = PIso; fa = mid_f;
+					}
+					else
+					{
+						b = PIso; fb = mid_f;
+					}
+				}
+			}
+
+			// final lerp
+			mu = FMathd::Clamp((IsoValue - fa) / (fb - fa), 0.0, 1.0);
+			PIso.X = a.X + mu * (b.X - a.X);
+			PIso.Y = a.Y + mu * (b.Y - a.Y);
+			PIso.Z = a.Z + mu * (b.Z - a.Z);
+		}
+	}
+
+
+
+
+	/*
+	* Below here are standard marching-cubes tables. 
+	*/
+
+
+	GEOMETRYCORE_API constexpr static int EdgeIndices[12][2] = {
+		{0,1}, {1,2}, {2,3}, {3,0}, {4,5}, {5,6}, {6,7}, {7,4}, {0,4}, {1,5}, {2,6}, {3,7}
+	};
+
+	GEOMETRYCORE_API constexpr static int EdgeTable[256] = {
+		0x0  , 0x109, 0x203, 0x30a, 0x406, 0x50f, 0x605, 0x70c,
+		0x80c, 0x905, 0xa0f, 0xb06, 0xc0a, 0xd03, 0xe09, 0xf00,
+		0x190, 0x99 , 0x393, 0x29a, 0x596, 0x49f, 0x795, 0x69c,
+		0x99c, 0x895, 0xb9f, 0xa96, 0xd9a, 0xc93, 0xf99, 0xe90,
+		0x230, 0x339, 0x33 , 0x13a, 0x636, 0x73f, 0x435, 0x53c,
+		0xa3c, 0xb35, 0x83f, 0x936, 0xe3a, 0xf33, 0xc39, 0xd30,
+		0x3a0, 0x2a9, 0x1a3, 0xaa , 0x7a6, 0x6af, 0x5a5, 0x4ac,
+		0xbac, 0xaa5, 0x9af, 0x8a6, 0xfaa, 0xea3, 0xda9, 0xca0,
+		0x460, 0x569, 0x663, 0x76a, 0x66 , 0x16f, 0x265, 0x36c,
+		0xc6c, 0xd65, 0xe6f, 0xf66, 0x86a, 0x963, 0xa69, 0xb60,
+		0x5f0, 0x4f9, 0x7f3, 0x6fa, 0x1f6, 0xff , 0x3f5, 0x2fc,
+		0xdfc, 0xcf5, 0xfff, 0xef6, 0x9fa, 0x8f3, 0xbf9, 0xaf0,
+		0x650, 0x759, 0x453, 0x55a, 0x256, 0x35f, 0x55 , 0x15c,
+		0xe5c, 0xf55, 0xc5f, 0xd56, 0xa5a, 0xb53, 0x859, 0x950,
+		0x7c0, 0x6c9, 0x5c3, 0x4ca, 0x3c6, 0x2cf, 0x1c5, 0xcc ,
+		0xfcc, 0xec5, 0xdcf, 0xcc6, 0xbca, 0xac3, 0x9c9, 0x8c0,
+		0x8c0, 0x9c9, 0xac3, 0xbca, 0xcc6, 0xdcf, 0xec5, 0xfcc,
+		0xcc , 0x1c5, 0x2cf, 0x3c6, 0x4ca, 0x5c3, 0x6c9, 0x7c0,
+		0x950, 0x859, 0xb53, 0xa5a, 0xd56, 0xc5f, 0xf55, 0xe5c,
+		0x15c, 0x55 , 0x35f, 0x256, 0x55a, 0x453, 0x759, 0x650,
+		0xaf0, 0xbf9, 0x8f3, 0x9fa, 0xef6, 0xfff, 0xcf5, 0xdfc,
+		0x2fc, 0x3f5, 0xff , 0x1f6, 0x6fa, 0x7f3, 0x4f9, 0x5f0,
+		0xb60, 0xa69, 0x963, 0x86a, 0xf66, 0xe6f, 0xd65, 0xc6c,
+		0x36c, 0x265, 0x16f, 0x66 , 0x76a, 0x663, 0x569, 0x460,
+		0xca0, 0xda9, 0xea3, 0xfaa, 0x8a6, 0x9af, 0xaa5, 0xbac,
+		0x4ac, 0x5a5, 0x6af, 0x7a6, 0xaa , 0x1a3, 0x2a9, 0x3a0,
+		0xd30, 0xc39, 0xf33, 0xe3a, 0x936, 0x83f, 0xb35, 0xa3c,
+		0x53c, 0x435, 0x73f, 0x636, 0x13a, 0x33 , 0x339, 0x230,
+		0xe90, 0xf99, 0xc93, 0xd9a, 0xa96, 0xb9f, 0x895, 0x99c,
+		0x69c, 0x795, 0x49f, 0x596, 0x29a, 0x393, 0x99 , 0x190,
+		0xf00, 0xe09, 0xd03, 0xc0a, 0xb06, 0xa0f, 0x905, 0x80c,
+		0x70c, 0x605, 0x50f, 0x406, 0x30a, 0x203, 0x109, 0x0   };
+
+
+	GEOMETRYCORE_API constexpr static int TriTable[256][16] =
+	{
+		{-1, -1, -1, -1, -1, -1, -1, -1, -1, -1, -1, -1, -1, -1, -1, -1},
+		{0, 8, 3, -1, -1, -1, -1, -1, -1, -1, -1, -1, -1, -1, -1, -1},
+		{0, 1, 9, -1, -1, -1, -1, -1, -1, -1, -1, -1, -1, -1, -1, -1},
+		{1, 8, 3, 9, 8, 1, -1, -1, -1, -1, -1, -1, -1, -1, -1, -1},
+		{1, 2, 10, -1, -1, -1, -1, -1, -1, -1, -1, -1, -1, -1, -1, -1},
+		{0, 8, 3, 1, 2, 10, -1, -1, -1, -1, -1, -1, -1, -1, -1, -1},
+		{9, 2, 10, 0, 2, 9, -1, -1, -1, -1, -1, -1, -1, -1, -1, -1},
+		{2, 8, 3, 2, 10, 8, 10, 9, 8, -1, -1, -1, -1, -1, -1, -1},
+		{3, 11, 2, -1, -1, -1, -1, -1, -1, -1, -1, -1, -1, -1, -1, -1},
+		{0, 11, 2, 8, 11, 0, -1, -1, -1, -1, -1, -1, -1, -1, -1, -1},
+		{1, 9, 0, 2, 3, 11, -1, -1, -1, -1, -1, -1, -1, -1, -1, -1},
+		{1, 11, 2, 1, 9, 11, 9, 8, 11, -1, -1, -1, -1, -1, -1, -1},
+		{3, 10, 1, 11, 10, 3, -1, -1, -1, -1, -1, -1, -1, -1, -1, -1},
+		{0, 10, 1, 0, 8, 10, 8, 11, 10, -1, -1, -1, -1, -1, -1, -1},
+		{3, 9, 0, 3, 11, 9, 11, 10, 9, -1, -1, -1, -1, -1, -1, -1},
+		{9, 8, 10, 10, 8, 11, -1, -1, -1, -1, -1, -1, -1, -1, -1, -1},
+		{4, 7, 8, -1, -1, -1, -1, -1, -1, -1, -1, -1, -1, -1, -1, -1},
+		{4, 3, 0, 7, 3, 4, -1, -1, -1, -1, -1, -1, -1, -1, -1, -1},
+		{0, 1, 9, 8, 4, 7, -1, -1, -1, -1, -1, -1, -1, -1, -1, -1},
+		{4, 1, 9, 4, 7, 1, 7, 3, 1, -1, -1, -1, -1, -1, -1, -1},
+		{1, 2, 10, 8, 4, 7, -1, -1, -1, -1, -1, -1, -1, -1, -1, -1},
+		{3, 4, 7, 3, 0, 4, 1, 2, 10, -1, -1, -1, -1, -1, -1, -1},
+		{9, 2, 10, 9, 0, 2, 8, 4, 7, -1, -1, -1, -1, -1, -1, -1},
+		{2, 10, 9, 2, 9, 7, 2, 7, 3, 7, 9, 4, -1, -1, -1, -1},
+		{8, 4, 7, 3, 11, 2, -1, -1, -1, -1, -1, -1, -1, -1, -1, -1},
+		{11, 4, 7, 11, 2, 4, 2, 0, 4, -1, -1, -1, -1, -1, -1, -1},
+		{9, 0, 1, 8, 4, 7, 2, 3, 11, -1, -1, -1, -1, -1, -1, -1},
+		{4, 7, 11, 9, 4, 11, 9, 11, 2, 9, 2, 1, -1, -1, -1, -1},
+		{3, 10, 1, 3, 11, 10, 7, 8, 4, -1, -1, -1, -1, -1, -1, -1},
+		{1, 11, 10, 1, 4, 11, 1, 0, 4, 7, 11, 4, -1, -1, -1, -1},
+		{4, 7, 8, 9, 0, 11, 9, 11, 10, 11, 0, 3, -1, -1, -1, -1},
+		{4, 7, 11, 4, 11, 9, 9, 11, 10, -1, -1, -1, -1, -1, -1, -1},
+		{9, 5, 4, -1, -1, -1, -1, -1, -1, -1, -1, -1, -1, -1, -1, -1},
+		{9, 5, 4, 0, 8, 3, -1, -1, -1, -1, -1, -1, -1, -1, -1, -1},
+		{0, 5, 4, 1, 5, 0, -1, -1, -1, -1, -1, -1, -1, -1, -1, -1},
+		{8, 5, 4, 8, 3, 5, 3, 1, 5, -1, -1, -1, -1, -1, -1, -1},
+		{1, 2, 10, 9, 5, 4, -1, -1, -1, -1, -1, -1, -1, -1, -1, -1},
+		{3, 0, 8, 1, 2, 10, 4, 9, 5, -1, -1, -1, -1, -1, -1, -1},
+		{5, 2, 10, 5, 4, 2, 4, 0, 2, -1, -1, -1, -1, -1, -1, -1},
+		{2, 10, 5, 3, 2, 5, 3, 5, 4, 3, 4, 8, -1, -1, -1, -1},
+		{9, 5, 4, 2, 3, 11, -1, -1, -1, -1, -1, -1, -1, -1, -1, -1},
+		{0, 11, 2, 0, 8, 11, 4, 9, 5, -1, -1, -1, -1, -1, -1, -1},
+		{0, 5, 4, 0, 1, 5, 2, 3, 11, -1, -1, -1, -1, -1, -1, -1},
+		{2, 1, 5, 2, 5, 8, 2, 8, 11, 4, 8, 5, -1, -1, -1, -1},
+		{10, 3, 11, 10, 1, 3, 9, 5, 4, -1, -1, -1, -1, -1, -1, -1},
+		{4, 9, 5, 0, 8, 1, 8, 10, 1, 8, 11, 10, -1, -1, -1, -1},
+		{5, 4, 0, 5, 0, 11, 5, 11, 10, 11, 0, 3, -1, -1, -1, -1},
+		{5, 4, 8, 5, 8, 10, 10, 8, 11, -1, -1, -1, -1, -1, -1, -1},
+		{9, 7, 8, 5, 7, 9, -1, -1, -1, -1, -1, -1, -1, -1, -1, -1},
+		{9, 3, 0, 9, 5, 3, 5, 7, 3, -1, -1, -1, -1, -1, -1, -1},
+		{0, 7, 8, 0, 1, 7, 1, 5, 7, -1, -1, -1, -1, -1, -1, -1},
+		{1, 5, 3, 3, 5, 7, -1, -1, -1, -1, -1, -1, -1, -1, -1, -1},
+		{9, 7, 8, 9, 5, 7, 10, 1, 2, -1, -1, -1, -1, -1, -1, -1},
+		{10, 1, 2, 9, 5, 0, 5, 3, 0, 5, 7, 3, -1, -1, -1, -1},
+		{8, 0, 2, 8, 2, 5, 8, 5, 7, 10, 5, 2, -1, -1, -1, -1},
+		{2, 10, 5, 2, 5, 3, 3, 5, 7, -1, -1, -1, -1, -1, -1, -1},
+		{7, 9, 5, 7, 8, 9, 3, 11, 2, -1, -1, -1, -1, -1, -1, -1},
+		{9, 5, 7, 9, 7, 2, 9, 2, 0, 2, 7, 11, -1, -1, -1, -1},
+		{2, 3, 11, 0, 1, 8, 1, 7, 8, 1, 5, 7, -1, -1, -1, -1},
+		{11, 2, 1, 11, 1, 7, 7, 1, 5, -1, -1, -1, -1, -1, -1, -1},
+		{9, 5, 8, 8, 5, 7, 10, 1, 3, 10, 3, 11, -1, -1, -1, -1},
+		{5, 7, 0, 5, 0, 9, 7, 11, 0, 1, 0, 10, 11, 10, 0, -1},
+		{11, 10, 0, 11, 0, 3, 10, 5, 0, 8, 0, 7, 5, 7, 0, -1},
+		{11, 10, 5, 7, 11, 5, -1, -1, -1, -1, -1, -1, -1, -1, -1, -1},
+		{10, 6, 5, -1, -1, -1, -1, -1, -1, -1, -1, -1, -1, -1, -1, -1},
+		{0, 8, 3, 5, 10, 6, -1, -1, -1, -1, -1, -1, -1, -1, -1, -1},
+		{9, 0, 1, 5, 10, 6, -1, -1, -1, -1, -1, -1, -1, -1, -1, -1},
+		{1, 8, 3, 1, 9, 8, 5, 10, 6, -1, -1, -1, -1, -1, -1, -1},
+		{1, 6, 5, 2, 6, 1, -1, -1, -1, -1, -1, -1, -1, -1, -1, -1},
+		{1, 6, 5, 1, 2, 6, 3, 0, 8, -1, -1, -1, -1, -1, -1, -1},
+		{9, 6, 5, 9, 0, 6, 0, 2, 6, -1, -1, -1, -1, -1, -1, -1},
+		{5, 9, 8, 5, 8, 2, 5, 2, 6, 3, 2, 8, -1, -1, -1, -1},
+		{2, 3, 11, 10, 6, 5, -1, -1, -1, -1, -1, -1, -1, -1, -1, -1},
+		{11, 0, 8, 11, 2, 0, 10, 6, 5, -1, -1, -1, -1, -1, -1, -1},
+		{0, 1, 9, 2, 3, 11, 5, 10, 6, -1, -1, -1, -1, -1, -1, -1},
+		{5, 10, 6, 1, 9, 2, 9, 11, 2, 9, 8, 11, -1, -1, -1, -1},
+		{6, 3, 11, 6, 5, 3, 5, 1, 3, -1, -1, -1, -1, -1, -1, -1},
+		{0, 8, 11, 0, 11, 5, 0, 5, 1, 5, 11, 6, -1, -1, -1, -1},
+		{3, 11, 6, 0, 3, 6, 0, 6, 5, 0, 5, 9, -1, -1, -1, -1},
+		{6, 5, 9, 6, 9, 11, 11, 9, 8, -1, -1, -1, -1, -1, -1, -1},
+		{5, 10, 6, 4, 7, 8, -1, -1, -1, -1, -1, -1, -1, -1, -1, -1},
+		{4, 3, 0, 4, 7, 3, 6, 5, 10, -1, -1, -1, -1, -1, -1, -1},
+		{1, 9, 0, 5, 10, 6, 8, 4, 7, -1, -1, -1, -1, -1, -1, -1},
+		{10, 6, 5, 1, 9, 7, 1, 7, 3, 7, 9, 4, -1, -1, -1, -1},
+		{6, 1, 2, 6, 5, 1, 4, 7, 8, -1, -1, -1, -1, -1, -1, -1},
+		{1, 2, 5, 5, 2, 6, 3, 0, 4, 3, 4, 7, -1, -1, -1, -1},
+		{8, 4, 7, 9, 0, 5, 0, 6, 5, 0, 2, 6, -1, -1, -1, -1},
+		{7, 3, 9, 7, 9, 4, 3, 2, 9, 5, 9, 6, 2, 6, 9, -1},
+		{3, 11, 2, 7, 8, 4, 10, 6, 5, -1, -1, -1, -1, -1, -1, -1},
+		{5, 10, 6, 4, 7, 2, 4, 2, 0, 2, 7, 11, -1, -1, -1, -1},
+		{0, 1, 9, 4, 7, 8, 2, 3, 11, 5, 10, 6, -1, -1, -1, -1},
+		{9, 2, 1, 9, 11, 2, 9, 4, 11, 7, 11, 4, 5, 10, 6, -1},
+		{8, 4, 7, 3, 11, 5, 3, 5, 1, 5, 11, 6, -1, -1, -1, -1},
+		{5, 1, 11, 5, 11, 6, 1, 0, 11, 7, 11, 4, 0, 4, 11, -1},
+		{0, 5, 9, 0, 6, 5, 0, 3, 6, 11, 6, 3, 8, 4, 7, -1},
+		{6, 5, 9, 6, 9, 11, 4, 7, 9, 7, 11, 9, -1, -1, -1, -1},
+		{10, 4, 9, 6, 4, 10, -1, -1, -1, -1, -1, -1, -1, -1, -1, -1},
+		{4, 10, 6, 4, 9, 10, 0, 8, 3, -1, -1, -1, -1, -1, -1, -1},
+		{10, 0, 1, 10, 6, 0, 6, 4, 0, -1, -1, -1, -1, -1, -1, -1},
+		{8, 3, 1, 8, 1, 6, 8, 6, 4, 6, 1, 10, -1, -1, -1, -1},
+		{1, 4, 9, 1, 2, 4, 2, 6, 4, -1, -1, -1, -1, -1, -1, -1},
+		{3, 0, 8, 1, 2, 9, 2, 4, 9, 2, 6, 4, -1, -1, -1, -1},
+		{0, 2, 4, 4, 2, 6, -1, -1, -1, -1, -1, -1, -1, -1, -1, -1},
+		{8, 3, 2, 8, 2, 4, 4, 2, 6, -1, -1, -1, -1, -1, -1, -1},
+		{10, 4, 9, 10, 6, 4, 11, 2, 3, -1, -1, -1, -1, -1, -1, -1},
+		{0, 8, 2, 2, 8, 11, 4, 9, 10, 4, 10, 6, -1, -1, -1, -1},
+		{3, 11, 2, 0, 1, 6, 0, 6, 4, 6, 1, 10, -1, -1, -1, -1},
+		{6, 4, 1, 6, 1, 10, 4, 8, 1, 2, 1, 11, 8, 11, 1, -1},
+		{9, 6, 4, 9, 3, 6, 9, 1, 3, 11, 6, 3, -1, -1, -1, -1},
+		{8, 11, 1, 8, 1, 0, 11, 6, 1, 9, 1, 4, 6, 4, 1, -1},
+		{3, 11, 6, 3, 6, 0, 0, 6, 4, -1, -1, -1, -1, -1, -1, -1},
+		{6, 4, 8, 11, 6, 8, -1, -1, -1, -1, -1, -1, -1, -1, -1, -1},
+		{7, 10, 6, 7, 8, 10, 8, 9, 10, -1, -1, -1, -1, -1, -1, -1},
+		{0, 7, 3, 0, 10, 7, 0, 9, 10, 6, 7, 10, -1, -1, -1, -1},
+		{10, 6, 7, 1, 10, 7, 1, 7, 8, 1, 8, 0, -1, -1, -1, -1},
+		{10, 6, 7, 10, 7, 1, 1, 7, 3, -1, -1, -1, -1, -1, -1, -1},
+		{1, 2, 6, 1, 6, 8, 1, 8, 9, 8, 6, 7, -1, -1, -1, -1},
+		{2, 6, 9, 2, 9, 1, 6, 7, 9, 0, 9, 3, 7, 3, 9, -1},
+		{7, 8, 0, 7, 0, 6, 6, 0, 2, -1, -1, -1, -1, -1, -1, -1},
+		{7, 3, 2, 6, 7, 2, -1, -1, -1, -1, -1, -1, -1, -1, -1, -1},
+		{2, 3, 11, 10, 6, 8, 10, 8, 9, 8, 6, 7, -1, -1, -1, -1},
+		{2, 0, 7, 2, 7, 11, 0, 9, 7, 6, 7, 10, 9, 10, 7, -1},
+		{1, 8, 0, 1, 7, 8, 1, 10, 7, 6, 7, 10, 2, 3, 11, -1},
+		{11, 2, 1, 11, 1, 7, 10, 6, 1, 6, 7, 1, -1, -1, -1, -1},
+		{8, 9, 6, 8, 6, 7, 9, 1, 6, 11, 6, 3, 1, 3, 6, -1},
+		{0, 9, 1, 11, 6, 7, -1, -1, -1, -1, -1, -1, -1, -1, -1, -1},
+		{7, 8, 0, 7, 0, 6, 3, 11, 0, 11, 6, 0, -1, -1, -1, -1},
+		{7, 11, 6, -1, -1, -1, -1, -1, -1, -1, -1, -1, -1, -1, -1, -1},
+		{7, 6, 11, -1, -1, -1, -1, -1, -1, -1, -1, -1, -1, -1, -1, -1},
+		{3, 0, 8, 11, 7, 6, -1, -1, -1, -1, -1, -1, -1, -1, -1, -1},
+		{0, 1, 9, 11, 7, 6, -1, -1, -1, -1, -1, -1, -1, -1, -1, -1},
+		{8, 1, 9, 8, 3, 1, 11, 7, 6, -1, -1, -1, -1, -1, -1, -1},
+		{10, 1, 2, 6, 11, 7, -1, -1, -1, -1, -1, -1, -1, -1, -1, -1},
+		{1, 2, 10, 3, 0, 8, 6, 11, 7, -1, -1, -1, -1, -1, -1, -1},
+		{2, 9, 0, 2, 10, 9, 6, 11, 7, -1, -1, -1, -1, -1, -1, -1},
+		{6, 11, 7, 2, 10, 3, 10, 8, 3, 10, 9, 8, -1, -1, -1, -1},
+		{7, 2, 3, 6, 2, 7, -1, -1, -1, -1, -1, -1, -1, -1, -1, -1},
+		{7, 0, 8, 7, 6, 0, 6, 2, 0, -1, -1, -1, -1, -1, -1, -1},
+		{2, 7, 6, 2, 3, 7, 0, 1, 9, -1, -1, -1, -1, -1, -1, -1},
+		{1, 6, 2, 1, 8, 6, 1, 9, 8, 8, 7, 6, -1, -1, -1, -1},
+		{10, 7, 6, 10, 1, 7, 1, 3, 7, -1, -1, -1, -1, -1, -1, -1},
+		{10, 7, 6, 1, 7, 10, 1, 8, 7, 1, 0, 8, -1, -1, -1, -1},
+		{0, 3, 7, 0, 7, 10, 0, 10, 9, 6, 10, 7, -1, -1, -1, -1},
+		{7, 6, 10, 7, 10, 8, 8, 10, 9, -1, -1, -1, -1, -1, -1, -1},
+		{6, 8, 4, 11, 8, 6, -1, -1, -1, -1, -1, -1, -1, -1, -1, -1},
+		{3, 6, 11, 3, 0, 6, 0, 4, 6, -1, -1, -1, -1, -1, -1, -1},
+		{8, 6, 11, 8, 4, 6, 9, 0, 1, -1, -1, -1, -1, -1, -1, -1},
+		{9, 4, 6, 9, 6, 3, 9, 3, 1, 11, 3, 6, -1, -1, -1, -1},
+		{6, 8, 4, 6, 11, 8, 2, 10, 1, -1, -1, -1, -1, -1, -1, -1},
+		{1, 2, 10, 3, 0, 11, 0, 6, 11, 0, 4, 6, -1, -1, -1, -1},
+		{4, 11, 8, 4, 6, 11, 0, 2, 9, 2, 10, 9, -1, -1, -1, -1},
+		{10, 9, 3, 10, 3, 2, 9, 4, 3, 11, 3, 6, 4, 6, 3, -1},
+		{8, 2, 3, 8, 4, 2, 4, 6, 2, -1, -1, -1, -1, -1, -1, -1},
+		{0, 4, 2, 4, 6, 2, -1, -1, -1, -1, -1, -1, -1, -1, -1, -1},
+		{1, 9, 0, 2, 3, 4, 2, 4, 6, 4, 3, 8, -1, -1, -1, -1},
+		{1, 9, 4, 1, 4, 2, 2, 4, 6, -1, -1, -1, -1, -1, -1, -1},
+		{8, 1, 3, 8, 6, 1, 8, 4, 6, 6, 10, 1, -1, -1, -1, -1},
+		{10, 1, 0, 10, 0, 6, 6, 0, 4, -1, -1, -1, -1, -1, -1, -1},
+		{4, 6, 3, 4, 3, 8, 6, 10, 3, 0, 3, 9, 10, 9, 3, -1},
+		{10, 9, 4, 6, 10, 4, -1, -1, -1, -1, -1, -1, -1, -1, -1, -1},
+		{4, 9, 5, 7, 6, 11, -1, -1, -1, -1, -1, -1, -1, -1, -1, -1},
+		{0, 8, 3, 4, 9, 5, 11, 7, 6, -1, -1, -1, -1, -1, -1, -1},
+		{5, 0, 1, 5, 4, 0, 7, 6, 11, -1, -1, -1, -1, -1, -1, -1},
+		{11, 7, 6, 8, 3, 4, 3, 5, 4, 3, 1, 5, -1, -1, -1, -1},
+		{9, 5, 4, 10, 1, 2, 7, 6, 11, -1, -1, -1, -1, -1, -1, -1},
+		{6, 11, 7, 1, 2, 10, 0, 8, 3, 4, 9, 5, -1, -1, -1, -1},
+		{7, 6, 11, 5, 4, 10, 4, 2, 10, 4, 0, 2, -1, -1, -1, -1},
+		{3, 4, 8, 3, 5, 4, 3, 2, 5, 10, 5, 2, 11, 7, 6, -1},
+		{7, 2, 3, 7, 6, 2, 5, 4, 9, -1, -1, -1, -1, -1, -1, -1},
+		{9, 5, 4, 0, 8, 6, 0, 6, 2, 6, 8, 7, -1, -1, -1, -1},
+		{3, 6, 2, 3, 7, 6, 1, 5, 0, 5, 4, 0, -1, -1, -1, -1},
+		{6, 2, 8, 6, 8, 7, 2, 1, 8, 4, 8, 5, 1, 5, 8, -1},
+		{9, 5, 4, 10, 1, 6, 1, 7, 6, 1, 3, 7, -1, -1, -1, -1},
+		{1, 6, 10, 1, 7, 6, 1, 0, 7, 8, 7, 0, 9, 5, 4, -1},
+		{4, 0, 10, 4, 10, 5, 0, 3, 10, 6, 10, 7, 3, 7, 10, -1},
+		{7, 6, 10, 7, 10, 8, 5, 4, 10, 4, 8, 10, -1, -1, -1, -1},
+		{6, 9, 5, 6, 11, 9, 11, 8, 9, -1, -1, -1, -1, -1, -1, -1},
+		{3, 6, 11, 0, 6, 3, 0, 5, 6, 0, 9, 5, -1, -1, -1, -1},
+		{0, 11, 8, 0, 5, 11, 0, 1, 5, 5, 6, 11, -1, -1, -1, -1},
+		{6, 11, 3, 6, 3, 5, 5, 3, 1, -1, -1, -1, -1, -1, -1, -1},
+		{1, 2, 10, 9, 5, 11, 9, 11, 8, 11, 5, 6, -1, -1, -1, -1},
+		{0, 11, 3, 0, 6, 11, 0, 9, 6, 5, 6, 9, 1, 2, 10, -1},
+		{11, 8, 5, 11, 5, 6, 8, 0, 5, 10, 5, 2, 0, 2, 5, -1},
+		{6, 11, 3, 6, 3, 5, 2, 10, 3, 10, 5, 3, -1, -1, -1, -1},
+		{5, 8, 9, 5, 2, 8, 5, 6, 2, 3, 8, 2, -1, -1, -1, -1},
+		{9, 5, 6, 9, 6, 0, 0, 6, 2, -1, -1, -1, -1, -1, -1, -1},
+		{1, 5, 8, 1, 8, 0, 5, 6, 8, 3, 8, 2, 6, 2, 8, -1},
+		{1, 5, 6, 2, 1, 6, -1, -1, -1, -1, -1, -1, -1, -1, -1, -1},
+		{1, 3, 6, 1, 6, 10, 3, 8, 6, 5, 6, 9, 8, 9, 6, -1},
+		{10, 1, 0, 10, 0, 6, 9, 5, 0, 5, 6, 0, -1, -1, -1, -1},
+		{0, 3, 8, 5, 6, 10, -1, -1, -1, -1, -1, -1, -1, -1, -1, -1},
+		{10, 5, 6, -1, -1, -1, -1, -1, -1, -1, -1, -1, -1, -1, -1, -1},
+		{11, 5, 10, 7, 5, 11, -1, -1, -1, -1, -1, -1, -1, -1, -1, -1},
+		{11, 5, 10, 11, 7, 5, 8, 3, 0, -1, -1, -1, -1, -1, -1, -1},
+		{5, 11, 7, 5, 10, 11, 1, 9, 0, -1, -1, -1, -1, -1, -1, -1},
+		{10, 7, 5, 10, 11, 7, 9, 8, 1, 8, 3, 1, -1, -1, -1, -1},
+		{11, 1, 2, 11, 7, 1, 7, 5, 1, -1, -1, -1, -1, -1, -1, -1},
+		{0, 8, 3, 1, 2, 7, 1, 7, 5, 7, 2, 11, -1, -1, -1, -1},
+		{9, 7, 5, 9, 2, 7, 9, 0, 2, 2, 11, 7, -1, -1, -1, -1},
+		{7, 5, 2, 7, 2, 11, 5, 9, 2, 3, 2, 8, 9, 8, 2, -1},
+		{2, 5, 10, 2, 3, 5, 3, 7, 5, -1, -1, -1, -1, -1, -1, -1},
+		{8, 2, 0, 8, 5, 2, 8, 7, 5, 10, 2, 5, -1, -1, -1, -1},
+		{9, 0, 1, 5, 10, 3, 5, 3, 7, 3, 10, 2, -1, -1, -1, -1},
+		{9, 8, 2, 9, 2, 1, 8, 7, 2, 10, 2, 5, 7, 5, 2, -1},
+		{1, 3, 5, 3, 7, 5, -1, -1, -1, -1, -1, -1, -1, -1, -1, -1},
+		{0, 8, 7, 0, 7, 1, 1, 7, 5, -1, -1, -1, -1, -1, -1, -1},
+		{9, 0, 3, 9, 3, 5, 5, 3, 7, -1, -1, -1, -1, -1, -1, -1},
+		{9, 8, 7, 5, 9, 7, -1, -1, -1, -1, -1, -1, -1, -1, -1, -1},
+		{5, 8, 4, 5, 10, 8, 10, 11, 8, -1, -1, -1, -1, -1, -1, -1},
+		{5, 0, 4, 5, 11, 0, 5, 10, 11, 11, 3, 0, -1, -1, -1, -1},
+		{0, 1, 9, 8, 4, 10, 8, 10, 11, 10, 4, 5, -1, -1, -1, -1},
+		{10, 11, 4, 10, 4, 5, 11, 3, 4, 9, 4, 1, 3, 1, 4, -1},
+		{2, 5, 1, 2, 8, 5, 2, 11, 8, 4, 5, 8, -1, -1, -1, -1},
+		{0, 4, 11, 0, 11, 3, 4, 5, 11, 2, 11, 1, 5, 1, 11, -1},
+		{0, 2, 5, 0, 5, 9, 2, 11, 5, 4, 5, 8, 11, 8, 5, -1},
+		{9, 4, 5, 2, 11, 3, -1, -1, -1, -1, -1, -1, -1, -1, -1, -1},
+		{2, 5, 10, 3, 5, 2, 3, 4, 5, 3, 8, 4, -1, -1, -1, -1},
+		{5, 10, 2, 5, 2, 4, 4, 2, 0, -1, -1, -1, -1, -1, -1, -1},
+		{3, 10, 2, 3, 5, 10, 3, 8, 5, 4, 5, 8, 0, 1, 9, -1},
+		{5, 10, 2, 5, 2, 4, 1, 9, 2, 9, 4, 2, -1, -1, -1, -1},
+		{8, 4, 5, 8, 5, 3, 3, 5, 1, -1, -1, -1, -1, -1, -1, -1},
+		{0, 4, 5, 1, 0, 5, -1, -1, -1, -1, -1, -1, -1, -1, -1, -1},
+		{8, 4, 5, 8, 5, 3, 9, 0, 5, 0, 3, 5, -1, -1, -1, -1},
+		{9, 4, 5, -1, -1, -1, -1, -1, -1, -1, -1, -1, -1, -1, -1, -1},
+		{4, 11, 7, 4, 9, 11, 9, 10, 11, -1, -1, -1, -1, -1, -1, -1},
+		{0, 8, 3, 4, 9, 7, 9, 11, 7, 9, 10, 11, -1, -1, -1, -1},
+		{1, 10, 11, 1, 11, 4, 1, 4, 0, 7, 4, 11, -1, -1, -1, -1},
+		{3, 1, 4, 3, 4, 8, 1, 10, 4, 7, 4, 11, 10, 11, 4, -1},
+		{4, 11, 7, 9, 11, 4, 9, 2, 11, 9, 1, 2, -1, -1, -1, -1},
+		{9, 7, 4, 9, 11, 7, 9, 1, 11, 2, 11, 1, 0, 8, 3, -1},
+		{11, 7, 4, 11, 4, 2, 2, 4, 0, -1, -1, -1, -1, -1, -1, -1},
+		{11, 7, 4, 11, 4, 2, 8, 3, 4, 3, 2, 4, -1, -1, -1, -1},
+		{2, 9, 10, 2, 7, 9, 2, 3, 7, 7, 4, 9, -1, -1, -1, -1},
+		{9, 10, 7, 9, 7, 4, 10, 2, 7, 8, 7, 0, 2, 0, 7, -1},
+		{3, 7, 10, 3, 10, 2, 7, 4, 10, 1, 10, 0, 4, 0, 10, -1},
+		{1, 10, 2, 8, 7, 4, -1, -1, -1, -1, -1, -1, -1, -1, -1, -1},
+		{4, 9, 1, 4, 1, 7, 7, 1, 3, -1, -1, -1, -1, -1, -1, -1},
+		{4, 9, 1, 4, 1, 7, 0, 8, 1, 8, 7, 1, -1, -1, -1, -1},
+		{4, 0, 3, 7, 4, 3, -1, -1, -1, -1, -1, -1, -1, -1, -1, -1},
+		{4, 8, 7, -1, -1, -1, -1, -1, -1, -1, -1, -1, -1, -1, -1, -1},
+		{9, 10, 8, 10, 11, 8, -1, -1, -1, -1, -1, -1, -1, -1, -1, -1},
+		{3, 0, 9, 3, 9, 11, 11, 9, 10, -1, -1, -1, -1, -1, -1, -1},
+		{0, 1, 10, 0, 10, 8, 8, 10, 11, -1, -1, -1, -1, -1, -1, -1},
+		{3, 1, 10, 11, 3, 10, -1, -1, -1, -1, -1, -1, -1, -1, -1, -1},
+		{1, 2, 11, 1, 11, 9, 9, 11, 8, -1, -1, -1, -1, -1, -1, -1},
+		{3, 0, 9, 3, 9, 11, 1, 2, 9, 2, 11, 9, -1, -1, -1, -1},
+		{0, 2, 11, 8, 0, 11, -1, -1, -1, -1, -1, -1, -1, -1, -1, -1},
+		{3, 2, 11, -1, -1, -1, -1, -1, -1, -1, -1, -1, -1, -1, -1, -1},
+		{2, 3, 8, 2, 8, 10, 10, 8, 9, -1, -1, -1, -1, -1, -1, -1},
+		{9, 10, 2, 0, 9, 2, -1, -1, -1, -1, -1, -1, -1, -1, -1, -1},
+		{2, 3, 8, 2, 8, 10, 0, 1, 8, 1, 10, 8, -1, -1, -1, -1},
+		{1, 10, 2, -1, -1, -1, -1, -1, -1, -1, -1, -1, -1, -1, -1, -1},
+		{1, 3, 8, 9, 1, 8, -1, -1, -1, -1, -1, -1, -1, -1, -1, -1},
+		{0, 9, 1, -1, -1, -1, -1, -1, -1, -1, -1, -1, -1, -1, -1, -1},
+		{0, 3, 8, -1, -1, -1, -1, -1, -1, -1, -1, -1, -1, -1, -1, -1},
+		{-1, -1, -1, -1, -1, -1, -1, -1, -1, -1, -1, -1, -1, -1, -1, -1}
+	};
+
+};
+
+
+} // end namespace UE::Geometry
+} // end namespace UE
\ No newline at end of file
