GeometryCore: create new GeometryCore Module. Move GeometricObjects code from GeometryProcessing plugin to this new module, as well as core FDynamicMesh3 classes (attributes, aabbtree/octree, change t

diff --git a/Engine/Source/Runtime/GeometryCore/Public/Generators/SweepGenerator.h b/Engine/Source/Runtime/GeometryCore/Public/Generators/SweepGenerator.h
new file mode 100644
--- /dev/null
+++ b/Engine/Source/Runtime/GeometryCore/Public/Generators/SweepGenerator.h
@@ -0,0 +1,844 @@
+// Copyright Epic Games, Inc. All Rights Reserved.
+
+// Port of geometry3Sharp CylinderGenerator
+
+#pragma once
+
+#include "CoreMinimal.h"
+
+#include "Math/UnrealMathUtility.h"
+#include "MeshShapeGenerator.h"
+#include "Misc/AssertionMacros.h"
+
+#include "CompGeom/PolygonTriangulation.h"
+#include "FrameTypes.h"
+#include "MatrixTypes.h"
+#include "Polygon2.h"
+#include "Curve/CurveUtil.h"
+
+#include "Util/ProgressCancel.h"
+
+namespace UE
+{
+namespace Geometry
+{
+
+/**
+ * ECapType indicates the type of cap to use on a sweep
+ */
+enum class /*GEOMETRYCORE_API*/ ECapType
+{
+	None = 0,
+	FlatTriangulation = 1,
+	FlatMidpointFan = 2
+	// TODO: Cone, other caps ...
+};
+
+class /*GEOMETRYCORE_API*/ FSweepGeneratorBase : public FMeshShapeGenerator
+{
+public:
+	virtual ~FSweepGeneratorBase()
+	{
+	}
+
+	/** If true, each quad gets a separate polygroup */
+	bool bPolygroupPerQuad = false;
+
+protected:
+	int32 CapVertStart[2], CapNormalStart[2], CapUVStart[2], CapTriangleStart[2], CapPolygonStart[2];
+
+	/**
+	 * Shared logic for creating vertex buffers and triangulations across all sweep primitives
+	 * Note: Does not set vertex positions or normals; a separate call must do that.
+	 */
+	void ConstructMeshTopology(const FPolygon2d& CrossSection,
+							   const TArrayView<const int32>& UVSections,
+							   const TArrayView<const int32>& NormalSections,
+							   const TArrayView<const int32>& SharpNormalsAlongLength,
+							   bool bEvenlySpaceUVs,
+							   const TArrayView<const FVector3d>& Path, // can be empty unless bEvenlySpaceUVs is true
+							   int32 NumCrossSections,
+							   bool bLoop,
+							   const ECapType Caps[2],
+							   FVector2f SectionsUVScale, FVector2f CapUVScale, FVector2f CapUVOffset)
+	{
+		// per cross section
+		int32 XVerts = CrossSection.VertexCount();
+		int32 XNormals = XVerts + NormalSections.Num();
+		int32 XUVs = XVerts + UVSections.Num() + 1;
+
+		double TotalPerimeter = 0, TotalPathLength = 0;
+		TArray<double> CrossSectionPercentages, PathPercentages;
+		if (bEvenlySpaceUVs)
+		{
+			CrossSectionPercentages.Add(0);
+			PathPercentages.Add(0);
+			for (int XIdx = 0, XNum = CrossSection.VertexCount(); XIdx < XNum; XIdx++)
+			{
+				double SegLen = CrossSection[XIdx].Distance(CrossSection[(XIdx + 1) % XNum]);
+				TotalPerimeter += SegLen;
+				CrossSectionPercentages.Add(TotalPerimeter);
+			}
+			TotalPerimeter = FMath::Max(TotalPerimeter, FMathd::ZeroTolerance);
+			for (int Idx = 1; Idx < CrossSectionPercentages.Num(); Idx++)
+			{
+				CrossSectionPercentages[Idx] /= TotalPerimeter;
+			}
+			int NumPathSegs = bLoop ? Path.Num() : Path.Num() - 1;
+			for (int PIdx = 0; PIdx < NumPathSegs; PIdx++)
+			{
+				double SegLen = Distance(Path[PIdx], Path[(PIdx + 1) % Path.Num()]);
+				TotalPathLength += SegLen;
+				PathPercentages.Add(TotalPathLength);
+			}
+			TotalPathLength = FMath::Max(TotalPathLength, FMathd::ZeroTolerance);
+			for (int Idx = 1; Idx < PathPercentages.Num(); Idx++)
+			{
+				PathPercentages[Idx] /= TotalPathLength;
+			}
+		}
+
+		int32 NumVerts = XVerts * NumCrossSections - (bLoop ? XVerts : 0);
+		int32 NumNormals = NumCrossSections > 1 ? (XNormals * NumCrossSections - (bLoop ? XNormals : 0)) : 0;
+		NumNormals += XNormals * SharpNormalsAlongLength.Num();
+		int32 NumUVs = NumCrossSections > 1 ? XUVs * NumCrossSections : 0;
+		int32 NumPolygons = (NumCrossSections - 1) * XVerts;
+		int32 NumTriangles = NumPolygons * 2;
+
+		TArray<FIndex3i> OutTriangles;
+
+		// doesn't make sense to have cap types if the sweep is a loop
+		ensure(!bLoop || (Caps[0] == ECapType::None && Caps[1] == ECapType::None));
+		
+		if (!bLoop)
+		{
+			for (int32 CapIdx = 0; !bLoop && CapIdx < 2; CapIdx++)
+			{
+				CapVertStart[CapIdx] = NumVerts;
+				CapNormalStart[CapIdx] = NumNormals;
+				CapUVStart[CapIdx] = NumUVs;
+				CapTriangleStart[CapIdx] = NumTriangles;
+				CapPolygonStart[CapIdx] = NumPolygons;
+
+				if (Caps[CapIdx] == ECapType::FlatTriangulation)
+				{
+					NumTriangles += XVerts - 2;
+					NumPolygons++;
+					NumUVs += XVerts;
+					NumNormals += XVerts;
+				}
+				else if (Caps[CapIdx] == ECapType::FlatMidpointFan)
+				{
+					NumTriangles += XVerts;
+					NumPolygons++;
+					NumUVs += XVerts + 1;
+					NumNormals += XVerts + 1;
+					NumVerts += 1;
+				}
+				// TODO: support more cap type; e.g.:
+				//else if (Caps[CapIdx] == ECapType::Cone)
+				//{
+				//	NumTriangles += XVerts;
+				//	NumPolygons += XVerts;
+				//	NumUVs += XVerts + 1;
+				//	NumNormals += XVerts * 2;
+				//	NumVerts += 1;
+				//}
+			}
+		}
+
+		SetBufferSizes(NumVerts, NumTriangles, NumUVs, NumNormals);
+
+		if (!bLoop)
+		{
+			for (int32 CapIdx = 0; CapIdx < 2; CapIdx++)
+			{
+				if (Caps[CapIdx] == ECapType::FlatTriangulation)
+				{
+					int32 VertOffset = CapIdx * (XVerts * (NumCrossSections - 1));
+
+					PolygonTriangulation::TriangulateSimplePolygon(CrossSection.GetVertices(), OutTriangles);
+					int32 TriIdx = CapTriangleStart[CapIdx];
+					int32 PolyIdx = CapPolygonStart[CapIdx];
+					for (const FIndex3i& Triangle : OutTriangles)
+					{
+						bool Flipped = CapIdx == 0;
+						Flipped = !Flipped;
+						SetTriangle(TriIdx,
+							Triangle.A + VertOffset, Triangle.B + VertOffset, Triangle.C + VertOffset,
+							Flipped);
+						SetTriangleUVs(TriIdx,
+							Triangle.A + CapUVStart[CapIdx],
+							Triangle.B + CapUVStart[CapIdx],
+							Triangle.C + CapUVStart[CapIdx],
+							Flipped);
+						SetTriangleNormals(TriIdx,
+							Triangle.A + CapNormalStart[CapIdx],
+							Triangle.B + CapNormalStart[CapIdx],
+							Triangle.C + CapNormalStart[CapIdx],
+							Flipped);
+						SetTrianglePolygon(TriIdx, PolyIdx);
+						TriIdx++;
+					}
+					float SideScale = 2 * CapIdx - 1;
+					for (int32 Idx = 0; Idx < XVerts; Idx++)
+					{
+						FVector2f CenteredVert = (FVector2f)CrossSection.GetVertices()[Idx] * CapUVScale + CapUVOffset;
+						SetUV(CapUVStart[CapIdx] + Idx, FVector2f(CenteredVert.X * SideScale, CenteredVert.Y), VertOffset + Idx);
+
+						// correct normal to be filled by subclass
+						SetNormal(CapNormalStart[CapIdx] + Idx, FVector3f::Zero(), VertOffset + Idx);
+					}
+				}
+				else if (Caps[CapIdx] == ECapType::FlatMidpointFan)
+				{
+					int32 VertOffset = CapIdx * (XVerts * (NumCrossSections - 1));
+					int32 CapVertStartIdx = CapVertStart[CapIdx];
+					int32 TriIdx = CapTriangleStart[CapIdx];
+					int32 PolyIdx = CapPolygonStart[CapIdx];
+					for (int32 VertIdx = 0; VertIdx < XVerts; VertIdx++)
+					{
+						bool Flipped = CapIdx == 0;
+						SetTriangle(TriIdx,
+							VertOffset + VertIdx,
+							CapVertStartIdx,
+							VertOffset + (VertIdx + 1) % XVerts,
+							Flipped);
+						SetTriangleUVs(TriIdx,
+							CapUVStart[CapIdx] + VertIdx,
+							CapUVStart[CapIdx] + XVerts,
+							CapUVStart[CapIdx] + (VertIdx + 1) % XVerts,
+							Flipped);
+						SetTriangleNormals(TriIdx,
+							CapNormalStart[CapIdx] + VertIdx,
+							CapNormalStart[CapIdx] + XVerts,
+							CapNormalStart[CapIdx] + (VertIdx + 1) % XVerts,
+							Flipped);
+						SetTrianglePolygon(TriIdx, PolyIdx);
+						TriIdx++;
+					}
+
+					// Set cap midpoint UV & Normal
+					// (correct normal to be filled by subclass)
+					SetUV(CapUVStart[CapIdx] + XVerts, FVector2f::Zero() + CapUVOffset, CapVertStartIdx);
+					SetNormal(CapNormalStart[CapIdx] + XVerts, FVector3f::Zero(), CapVertStartIdx);
+
+					// Set cap profile UVs & Normal
+					for (int32 Idx = 0; Idx < XVerts; Idx++)
+					{
+						FVector2f CenteredVert = (FVector2f)CrossSection.GetVertices()[Idx] * CapUVScale + CapUVOffset;
+						SetUV(CapUVStart[CapIdx] + Idx, FVector2f(CenteredVert.X, CenteredVert.Y), VertOffset + Idx);
+						SetNormal(CapNormalStart[CapIdx] + Idx, FVector3f::Zero(), VertOffset + Idx);
+					}
+				}
+			}
+		}
+
+		// fill in UVs and triangles along length
+		int MinValidCrossSections = bLoop ? 3 : 2;
+		int CurFaceGroupIndex = NumPolygons;
+		if (NumCrossSections >= MinValidCrossSections)
+		{
+			int32 UVSection = 0, UVSubIdx = 0;
+
+			int CrossSectionsMod = NumCrossSections;
+			if (bLoop)
+			{
+				CrossSectionsMod--; // last cross section becomes the first
+			}
+			int NormalCrossSectionsMod = CrossSectionsMod + SharpNormalsAlongLength.Num();
+
+			int32 NumSections = UVSections.Num();
+			int32 NextDupVertIdx = UVSection < NumSections ? UVSections[UVSection] : -1;
+			for (int32 VertSubIdx = 0; VertSubIdx < XVerts; UVSubIdx++)
+			{
+				float UVX = bEvenlySpaceUVs ? CrossSectionPercentages[VertSubIdx] : VertSubIdx / float(XVerts);
+				for (int32 XIdx = 0; XIdx < NumCrossSections; XIdx++)
+				{
+					float UVY = bEvenlySpaceUVs ? PathPercentages[XIdx] : XIdx / float(NumCrossSections - 1);
+					SetUV(XIdx * XUVs + UVSubIdx, FVector2f(1-UVX, 1-UVY) * SectionsUVScale, (XIdx % CrossSectionsMod) * XVerts + VertSubIdx);
+				}
+
+				if (VertSubIdx == NextDupVertIdx)
+				{
+					NextDupVertIdx = UVSection < NumSections ? UVSections[UVSection] : -1;
+				}
+				else
+				{
+					for (int32 XIdx = 0; XIdx + 1 < NumCrossSections; XIdx++)
+					{
+						SetTriangleUVs(
+							XVerts * 2 * XIdx + 2 * VertSubIdx,
+							XIdx * XUVs + UVSubIdx,
+							XIdx * XUVs + UVSubIdx + 1,
+							(XIdx + 1) * XUVs + UVSubIdx, true);
+						SetTriangleUVs(
+							XVerts * 2 * XIdx + 2 * VertSubIdx + 1,
+							(XIdx + 1) * XUVs + UVSubIdx + 1,
+							(XIdx + 1) * XUVs + UVSubIdx,
+							XIdx * XUVs + UVSubIdx + 1, true);
+					}
+					VertSubIdx++;
+				}
+			}
+			{
+				// final UV
+				float UVX = 1.0f;
+				int32 VertSubIdx = 0;
+				for (int32 XIdx = 0; XIdx < NumCrossSections; XIdx++)
+				{
+					float UVY = bEvenlySpaceUVs ? PathPercentages[XIdx] : XIdx / float(NumCrossSections - 1);
+					SetUV(XIdx * XUVs + UVSubIdx, FVector2f(1-UVX, 1-UVY) * SectionsUVScale, (XIdx % CrossSectionsMod) * XVerts + VertSubIdx);
+				}
+			}
+			NumSections = NormalSections.Num();
+			int32 NormalSection = 0;
+			NextDupVertIdx = NormalSection < NumSections ? NormalSections[NormalSection] : -1;
+			check(NextDupVertIdx < XVerts);
+			for (int32 VertSubIdx = 0, NormalSubIdx = 0; VertSubIdx < XVerts; NormalSubIdx++)
+			{
+				int SharpNormalIdx = 0;
+				for (int32 XIdx = 0, NormalXIdx = 0; XIdx < NumCrossSections; XIdx++, NormalXIdx++)
+				{
+					// just set the normal parent; don't compute normal yet
+					SetNormal((NormalXIdx % NormalCrossSectionsMod) * XNormals + NormalSubIdx, FVector3f(0, 0, 0), (XIdx % CrossSectionsMod) * XVerts + VertSubIdx);
+					// duplicate normals for cross sections that are 'sharp'
+					if (SharpNormalIdx < SharpNormalsAlongLength.Num() && XIdx == SharpNormalsAlongLength[SharpNormalIdx])
+					{
+						NormalXIdx++;
+						SetNormal((NormalXIdx % NormalCrossSectionsMod) * XNormals + NormalSubIdx, FVector3f(0, 0, 0), (XIdx % CrossSectionsMod) * XVerts + VertSubIdx);
+						SharpNormalIdx++;
+					}
+				}
+
+				if (VertSubIdx == NextDupVertIdx)
+				{
+					NextDupVertIdx = NormalSection < NumSections ? NormalSections[NormalSection] : -1;
+					check(NextDupVertIdx < XVerts);
+				}
+				else
+				{
+					int32 WrappedNextNormalSubIdx = (NormalSubIdx + 1) % XNormals;
+					int32 WrappedNextVertexSubIdx = (VertSubIdx + 1) % XVerts;
+					SharpNormalIdx = 0;
+					for (int32 XIdx = 0, NXIdx = 0; XIdx + 1 < NumCrossSections; XIdx++, NXIdx++)
+					{
+						int32 T0Idx = XVerts * 2 * XIdx + 2 * VertSubIdx;
+						int32 T1Idx = T0Idx + 1;
+						int32 PIdx = XVerts * XIdx + VertSubIdx;
+						int32 NextXIdx = (XIdx + 1) % CrossSectionsMod;
+						int32 NextNXIdx = (NXIdx + 1) % NormalCrossSectionsMod;
+						SetTrianglePolygon(T0Idx, (bPolygroupPerQuad) ? PIdx : (CurFaceGroupIndex+XIdx) );
+						SetTrianglePolygon(T1Idx, (bPolygroupPerQuad) ? PIdx : (CurFaceGroupIndex+XIdx) );
+						SetTriangle(T0Idx,
+									XIdx * XVerts + VertSubIdx,
+									XIdx * XVerts + WrappedNextVertexSubIdx,
+									NextXIdx * XVerts + VertSubIdx, true);
+						SetTriangle(T1Idx,
+									NextXIdx * XVerts + WrappedNextVertexSubIdx,
+									NextXIdx * XVerts + VertSubIdx,
+									XIdx * XVerts + WrappedNextVertexSubIdx, true);
+						SetTriangleNormals(
+							T0Idx,
+							NXIdx * XNormals + NormalSubIdx,
+							NXIdx * XNormals + WrappedNextNormalSubIdx,
+							NextNXIdx * XNormals + NormalSubIdx, true);
+						SetTriangleNormals(
+							T1Idx,
+							NextNXIdx * XNormals + WrappedNextNormalSubIdx,
+							NextNXIdx * XNormals + NormalSubIdx,
+							NXIdx * XNormals + WrappedNextNormalSubIdx, true);
+						if (SharpNormalIdx < SharpNormalsAlongLength.Num() && XIdx+1 == SharpNormalsAlongLength[SharpNormalIdx])
+						{
+							NXIdx++;
+							SharpNormalIdx++;
+						}
+					}
+					VertSubIdx++;
+				}
+			}
+		}
+	}
+};
+
+/**
+* Generate a cylinder with optional end caps
+*/
+class /*GEOMETRYCORE_API*/ FVerticalCylinderGeneratorBase : public FSweepGeneratorBase
+{
+public:
+	int AngleSamples = 16;
+	bool bCapped = false;
+	bool bUVScaleMatchSidesAndCaps = true;
+	ECapType CapType = ECapType::FlatMidpointFan;
+
+	static float ComputeSegLengths(const TArrayView<float>& Radii, const TArrayView<float>& Heights, TArray<float>& AlongPercents)
+	{
+		float LenAlong = 0;
+		int32 NumX = Radii.Num();
+		AlongPercents.SetNum(NumX);
+		AlongPercents[0] = 0;
+		for (int XIdx = 0; XIdx+1 < NumX; XIdx++)
+		{
+			float Dist = FVector2d(Radii[XIdx], Heights[XIdx]).Distance(FVector2d(Radii[XIdx + 1], Heights[XIdx + 1]));
+			LenAlong += Dist;
+			AlongPercents[XIdx + 1] = LenAlong;
+		}
+		for (int XIdx = 0; XIdx+1 < NumX; XIdx++)
+		{
+			AlongPercents[XIdx+1] /= LenAlong;
+		}
+		return LenAlong;
+	}
+
+	bool GenerateVerticalCircleSweep(const TArrayView<float>& Radii, const TArrayView<float>& Heights, const TArrayView<int>& SharpNormalsAlongLength)
+	{
+		FPolygon2d X = FPolygon2d::MakeCircle(1.0, AngleSamples);
+		const TArray<FVector2d>& XVerts = X.GetVertices();
+		ECapType Caps[2] = {ECapType::None, ECapType::None};
+
+		if (bCapped)
+		{
+			Caps[0] = CapType;
+			Caps[1] = CapType;
+		}
+
+		int NumX = Radii.Num();
+		if (!ensure(NumX == Heights.Num()))
+		{
+			return false;
+		}
+		// first and last cross sections can't be sharp, so can't have more than NumX-2 sharp normal indices
+		if (!ensure(SharpNormalsAlongLength.Num() + 2 <= NumX))
+		{
+			return false;
+		}
+
+		TArray<float> AlongPercents;
+		float LenAlong = ComputeSegLengths(Radii, Heights, AlongPercents);
+
+		ConstructMeshTopology(X, {}, {}, SharpNormalsAlongLength, false, {}, NumX, false, Caps, FVector2f(1, 1), FVector2f(.5f, .5f), FVector2f(.5f, .5f));
+
+		TArray<FVector2d> NormalSides; NormalSides.SetNum(NumX - 1);
+		for (int XIdx = 0; XIdx+1 < NumX; XIdx++)
+		{
+			FVector2d Vec = FVector2d(Radii[XIdx + 1], Heights[XIdx + 1]) - FVector2d(Radii[XIdx], Heights[XIdx]);
+			NormalSides[XIdx] = Normalized(PerpCW(Vec));
+		}
+		TArray<FVector2d> SmoothedNormalSides; SmoothedNormalSides.SetNum(NumX);
+		// smooth internal normals
+		SmoothedNormalSides[0] = NormalSides[0];
+		SmoothedNormalSides.Last() = NormalSides.Last();
+		for (int XIdx = 1; XIdx + 1 < NumX; XIdx++)
+		{
+			SmoothedNormalSides[XIdx] = Normalized(NormalSides[XIdx] + NormalSides[XIdx - 1]);
+		}
+
+
+		// set vertex positions and normals for all cross sections along length
+		for (int SubIdx = 0; SubIdx < X.VertexCount(); SubIdx++)
+		{
+			int SharpNormalIdx = 0;
+			for (int XIdx = 0, NormalXIdx = 0; XIdx < NumX; ++XIdx, ++NormalXIdx)
+			{
+				double Along = AlongPercents[XIdx];
+				double AlongRadius = Radii[XIdx];
+				Vertices[SubIdx + XIdx * AngleSamples] =
+					FVector3d(XVerts[SubIdx].X * AlongRadius, XVerts[SubIdx].Y * AlongRadius, Heights[XIdx]);
+				if (SharpNormalIdx < SharpNormalsAlongLength.Num() && XIdx == SharpNormalsAlongLength[SharpNormalIdx])
+				{
+					// write sharp normals
+					if (ensure(XIdx > 0)) // very first index cannot be sharp
+					{
+						Normals[SubIdx + NormalXIdx * AngleSamples] = FVector3f(XVerts[SubIdx].X*NormalSides[XIdx-1].X, XVerts[SubIdx].Y*NormalSides[XIdx-1].X, NormalSides[XIdx-1].Y);
+					}
+					NormalXIdx++;
+					if (ensure(XIdx + 1 < NumX)) // very last index cannot be sharp
+					{
+						Normals[SubIdx + NormalXIdx * AngleSamples] = FVector3f(XVerts[SubIdx].X*NormalSides[XIdx].X, XVerts[SubIdx].Y*NormalSides[XIdx].X, NormalSides[XIdx].Y);
+					}
+					SharpNormalIdx++;
+				}
+				else
+				{
+					// write smoothed normal
+					Normals[SubIdx + NormalXIdx * AngleSamples] = FVector3f(XVerts[SubIdx].X*SmoothedNormalSides[XIdx].X, XVerts[SubIdx].Y*SmoothedNormalSides[XIdx].X, SmoothedNormalSides[XIdx].Y);
+				}
+			}
+		}
+		
+		if (bCapped)
+		{
+			// if capped, set vertices.
+			for (int CapIdx = 0; CapIdx < 2; CapIdx++)
+			{
+				if (Caps[CapIdx] == ECapType::FlatMidpointFan)
+				{
+					Vertices[CapVertStart[CapIdx]] = FVector3d::UnitZ() * (double)Heights[CapIdx];
+				}
+			}
+
+			// if capped, set top/bottom normals
+			for (int SubIdx = 0; SubIdx < X.VertexCount(); SubIdx++)
+			{
+				for (int XBotTop = 0; XBotTop < 2; ++XBotTop)
+				{
+					Normals[CapNormalStart[XBotTop] + SubIdx] = FVector3f(0, 0, 2 * XBotTop - 1);
+				}
+			}
+			for (int CapIdx = 0; CapIdx < 2; CapIdx++)
+			{
+				if (Caps[CapIdx] == ECapType::FlatMidpointFan)
+				{
+					Normals[CapNormalStart[CapIdx] + X.VertexCount()] = FVector3f(0, 0, 2 * CapIdx - 1);
+				}
+			}
+		}
+
+		if (bUVScaleMatchSidesAndCaps)
+		{
+			float MaxAbsRad = FMathf::Abs(Radii[0]);
+			for (int XIdx = 0; XIdx < NumX; XIdx++)
+			{
+				MaxAbsRad = FMathf::Max(FMathf::Abs(Radii[XIdx]), MaxAbsRad);
+			}
+			float AbsHeight = LenAlong;
+			float MaxAbsCircumference = MaxAbsRad * FMathf::TwoPi;
+
+			// scales to put each differently-scaled UV coordinate into the same space
+			float ThetaScale = MaxAbsCircumference;
+			float HeightScale = AbsHeight;
+			float CapScale = MaxAbsRad*2;
+
+			float MaxScale = FMathf::Max3(ThetaScale, HeightScale, CapScale);
+			ThetaScale /= MaxScale;
+			HeightScale /= MaxScale;
+			CapScale /= MaxScale;
+			for (int UVIdx = 0; UVIdx < CapUVStart[0]; UVIdx++)
+			{
+				UVs[UVIdx].X *= ThetaScale;
+				UVs[UVIdx].Y *= HeightScale;
+			}
+			for (int UVIdx = CapUVStart[0]; UVIdx < UVs.Num(); UVIdx++)
+			{
+				UVs[UVIdx] *= CapScale;
+			}
+		}
+
+		return true;
+	}
+};
+
+/**
+ * Generate a cylinder with optional end caps
+ */
+class /*GEOMETRYCORE_API*/ FCylinderGenerator : public FVerticalCylinderGeneratorBase
+{
+public:
+	float Radius[2] = {1.0f, 1.0f};
+	float Height = 1.0f;
+	int LengthSamples = 0;
+
+public:
+	/** Generate the mesh */
+	virtual FMeshShapeGenerator& Generate() override
+	{
+		if (bCapped)
+		{
+			AngleSamples = FMath::Max(AngleSamples, 3);
+		}
+
+		TArray<float> Radii, Heights;
+
+		Radii.Add(Radius[0]);
+		Heights.Add(0);
+		for (int ExtraIdx = 0; ExtraIdx < LengthSamples; ExtraIdx++)
+		{
+			float Along = float(ExtraIdx + 1) / float(LengthSamples + 1);
+			Radii.Add(FMath::Lerp(Radius[0], Radius[1], Along));
+			Heights.Add(Height * Along);
+		}
+		Radii.Add(Radius[1]);
+		Heights.Add(Height);
+
+		GenerateVerticalCircleSweep(Radii, Heights, {});
+
+		return *this;
+	}
+};
+
+/**
+* Generate a 3D arrow
+*/
+class /*GEOMETRYCORE_API*/ FArrowGenerator : public FVerticalCylinderGeneratorBase
+{
+public:
+	float StickRadius = 0.5f;
+	float StickLength = 1.0f;
+	float HeadBaseRadius = 1.0f;
+	float HeadTipRadius = 0.01f;
+	float HeadLength = 0.5f;
+
+	int AdditionalLengthSamples[3]{ 0,0,0 }; // additional length-wise samples on the three segments (along stick, along arrow base, along arrow cone)
+
+	void DistributeAdditionalLengthSamples(int TargetSamples)
+	{
+		TArray<float> AlongPercents;
+		TArray<float> Radii{ StickRadius, StickRadius, HeadBaseRadius, HeadTipRadius };
+		TArray<float> Heights{ 0, StickLength, StickLength, StickLength + HeadLength };
+		float LenAlong = ComputeSegLengths(Radii, Heights, AlongPercents);
+		for (int Idx = 0; Idx < 3; Idx++)
+		{
+			AdditionalLengthSamples[Idx] = (int)(.5f+AlongPercents[Idx + 1] * TargetSamples);
+		}
+	}
+
+public:
+	/** Generate the mesh */
+	virtual FMeshShapeGenerator& Generate() override
+	{
+		TArray<float> Radii, Heights;
+		TArray<int> SharpNormalsAlongLength;
+		float SrcRadii[] { StickRadius, StickRadius, HeadBaseRadius, HeadTipRadius };
+		float SrcHeights[] { 0, StickLength, StickLength, StickLength + HeadLength };
+
+		int SegIdx = 0;
+		for (; SegIdx < 3; SegIdx++)
+		{
+			Radii.Add(SrcRadii[SegIdx]);
+			Heights.Add(SrcHeights[SegIdx]);
+			for (int ExtraIdx = 0; ExtraIdx < AdditionalLengthSamples[SegIdx]; ExtraIdx++)
+			{
+				float Along = float(ExtraIdx + 1) / float(AdditionalLengthSamples[SegIdx] + 1);
+				Radii.Add(FMath::Lerp(SrcRadii[SegIdx], SrcRadii[SegIdx + 1], Along));
+				Heights.Add(FMath::Lerp(SrcHeights[SegIdx], SrcHeights[SegIdx + 1], Along));
+			}
+		}
+		SharpNormalsAlongLength.Add(1 + AdditionalLengthSamples[0]);
+		SharpNormalsAlongLength.Add(SharpNormalsAlongLength.Last() + 1 + AdditionalLengthSamples[1]);
+		Radii.Add(SrcRadii[SegIdx]);
+		Heights.Add(SrcHeights[SegIdx]);
+		
+		GenerateVerticalCircleSweep(Radii, Heights, SharpNormalsAlongLength);
+
+		return *this;
+	}
+};
+
+
+/**
+ * Sweep a 2D Profile Polygon along a 3D Path.
+ * 
+ * TODO: 
+ *  - a custom variant for toruses specifically (would be faster)
+ *  - Mitering cross sections support?
+ */
+class /*GEOMETRYCORE_API*/ FGeneralizedCylinderGenerator : public FSweepGeneratorBase
+{
+public:
+	FPolygon2d CrossSection;
+	TArray<FVector3d> Path;
+
+	FFrame3d InitialFrame;
+
+	bool bCapped = false;
+	bool bLoop = false;
+	ECapType CapType = ECapType::FlatTriangulation;
+
+	// When true, the generator attempts to scale UV's in a way that preserves scaling across different mesh
+	// results, aiming for 1.0 in UV space to be equal to UnitUVInWorldCoordinates in world space. This in
+	// practice means adjusting the U scale relative to the CrossSection curve length and V scale relative
+	// to the distance between vertices on the Path.
+	bool bUVScaleRelativeWorld = false;
+
+	// Only relevant if bUVScaleRelativeWorld is true (see that description)
+	float UnitUVInWorldCoordinates = 100;
+
+public:
+	/** Generate the mesh */
+	virtual FMeshShapeGenerator& Generate() override
+	{
+		const TArray<FVector2d>& XVerts = CrossSection.GetVertices();
+		ECapType Caps[2] = {ECapType::None, ECapType::None};
+
+		if (bCapped && !bLoop)
+		{
+			Caps[0] = CapType;
+			Caps[1] = CapType;
+		}
+		int PathNum = Path.Num();
+		
+		FAxisAlignedBox2f Bounds = (FAxisAlignedBox2f)CrossSection.Bounds();
+		double BoundsMaxDimInv = 1.0 / FMathd::Max(Bounds.MaxDim(), .001);
+		FVector2f SectionScale(1, 1), CapScale(BoundsMaxDimInv, BoundsMaxDimInv);
+		if (bUVScaleRelativeWorld)
+		{
+			double Perimeter = CrossSection.Perimeter();
+			double PathLen = UE::Geometry::CurveUtil::ArcLength<double, FVector3d>(Path, bLoop);
+			SectionScale.X = Perimeter / UnitUVInWorldCoordinates;
+			SectionScale.Y = PathLen / UnitUVInWorldCoordinates;
+			CapScale.X = CapScale.Y = 1.0f / UnitUVInWorldCoordinates;
+		}
+		ConstructMeshTopology(CrossSection, {}, {}, {}, true, Path, PathNum + (bLoop ? 1 : 0), bLoop, Caps, SectionScale, CapScale, Bounds.Center());
+
+		int XNum = CrossSection.VertexCount();
+		TArray<FVector2d> XNormals; XNormals.SetNum(XNum);
+		for (int Idx = 0; Idx < XNum; Idx++)
+		{
+			XNormals[Idx] = CrossSection.GetNormal_FaceAvg(Idx);
+		}
+
+		FFrame3d CrossSectionFrame = InitialFrame;
+		for (int PathIdx = 0; PathIdx < PathNum; ++PathIdx)
+		{
+			FVector3d Tangent = UE::Geometry::CurveUtil::Tangent<double, FVector3d>(Path, PathIdx, bLoop);
+			CrossSectionFrame.AlignAxis(2, Tangent);
+			FVector3d C = Path[PathIdx];
+			FVector3d X = CrossSectionFrame.X();
+			FVector3d Y = CrossSectionFrame.Y();
+			for (int SubIdx = 0; SubIdx < XNum; SubIdx++)
+			{
+				FVector2d XP = CrossSection[SubIdx];
+				FVector2d XN = XNormals[SubIdx];
+				Vertices[SubIdx + PathIdx * XNum] = C + X * XP.X + Y * XP.Y;
+				Normals[SubIdx + PathIdx * XNum] = (FVector3f)(X * XN.X + Y * XN.Y);
+			}
+		}
+		if (bCapped && !bLoop)
+		{
+			for (int CapIdx = 0; CapIdx < 2; CapIdx++)
+			{
+				FVector3d Normal = CurveUtil::Tangent<double, FVector3d>(Path, CapIdx * (PathNum - 1), bLoop) * (double)(CapIdx * 2 - 1);
+				for (int SubIdx = 0; SubIdx < XNum; SubIdx++)
+				{
+					Normals[CapNormalStart[CapIdx] + SubIdx] = (FVector3f)Normal;
+				}
+			}
+		}
+
+		for (int k = 0; k < Normals.Num(); ++k)
+		{
+			Normalize(Normals[k]);
+		}
+
+		return *this;
+	}
+};
+
+enum class EProfileSweepPolygonGrouping : uint8
+{
+	/** One polygroup for entire output mesh */
+	Single,
+	/** One polygroup per mesh quad/triangle */
+	PerFace,
+
+	/* One polygroup per strip that represents a step along the sweep curve. */
+	PerSweepSegment,
+	/* One polygroup per strip coming from each individual edge of the profile curve. */
+	PerProfileSegment
+};
+
+enum class EProfileSweepQuadSplit : uint8
+{
+	/** Always split the quad in the same way relative sweep direction and profile direction. */
+	Uniform,
+	/** Split the quad to connect the shortest diagonal. */
+	ShortestDiagonal
+};
+
+/**
+ * Much like FGeneralizedCylinderGenerator, but allows an arbitrary profile curve to be swept, and gives
+ * control over the frames of the sweep curve. A mesh will be properly oriented if the profile curve is
+ * oriented counterclockwise when facing down the direction in which it is being swept.
+ *
+ * Because it supports open profile curves, as well as welded points (for welding points on an axis of rotation), 
+ * it cannot actually use the utility function from FSweepGeneratorBase, and so it doesn't inherit from 
+ * that class.
+ */
+class GEOMETRYCORE_API FProfileSweepGenerator : public FMeshShapeGenerator
+{
+public:
+
+	// Curve that will be swept along the curve, given in coordinates of the frames used in the sweep curve.
+	TArray<FVector3d> ProfileCurve;
+
+	// Curve along which to sweep the profile curve
+	TArray<FFrame3d> SweepCurve;
+
+	// Indices into ProfileCurve that should not be swept along the curve, instead being instantiated
+	// just once. This is useful for welding vertices on an axis of rotation if the sweep curve denotes
+	// a revolution.
+	TSet<int32> WeldedVertices;
+
+	// Generated UV coordinates will be multiplied by these values.
+	FVector2d UVScale = FVector2d(1,1);
+
+	// These values will be added to the generated UV coordinates after applying UVScale.
+	FVector2d UVOffset = FVector2d(0, 0);
+
+	// When true, the generator attempts to scale UV's in a way that preserves scaling across different mesh
+	// results, aiming for 1.0 in UV space to be equal to UnitUVInWorldCoordinates in world space. This is 
+	// generally speaking unrealistic because UV's are going to be variably stretched no matter what, but 
+	// in practice it means adjusting the V scale relative to the profile curve length and U scale relative
+	// to a very crude measurement of movement across sweep frames.
+	bool bUVScaleRelativeWorld = false;
+
+	// Only relevant if bUVScaleRelativeWorld is true (see that description)
+	float UnitUVInWorldCoordinates = 100;
+
+	// If true, the last point of the sweep curve is considered to be connected to the first.
+	bool bSweepCurveIsClosed = false;
+
+	// If true, the last point of the profile curve is considered to be connected to the first.
+	bool bProfileCurveIsClosed = false;
+
+	// If true, each triangle will have its own normals at each vertex, rather than sharing averaged
+	// ones with nearby triangles.
+	bool bSharpNormals = true;
+
+	// If true, welded-to-welded connections in the profile curve (which can't result in triangles)
+	// do not affect the UV layout.
+	bool bUVsSkipFullyWeldedEdges = true;
+
+	EProfileSweepQuadSplit QuadSplitMethod = EProfileSweepQuadSplit::ShortestDiagonal;
+
+	// When QuadSplitMode is ShortestDiagonal, biases one of the diagonals so that symmetric
+	// quads are split uniformly. The tolerance is a proportion allowable difference.
+	double DiagonalTolerance = 0.01;
+
+	EProfileSweepPolygonGrouping PolygonGroupingMode = EProfileSweepPolygonGrouping::PerFace;
+
+	// If not null, this pointer is intermittently used to check whether the current operation should stop early
+	FProgressCancel* Progress = nullptr;
+
+	// TODO: We could allow the user to dissallow bowtie vertex creation, which currently could 
+	// happen depending on which vertices are welded.
+
+public:
+
+	/** Generate the mesh */
+	virtual FMeshShapeGenerator& Generate() override;
+
+	/** If the sweep curve is not closed, this will store the vertex ids of the first and last instances
+	 * of the profile curve. Note that even if the profile curve is closed, depending on the welding,
+	 * these could be part of a single boundary (ie, a square revolved 90 degrees around a welded side
+	 * actually has one open boundary rather than two, since they are joined), but the user likely
+	 * wants to be given them separately for ease in making end caps.
+	 */
+	TArray<int32> EndProfiles[2];
+
+	// TODO: We could output other boundaries too, but that's probably only worth doing once we find
+	// a case where we would actually use them.
+protected:
+
+	void InitializeUvBuffer(const TArray<int32>& VertPositionOffsets, 
+		int32& NumUvRowsOut, int32& NumUvColumnsOut);
+	void AdjustNormalsForTriangle(int32 TriIndex, int32 FirstIndex, int32 SecondIndex, int32 ThirdIndex,
+		TArray<FVector3d>& WeightedNormals);
+	void AdjustNormalsForTriangle(int32 TriIndex, int32 FirstIndex, int32 SecondIndex, int32 ThirdIndex,
+		TArray<FVector3d>& WeightedNormals, const FVector3d& AbNormalized);
+};
+
+
+} // end namespace UE::Geometry
+} // end namespace UE
\ No newline at end of file
diff --git a/Engine/Source/Runtime/GeometryCore/Public/Image/ImageDimensions.h b/Engine/Source/Runtime/GeometryCore/Public/Image/ImageDimensions.h
new file mode 100644
--- /dev/null
+++ b/Engine/Source/Runtime/GeometryCore/Public/Image/ImageDimensions.h
@@ -0,0 +1,270 @@
+// Copyright Epic Games, Inc. All Rights Reserved.
+
+#pragma once
+
+#include "VectorTypes.h"
+#include "IntVectorTypes.h"
+
+namespace UE
+{
+namespace Geometry
+{
+
+/**
+ * FImageDimensions provides various functions for working with size/indices/coordinates of a
+ * 2D image, as well as standard UV spaces
+ */
+class FImageDimensions
+{
+protected:
+	int32 Width = 0;
+	int32 Height = 0;
+
+	bool bIsTile = false;
+	int32 SourceWidth = 0;
+	int32 SourceHeight = 0;
+	FVector2i SourceOffset = FVector2i::Zero();
+
+public:
+	FImageDimensions(int32 WidthIn = 0, int32 HeightIn = 0)
+		: SourceWidth(WidthIn), SourceHeight(HeightIn)
+	{
+		SetDimensions(WidthIn, HeightIn);
+	}
+
+	/** Tile constructor. */
+	FImageDimensions(int32 WidthIn, int32 HeightIn, int32 SourceWidthIn, int32 SourceHeightIn, FVector2i SourceOffsetIn)
+		: bIsTile(true), SourceWidth(SourceWidthIn), SourceHeight(SourceHeightIn), SourceOffset(SourceOffsetIn)
+	{
+		SetDimensions(WidthIn, HeightIn);
+		checkSlow(Width >= 0 && Width <= SourceWidth && SourceOffset.X >= 0 && Width + SourceOffset.X <= SourceWidth);
+		checkSlow(Height >= 0 && Height <= SourceHeight && SourceOffset.Y >= 0 && Height + SourceOffset.Y <= SourceHeight);
+	}
+
+	/** Set the dimensions of this image. */
+	void SetDimensions(int32 WidthIn, int32 HeightIn)
+	{
+		check(WidthIn >= 0 && HeightIn >= 0);
+		Width = WidthIn;
+		Height = HeightIn;
+	}
+
+	/** @return width of image */
+	int32 GetWidth() const { return Width; }
+	/** @return height of image */
+	int32 GetHeight() const { return Height; }
+	/** @return number of elements in image */
+	int64 Num() const { return (int64)Width * (int64)Height; }
+
+	/** @return true if image is square */
+	bool IsSquare() const { return bIsTile ? SourceWidth == SourceHeight : Width == Height; }
+
+	/** @return true if coordinates are valid, ie in-bounds of image dimensions */
+	bool IsValidCoords(const FVector2i& Coords) const
+	{
+		return (Coords.X >= 0 && Coords.X < Width && Coords.Y >= 0 && Coords.Y < Height);
+	}
+
+	/** Clamp input coordinates to valid range of image coordinates */
+	void Clamp(int32& X, int32& Y) const
+	{
+		X = FMath::Clamp(X, 0, Width - 1);
+		Y = FMath::Clamp(Y, 0, Height - 1);
+	}
+
+	/** Clamp input coordinates to valid range of image coordinates */
+	void Clamp(FVector2i& Coords) const
+	{
+		Coords.X = FMath::Clamp(Coords.X, 0, Width - 1);
+		Coords.Y = FMath::Clamp(Coords.Y, 0, Height - 1);
+	}
+
+	/** @return linear index into image from 2D coordinates */
+	int64 GetIndex(int32 X, int32 Y) const
+	{
+		return (Y * Width) + X;
+	}
+
+	/** @return linear index into image from 2D coordinates */
+	int64 GetIndex(const FVector2i& Coords) const
+	{
+		checkSlow(IsValidCoords(Coords));
+		return ((int64)Coords.Y * (int64)Width) + (int64)Coords.X;
+	}
+
+	/** @return linear index into image from 2D coordinates, optionally flipped around X and Y axes */
+	int64 GetIndexMirrored(const FVector2i& Coords, bool bFlipX, bool bFlipY) const
+	{
+		checkSlow(IsValidCoords(Coords));
+		int64 UseX = (bFlipX) ? (Width - 1 - Coords.X) : Coords.X;
+		int64 UseY = (bFlipY) ? (Height - 1 - Coords.Y) : Coords.Y;
+		return (UseY * (int64)Width) + UseX;
+	}
+
+	/** @return 2D image coordinates from linear index */
+	FVector2i GetCoords(int64 LinearIndex) const
+	{
+		checkSlow(LinearIndex >= 0 && LinearIndex < Num());
+		return FVector2i((int32)(LinearIndex % (int64)Width), (int32)(LinearIndex / (int64)Width));
+	}
+
+	/** @return Real-valued dimensions of a pixel/texel in the image, relative to default UV space [0..1]^2 */
+	FVector2d GetTexelSize() const
+	{
+		return FVector2d(1.0 / (double)SourceWidth, 1.0 / (double)SourceHeight);
+	}
+
+	/** @return Real-valued position of given texel center in default UV-space [0..1]^2 */
+	FVector2d GetTexelUV(const FVector2i& Coords) const
+	{
+		if (bIsTile)
+		{
+			return FVector2d(
+				((double)SourceOffset.X + (double)Coords.X + 0.5) / (double)SourceWidth,
+				((double)SourceOffset.Y + (double)Coords.Y + 0.5) / (double)SourceWidth);
+		}
+		else
+		{
+			return FVector2d(
+				((double)Coords.X + 0.5) / (double)Width,
+				((double)Coords.Y + 0.5) / (double)Height);
+		}
+	}
+
+	/** @return Real-valued position of given texel center in default UV-space [0..1]^2 */
+	FVector2d GetTexelUV(int64 LinearIndex) const
+	{
+		return GetTexelUV(GetCoords(LinearIndex));
+	}
+
+	/** @return integer XY coordinates for real-valued XY coordinates (ie texel that contains value, if texel origin is in bottom-left */
+	FVector2i PixelToCoords(const FVector2d& PixelPosition) const
+	{
+		if (bIsTile)
+		{
+			int32 X = FMath::Clamp((int32)PixelPosition.X, 0, SourceWidth - 1);
+			int32 Y = FMath::Clamp((int32)PixelPosition.Y, 0, SourceHeight - 1);
+			return FVector2i(X, Y);
+		}
+		else
+		{
+			int32 X = FMath::Clamp((int32)PixelPosition.X, 0, Width - 1);
+			int32 Y = FMath::Clamp((int32)PixelPosition.Y, 0, Height - 1);
+			return FVector2i(X, Y);
+		}
+	}
+
+	/** @return integer XY coordinates for UV coordinates, assuming default UV space [0..1]^2 */
+	FVector2i UVToCoords(const FVector2d& UVPosition) const
+	{
+		if (bIsTile)
+		{
+			double PixelX = (UVPosition.X * (double)SourceWidth /*- 0.5*/);
+			double PixelY = (UVPosition.Y * (double)SourceHeight /*- 0.5*/);
+			return PixelToCoords(FVector2d(PixelX, PixelY));
+		}
+		else
+		{
+			double PixelX = (UVPosition.X * (double)Width /*- 0.5*/);
+			double PixelY = (UVPosition.Y * (double)Height /*- 0.5*/);
+			return PixelToCoords(FVector2d(PixelX, PixelY));
+		}
+	}
+
+	bool operator==(const FImageDimensions& Other) const
+	{
+		return Width == Other.Width && Height == Other.Height &&
+			bIsTile == Other.bIsTile &&
+			SourceWidth == Other.SourceWidth &&
+			SourceHeight == Other.SourceHeight &&
+			SourceOffset == Other.SourceOffset;
+	}
+
+	bool operator!=(const FImageDimensions& Other) const
+	{
+		return !(*this == Other);
+	}
+
+	// Tiling
+
+	/** @return true if this image is a tile. */
+	bool IsTile() const
+	{
+		return bIsTile;
+	}
+
+	/** @return source image dimensions for this tile. */
+	FImageDimensions GetSourceDimensions() const
+	{
+		return FImageDimensions(SourceWidth, SourceHeight);
+	}
+
+	/** @return source image linear index from tile coordinates. */
+	int64 GetSourceIndex(const FVector2i& Coords) const
+	{
+		FVector2i SourceCoords = GetSourceCoords(Coords.X, Coords.Y);
+		return GetSourceDimensions().GetIndex(SourceCoords);
+	}
+
+	/** @return source image coordinates from tile coordinates. */
+	FVector2i GetSourceCoords(int32 TileX, int32 TileY) const
+	{
+		checkSlow(IsValidCoords(FVector2i(TileX, TileY)));
+		return FVector2i(SourceOffset.X + TileX, SourceOffset.Y + TileY);
+	}
+};
+
+class FImageTiling
+{
+private:
+	FImageDimensions Dimensions;
+	int32 TileWidth = 32;
+	int32 TileHeight = 32;
+	int32 TilePadding = 0;
+
+public:
+	FImageTiling(FImageDimensions DimensionsIn, int32 TileWidthIn, int32 TileHeightIn, int32 TilePaddingIn = 0)
+	{
+		SetTiling(DimensionsIn, TileWidthIn, TileHeightIn, TilePaddingIn);
+	}
+
+	void SetTiling(FImageDimensions DimensionsIn, int32 TileWidthIn, int32 TileHeightIn, int32 TilePaddingIn)
+	{
+		checkSlow(TileWidthIn >= 1 && TileHeightIn >= 1 && TilePaddingIn >= 0);
+		Dimensions = DimensionsIn;
+		TileWidth = FMath::Clamp(TileWidthIn, 1, Dimensions.GetWidth());
+		TileHeight = FMath::Clamp(TileHeightIn, 1, Dimensions.GetHeight());
+		TilePadding = TilePaddingIn < 0 ? 0 : TilePaddingIn;
+	}
+
+	int32 NumTilesX() const
+	{
+		return Dimensions.GetWidth() / TileWidth + (Dimensions.GetWidth() % TileWidth == 0 ? 0 : 1);
+	}
+
+	int32 NumTilesY() const
+	{
+		return Dimensions.GetHeight() / TileHeight + (Dimensions.GetHeight() % TileHeight == 0 ? 0 : 1);
+	}
+
+	int32 Num() const
+	{
+		return NumTilesX() * NumTilesY();
+	}
+
+	/** @return the tile corresponding to the linear index [0, NumTiles()-1]. */
+	FImageDimensions GetTile(int32 LinearTileIndex) const
+	{
+		FVector2i TileIndex(LinearTileIndex % NumTilesX(), LinearTileIndex / NumTilesX());
+		FVector2i TileOffset(TileIndex.X * TileWidth, TileIndex.Y * TileHeight);
+		int32 ClampedWidth = FMath::Clamp(TileWidth, 0, Dimensions.GetWidth() - TileIndex.X * TileWidth);
+		int32 ClampedHeight = FMath::Clamp(TileHeight, 0, Dimensions.GetHeight() - TileIndex.Y * TileHeight);
+		// TODO: Padding support
+		return FImageDimensions(ClampedWidth, ClampedHeight, Dimensions.GetWidth(), Dimensions.GetHeight(), TileOffset);
+	}
+};
+
+} // end namespace UE::Geometry
+} // end namespace UE
+
+
diff --git a/Engine/Source/Runtime/GeometryCore/Public/Image/ImageOccupancyMap.h b/Engine/Source/Runtime/GeometryCore/Public/Image/ImageOccupancyMap.h
new file mode 100644
--- /dev/null
+++ b/Engine/Source/Runtime/GeometryCore/Public/Image/ImageOccupancyMap.h
@@ -0,0 +1,287 @@
+// Copyright Epic Games, Inc. All Rights Reserved.
+
+#pragma once
+
+#include "Templates/Tuple.h"
+#include "VectorTypes.h"
+#include "Image/ImageDimensions.h"
+#include "Spatial/MeshAABBTree3.h"
+#include "Sampling/GridSampler.h"
+
+namespace UE
+{
+namespace Geometry
+{
+
+/**
+ * ImageOccupancyMap calculates and stores coverage information for a 2D image/texture,
+ * for example coverage derived from UV islands of a mesh, 2D polygons, etc.
+ * 
+ * An optional set of gutter texels can be calculated, and correspondence between gutter
+ * texels and the nearest interior texel is stored.
+ *
+ * In addition, a 2D coordinate (eg UV) and integer ID (eg Triangle ID) of each texel can be calculated/stored. 
+ * This is not just a cache. For 'border' texels where the texel center is technically outside the mesh/polygon,
+ * but the texel rectangle may still overlap the shape, the nearest UV/Triangle is stored.
+ * This simplifies computing samples around the borders such that the shape is covered under linear interpolatione/etc.
+ */
+class FImageOccupancyMap
+{
+public:
+	/** Image Dimensions */
+	FImageDimensions Dimensions;
+
+	/** Width of the gutter. This is actually multiplied by the diagonal length of a texel, so
+	    the gutter is generally larger than this number of pixels */
+	int32 GutterSize = 4;
+
+	// texel types
+	const int8 EmptyTexel = 0;
+	const int8 InteriorTexel = 1;
+	//const int8 BorderTexel = 2;
+	const int8 GutterTexel = 3;
+
+	/** Texel Type for each texel in image, Size = Width x Height x Multisamples */
+	TArray64<int8> TexelType;
+	/** Count of samples in the texel that are interior. */
+	TArray64<int32> TexelInteriorSamples;
+	/** UV for each texel in image. Only set for Interior texels */
+	TArray64<FVector2f> TexelQueryUV;
+	/** integer/Triangle ID for each texel in image. Only set for Interior texels. */
+	TArray64<int32> TexelQueryTriangle;
+
+	/** Set of Gutter Texels. Pair is <LinearIndexOfGutterTexel, LinearIndexOfNearestInteriorTexel>, so
+	    Gutter can be filled by directly copying from source to target. */
+	TArray64<TTuple<int64, int64>> GutterTexels;
+
+	using FGridSampler = TGridSampler<double>;
+	FGridSampler Multisampler = FGridSampler(1);
+
+	bool bParallel = true;
+
+
+	void Initialize(FImageDimensions DimensionsIn, int32 Multisampling = 1)
+	{
+		check(DimensionsIn.IsSquare()); // are we sure it works otherwise?
+		Dimensions = DimensionsIn;
+		Multisampler = FGridSampler(Multisampling);
+	}
+
+	/**
+	 * @return true if texel sample at this sample linear index is an Interior texel
+	 */
+	bool IsInterior(int64 LinearIndex) const
+	{
+		return TexelType[LinearIndex] == InteriorTexel;
+	}
+
+	/** @return number of interior samples for the texel at this image linear index. */
+	int32 TexelNumSamples(int64 LinearIndex) const
+	{
+		return TexelInteriorSamples[LinearIndex];
+	}
+
+	/** Set texel type at the given X/Y */
+	//void SetTexelType(int64 X, int64 Y, int8 Type)
+	//{
+	//	TexelType[Dimensions.GetIndex(X, Y)] = Type;
+	//}
+
+
+	template<typename MeshType, typename GetTriangleIDFuncType>
+	bool ComputeFromUVSpaceMesh(const MeshType& UVSpaceMesh, 
+		GetTriangleIDFuncType GetTriangleIDFunc = [](int32 TriangleID) { return TriangleID; } )
+	{
+		// make flat mesh
+		TMeshAABBTree3<MeshType> FlatSpatial(&UVSpaceMesh, true);
+
+		const int32 LinearImageSize = Dimensions.Num();
+		const int32 LinearSampleSize = Dimensions.Num() * Multisampler.Num();
+		TexelType.Init(EmptyTexel, LinearSampleSize);
+		TexelInteriorSamples.Init(0, LinearImageSize);
+		TexelQueryUV.Init(FVector2f::Zero(), LinearSampleSize);
+		TexelQueryTriangle.Init(IndexConstants::InvalidID, LinearSampleSize);
+
+		const FVector2d TexelSize = Dimensions.GetTexelSize();
+		const double TexelDiag = TexelSize.Length();
+
+		IMeshSpatial::FQueryOptions QueryOptions;
+		QueryOptions.MaxDistance = (double)GutterSize * TexelDiag;
+
+		// find interior texels
+		TAtomic<int64> InteriorCounter;
+		FCriticalSection GutterLock;
+		ParallelFor(Dimensions.GetHeight(), 
+			[this, &UVSpaceMesh, &GetTriangleIDFunc, &FlatSpatial, TexelDiag, &QueryOptions, &InteriorCounter, &GutterLock, &TexelSize]
+		(int32 ImgY)
+		{
+			// To minimize locks, accumulate gutter texel data locally and append
+			// to our gutter texel array at the end.
+			TArray64<TTuple<int64, int64>> LocalGutterTexels;
+			int64 LocalInteriorCounter = 0;
+			for (int32 ImgX = 0; ImgX < Dimensions.GetWidth(); ++ImgX)
+			{
+				// With multiple samples, a texel may contain multiple gutter samples.
+				// Since we copy nearest interior texel over top our gutter texels, we
+				// should only consider a texel a gutter texel iff all samples in a
+				// texel are gutter texels. This bool will track that as we iterate
+				// over each sample of the texel.
+				bool bIsGutterTexel = true;
+				TTuple<int64, int64> GutterNearestTexel;
+				for (int32 Sample = 0; Sample < Multisampler.Num(); ++Sample)
+				{
+					const int64 TexelLinearIdx = Dimensions.GetIndex(ImgX, ImgY);
+					const int64 SampleLinearIdx = TexelLinearIdx * Multisampler.Num() + Sample;
+					const FVector2d SampleUV = Multisampler.Sample(Sample);
+					FVector2d UVPoint = Dimensions.GetTexelUV(TexelLinearIdx);
+					UVPoint = UVPoint - 0.5 * TexelSize + SampleUV * TexelSize;
+					const FVector3d UVPoint3d(UVPoint.X, UVPoint.Y, 0);
+
+					double NearDistSqr;
+					const int32 NearestTriID = FlatSpatial.FindNearestTriangle(UVPoint3d, NearDistSqr, QueryOptions);
+					if (NearestTriID >= 0)
+					{
+						FVector3d A, B, C;
+						UVSpaceMesh.GetTriVertices(NearestTriID, A, B, C);
+						const FTriangle2d UVTriangle(GetXY(A), GetXY(B), GetXY(C));
+
+						if (UVTriangle.IsInsideOrOn(UVPoint))
+						{
+							TexelType[SampleLinearIdx] = InteriorTexel;
+							TexelQueryUV[SampleLinearIdx] = (FVector2f)UVPoint;
+							TexelQueryTriangle[SampleLinearIdx] = GetTriangleIDFunc(NearestTriID);
+							++TexelInteriorSamples[TexelLinearIdx];
+							++LocalInteriorCounter;
+							bIsGutterTexel = false;
+						}
+						else if (NearDistSqr < TexelDiag * TexelDiag)
+						{
+							const FDistPoint3Triangle3d DistQuery = TMeshQueries<MeshType>::TriangleDistance(UVSpaceMesh, NearestTriID, UVPoint3d);
+							FVector2d NearestUV = GetXY(DistQuery.ClosestTrianglePoint);
+							// nudge point into triangle to improve numerical behavior of things like barycentric coord calculation
+							NearestUV += (10.0 * FMathf::ZeroTolerance) * Normalized(NearestUV - UVPoint);
+
+							TexelType[SampleLinearIdx] = InteriorTexel;
+							TexelQueryUV[SampleLinearIdx] = (FVector2f)NearestUV;
+							TexelQueryTriangle[SampleLinearIdx] = GetTriangleIDFunc(NearestTriID);
+							++TexelInteriorSamples[TexelLinearIdx];
+							++LocalInteriorCounter;
+							bIsGutterTexel = false;
+						}
+						else
+						{
+							TexelType[SampleLinearIdx] = GutterTexel;
+							const FDistPoint3Triangle3d DistQuery = TMeshQueries<MeshType>::TriangleDistance(UVSpaceMesh, NearestTriID, UVPoint3d);
+							const FVector2d NearestUV = GetXY(DistQuery.ClosestTrianglePoint);
+							const FVector2i NearestCoords = Dimensions.UVToCoords(NearestUV);
+
+							// To avoid artifacts with mipmapping the gutter texels store the nearest
+							// interior texel and a post pass copies those nearest interior texel values
+							// to each gutter pixel.
+							//
+							// The gutter pixel structure stores pixel indices to copy in the source
+							// image's linear index.
+							//
+							// For tiled images:
+							//    NearestCoords are in source space --> GetSource().GetIndex() will return source linear index.
+							//    ImgX, ImgY are in local tile space --> GetSourceIndex() will return source linear index.
+							const int64 NearestLinearIdx = Dimensions.IsTile() ? Dimensions.GetSourceDimensions().GetIndex(NearestCoords) : Dimensions.GetIndex(NearestCoords);
+							const int32 LinearSourceIdx = Dimensions.IsTile() ? Dimensions.GetSourceIndex(FVector2i(ImgX, ImgY)) : TexelLinearIdx;
+
+							GutterNearestTexel = TTuple<int64, int64>(LinearSourceIdx, NearestLinearIdx);
+						}
+					}
+					if (bIsGutterTexel)
+					{
+						LocalGutterTexels.Add(GutterNearestTexel);
+					}
+				}
+			}
+			if (LocalGutterTexels.Num() > 0)
+			{
+				GutterLock.Lock();
+				GutterTexels += LocalGutterTexels;
+				GutterLock.Unlock();
+			}
+			InteriorCounter += LocalInteriorCounter;
+		}, !bParallel ? EParallelForFlags::ForceSingleThread : EParallelForFlags::None);
+
+		return true;
+	}
+
+
+
+
+	template<typename TexelValueType>
+	void ParallelProcessingPass(
+		TFunctionRef<TexelValueType(int64 LinearIdx)> BeginTexel,
+		TFunctionRef<void(int64 LinearIdx, float Weight, TexelValueType&)> AccumulateTexel,
+		TFunctionRef<void(int64 LinearIdx, float Weight, TexelValueType&)> CompleteTexel,
+		TFunctionRef<void(int64 LinearIdx, TexelValueType&)> WriteTexel,
+		TFunctionRef<float(const FVector2i& TexelOffset)> WeightFunction,
+		int32 FilterWidth,
+		TArray<TexelValueType>& PassBuffer
+	) const
+	{
+		int64 N = Dimensions.Num();
+		PassBuffer.SetNum(N);
+		
+		ParallelFor(Dimensions.GetHeight(), 
+			[this, &BeginTexel, & AccumulateTexel, &CompleteTexel, &WeightFunction, FilterWidth, &PassBuffer]
+		(int32 ImgY)
+		{
+			for (int32 ImgX = 0; ImgX < Dimensions.GetWidth(); ++ImgX)
+			{
+				int64 LinearIdx = Dimensions.GetIndex(ImgX, ImgY);
+				if (TexelType[LinearIdx] != EmptyTexel)
+				{
+					TexelValueType AccumValue = BeginTexel(LinearIdx);
+					float WeightSum = 0;
+
+					FVector2i Coords(ImgX, ImgY);
+					FVector2i MaxNbr = Coords + FVector2i(FilterWidth, FilterWidth);
+					Dimensions.Clamp(MaxNbr);
+					FVector2i MinNbr = Coords - FVector2i(FilterWidth, FilterWidth);
+					Dimensions.Clamp(MinNbr);
+
+					for (int32 Y = MinNbr.Y; Y <= MaxNbr.Y; ++Y)
+					{
+						for (int32 X = MinNbr.X; X <= MaxNbr.X; ++X)
+						{
+							FVector2i NbrCoords(X, Y);
+							if (Dimensions.IsValidCoords(NbrCoords))
+							{
+								FVector2i Offset = NbrCoords - Coords;
+								int64 LinearNbrIndex = Dimensions.GetIndex(NbrCoords);
+								if (TexelType[LinearNbrIndex] != EmptyTexel)
+								{
+									float NbrWeight = WeightFunction(Offset);
+									AccumulateTexel(LinearNbrIndex, NbrWeight, AccumValue);
+									WeightSum += NbrWeight;
+								}
+							}
+						}
+					}
+
+					CompleteTexel(LinearIdx, WeightSum, AccumValue);
+					PassBuffer[LinearIdx] = AccumValue;
+				}
+			}
+		});
+
+		// write results
+		for (int64 k = 0; k < N; ++k)
+		{
+			if (TexelType[k] != EmptyTexel)
+			{
+				WriteTexel(k, PassBuffer[k]);
+			}
+		}
+	}
+
+
+};
+
+
+} // end namespace UE::Geometry
+} // end namespace UE
\ No newline at end of file
diff --git a/Engine/Source/Runtime/GeometryCore/Public/Implicit/GridInterpolant.h b/Engine/Source/Runtime/GeometryCore/Public/Implicit/GridInterpolant.h
new file mode 100644
--- /dev/null
+++ b/Engine/Source/Runtime/GeometryCore/Public/Implicit/GridInterpolant.h
@@ -0,0 +1,196 @@
+// Copyright Epic Games, Inc. All Rights Reserved.
+
+// Port of geometry3Sharp's CachingMeshSDFImplicit
+
+#pragma once
+
+#include "MathUtil.h"
+#include "VectorTypes.h"
+
+namespace UE
+{
+namespace Geometry
+{
+
+
+/**
+ * Tri-linear interpolant for a 3D dense Grid. Supports Grid translation
+ * via GridOrigin, but does not support scaling or rotation. If you need those,
+ * you can wrap this in something that does the xform.
+ *
+ * GridType must have a GetValue() that returns a value to interpolate at a given FVector3i location -- (w/ locations ranging from [0,0,0] to Dimensions (exclusive))
+ */
+template <class GridType>
+class TTriLinearGridInterpolant /*: public BoundedImplicitFunction3d (TODO: consider add ImplicitFunction3d interface concept once more implicit functions are available*/
+{
+public:
+	GridType* Grid;
+	FVector3d GridOrigin;
+	double CellSize;
+	FVector3i Dimensions;
+
+	// value to return if query point is outside Grid (in an SDF
+	// outside is usually positive). Need to do math with this value,
+	// and cast this value to/from float; use FMathf::MaxReal to avoid overflow
+	double Outside = FMathf::MaxReal;
+
+	TTriLinearGridInterpolant(GridType* Grid, FVector3d GridOrigin, double CellSize, FVector3i Dimensions) : Grid(Grid), GridOrigin(GridOrigin), CellSize(CellSize), Dimensions(Dimensions)
+	{
+	}
+
+	FAxisAlignedBox3d Bounds() const
+	{
+		return FAxisAlignedBox3d(
+			{ GridOrigin.X, GridOrigin.Y, GridOrigin.Z },
+			{ GridOrigin.X + CellSize * Dimensions.X,
+			  GridOrigin.Y + CellSize * Dimensions.Y,
+			  GridOrigin.Z + CellSize * Dimensions.Z });
+	}
+
+
+	double Value(const FVector3d& Pt) const
+	{
+		FVector3d gridPt(
+			((Pt.X - GridOrigin.X) / CellSize),
+			((Pt.Y - GridOrigin.Y) / CellSize),
+			((Pt.Z - GridOrigin.Z) / CellSize));
+
+		// compute integer coordinates
+		int X0 = (int)gridPt.X;
+		int Y0 = (int)gridPt.Y, Y1 = Y0 + 1;
+		int Z0 = (int)gridPt.Z, Z1 = Z0 + 1;
+
+		// clamp to Grid
+		if (X0 < 0 || (X0 + 1) >= Dimensions.X ||
+			Y0 < 0 || Y1 >= Dimensions.Y ||
+			Z0 < 0 || Z1 >= Dimensions.Z)
+		{
+			return Outside;
+		}
+
+		// convert double coords to [0,1] range
+		double fAx = gridPt.X - (double)X0;
+		double fAy = gridPt.Y - (double)Y0;
+		double fAz = gridPt.Z - (double)Z0;
+		double OneMinusfAx = 1.0 - fAx;
+
+		// compute trilinear interpolant. The code below tries to do this with the fewest 
+		// number of variables, in hopes that optimizer will be clever about re-using registers, etc.
+		// Commented code at bottom is fully-expanded version.
+		// [TODO] it is possible to implement lerps here as A+(B-A)*t, saving a multiply and a variable.
+		//   This is numerically worse, but since the Grid values are floats and
+		//   we are computing in doubles, does it matter?
+		double xa, xb;
+
+		get_value_pair(X0, Y0, Z0, xa, xb);
+		double yz = (1 - fAy) * (1 - fAz);
+		double sum = (OneMinusfAx * xa + fAx * xb) * yz;
+
+		get_value_pair(X0, Y0, Z1, xa, xb);
+		yz = (1 - fAy) * (fAz);
+		sum += (OneMinusfAx * xa + fAx * xb) * yz;
+
+		get_value_pair(X0, Y1, Z0, xa, xb);
+		yz = (fAy) * (1 - fAz);
+		sum += (OneMinusfAx * xa + fAx * xb) * yz;
+
+		get_value_pair(X0, Y1, Z1, xa, xb);
+		yz = (fAy) * (fAz);
+		sum += (OneMinusfAx * xa + fAx * xb) * yz;
+
+		return sum;
+
+		// fV### is Grid cell corner index
+		//return
+		//    fV000 * (1 - fAx) * (1 - fAy) * (1 - fAz) +
+		//    fV001 * (1 - fAx) * (1 - fAy) * (fAz) +
+		//    fV010 * (1 - fAx) * (fAy) * (1 - fAz) +
+		//    fV011 * (1 - fAx) * (fAy) * (fAz) +
+		//    fV100 * (fAx) * (1 - fAy) * (1 - fAz) +
+		//    fV101 * (fAx) * (1 - fAy) * (fAz) +
+		//    fV110 * (fAx) * (fAy) * (1 - fAz) +
+		//    fV111 * (fAx) * (fAy) * (fAz);
+	}
+
+
+protected:
+	void get_value_pair(int I, int J, int K, double& A, double& B) const
+	{
+		A = (double)Grid->GetValue(FVector3i(I, J, K));
+		B = (double)Grid->GetValue(FVector3i(I + 1, J, K));
+	}
+
+
+public:
+	FVector3d Gradient(const FVector3d& Pt) const
+	{
+		FVector3d gridPt = FVector3d(
+			((Pt.X - GridOrigin.X) / CellSize),
+			((Pt.Y - GridOrigin.Y) / CellSize),
+			((Pt.Z - GridOrigin.Z) / CellSize));
+
+		// clamp to Grid
+		if (gridPt.X < 0 || gridPt.X >= Dimensions.X - 1 ||
+			gridPt.Y < 0 || gridPt.Y >= Dimensions.Y - 1 ||
+			gridPt.Z < 0 || gridPt.Z >= Dimensions.Z - 1)
+		{
+			return FVector3d::Zero();
+		}
+
+		// compute integer coordinates
+		int X0 = (int)gridPt.X;
+		int Y0 = (int)gridPt.Y, Y1 = Y0 + 1;
+		int Z0 = (int)gridPt.Z, Z1 = Z0 + 1;
+
+		// convert double coords to [0,1] range
+		double fAx = gridPt.X - (double)X0;
+		double fAy = gridPt.Y - (double)Y0;
+		double fAz = gridPt.Z - (double)Z0;
+
+		double fV000, fV100;
+		get_value_pair(X0, Y0, Z0, fV000, fV100);
+		double fV010, fV110;
+		get_value_pair(X0, Y1, Z0, fV010, fV110);
+		double fV001, fV101;
+		get_value_pair(X0, Y0, Z1, fV001, fV101);
+		double fV011, fV111;
+		get_value_pair(X0, Y1, Z1, fV011, fV111);
+
+		// [TODO] can re-order this to vastly reduce number of ops!
+		double gradX =
+			-fV000 * (1 - fAy) * (1 - fAz) +
+			-fV001 * (1 - fAy) * (fAz)+
+			-fV010 * (fAy) * (1 - fAz) +
+			-fV011 * (fAy) * (fAz)+
+			fV100 * (1 - fAy) * (1 - fAz) +
+			fV101 * (1 - fAy) * (fAz)+
+			fV110 * (fAy) * (1 - fAz) +
+			fV111 * (fAy) * (fAz);
+
+		double gradY =
+			-fV000 * (1 - fAx) * (1 - fAz) +
+			-fV001 * (1 - fAx) * (fAz)+
+			fV010 * (1 - fAx) * (1 - fAz) +
+			fV011 * (1 - fAx) * (fAz)+
+			-fV100 * (fAx) * (1 - fAz) +
+			-fV101 * (fAx) * (fAz)+
+			fV110 * (fAx) * (1 - fAz) +
+			fV111 * (fAx) * (fAz);
+
+		double gradZ =
+			-fV000 * (1 - fAx) * (1 - fAy) +
+			fV001 * (1 - fAx) * (1 - fAy) +
+			-fV010 * (1 - fAx) * (fAy)+
+			fV011 * (1 - fAx) * (fAy)+
+			-fV100 * (fAx) * (1 - fAy) +
+			fV101 * (fAx) * (1 - fAy) +
+			-fV110 * (fAx) * (fAy)+
+			fV111 * (fAx) * (fAy);
+
+		return FVector3d(gradX, gradY, gradZ);
+	}
+};
+
+
+} // end namespace UE::Geometry
+} // end namespace UE
\ No newline at end of file
diff --git a/Engine/Source/Runtime/GeometryCore/Public/Intersection/IntersectionQueries3.h b/Engine/Source/Runtime/GeometryCore/Public/Intersection/IntersectionQueries3.h
new file mode 100644
--- /dev/null
+++ b/Engine/Source/Runtime/GeometryCore/Public/Intersection/IntersectionQueries3.h
@@ -0,0 +1,34 @@
+// Copyright Epic Games, Inc. All Rights Reserved.
+
+#pragma once
+
+#include "CoreMinimal.h"
+#include "VectorTypes.h"
+#include "SphereTypes.h"
+#include "CapsuleTypes.h"
+#include "HalfspaceTypes.h"
+#include "OrientedBoxTypes.h"
+
+namespace UE
+{
+	namespace Geometry
+	{
+		//
+		// Halfspace Intersection Queries
+		//
+
+		/** @return true if Halfspace and Sphere intersect */
+		template<typename RealType>
+		bool TestIntersection(const THalfspace3<RealType>& Halfspace, const TSphere3<RealType>& Sphere);
+
+		/** @return true if Halfspace and Capsule intersect */
+		template<typename RealType>
+		bool TestIntersection(const THalfspace3<RealType>& Halfspace, const TCapsule3<RealType>& Capsule);
+
+		/** @return true if Halfspace and Box intersect */
+		template<typename RealType>
+		bool TestIntersection(const THalfspace3<RealType>& Halfspace, const TOrientedBox3<RealType>& Box);
+
+
+	}
+}
\ No newline at end of file
diff --git a/Engine/Source/Runtime/GeometryCore/Public/Intersection/IntrRay3Triangle3.h b/Engine/Source/Runtime/GeometryCore/Public/Intersection/IntrRay3Triangle3.h
new file mode 100644
--- /dev/null
+++ b/Engine/Source/Runtime/GeometryCore/Public/Intersection/IntrRay3Triangle3.h
@@ -0,0 +1,180 @@
+// Copyright Epic Games, Inc. All Rights Reserved.
+
+// Port of WildMagic IntrRay3Triangle3
+
+#pragma once
+
+#include "VectorTypes.h"
+#include "RayTypes.h"
+#include "TriangleTypes.h"
+#include "VectorUtil.h"
+
+namespace UE
+{
+namespace Geometry
+{
+
+/**
+ * Compute intersection between 3D ray and 3D triangle
+ */
+template <typename Real>
+class TIntrRay3Triangle3
+{
+public:
+	// Input
+	TRay3<Real> Ray;
+	TTriangle3<Real> Triangle;
+
+	// Output
+	Real RayParameter;
+	FVector3d TriangleBaryCoords;
+	EIntersectionType IntersectionType;
+
+
+	TIntrRay3Triangle3(const TRay3<Real>& RayIn, const TTriangle3<Real>& TriangleIn)
+	{
+		Ray = RayIn;
+		Triangle = TriangleIn;
+	}
+
+
+	/**
+	 * @return true if ray intersects triangle
+	 */
+	bool Test()
+	{
+		// Compute the offset origin, edges, and normal.
+		FVector3<Real> diff = Ray.Origin - Triangle.V[0];
+		FVector3<Real> edge1 = Triangle.V[1] - Triangle.V[0];
+		FVector3<Real> edge2 = Triangle.V[2] - Triangle.V[0];
+		FVector3<Real> normal = edge1.Cross(edge2);
+
+		// Solve Q + t*D = b1*E1 + b2*E2 (Q = kDiff, D = ray direction,
+		// E1 = kEdge1, E2 = kEdge2, N = Cross(E1,E2)) by
+		//   |Dot(D,N)|*b1 = sign(Dot(D,N))*Dot(D,Cross(Q,E2))
+		//   |Dot(D,N)|*b2 = sign(Dot(D,N))*Dot(D,Cross(E1,Q))
+		//   |Dot(D,N)|*t = -sign(Dot(D,N))*Dot(Q,N)
+		Real DdN = Ray.Direction.Dot(normal);
+		Real sign;
+		if (DdN > TMathUtil<Real>::ZeroTolerance)
+		{
+			sign = (Real)1;
+		}
+		else if (DdN < -TMathUtil<Real>::ZeroTolerance)
+		{
+			sign = (Real)-1;
+			DdN = -DdN;
+		}
+		else
+		{
+			// Ray and triangle are parallel, call it a "no intersection"
+			// even if the ray does intersect.
+			IntersectionType = EIntersectionType::Empty;
+			return false;
+		}
+
+		Real DdQxE2 = sign * Ray.Direction.Dot(diff.Cross(edge2));
+		if (DdQxE2 >= (Real)0)
+		{
+			Real DdE1xQ = sign * Ray.Direction.Dot(edge1.Cross(diff));
+			if (DdE1xQ >= (Real)0)
+			{
+				if (DdQxE2 + DdE1xQ <= DdN)
+				{
+					// Line intersects triangle, check if ray does.
+					Real QdN = -sign * diff.Dot(normal);
+					if (QdN >= (Real)0)
+					{
+						// Ray intersects triangle.
+						IntersectionType = EIntersectionType::Point;
+						return true;
+					}
+					// else: t < 0, no intersection
+				}
+				// else: b1+b2 > 1, no intersection
+			}
+			// else: b2 < 0, no intersection
+		}
+		// else: b1 < 0, no intersection
+
+		IntersectionType = EIntersectionType::Empty;
+		return false;
+	}
+
+
+	/**
+	 * Find intersection point
+	 * @return true if ray intersects triangle
+	 */
+	bool Find()
+	{
+		// Compute the offset origin, edges, and normal.
+		FVector3<Real> diff = Ray.Origin - Triangle.V[0];
+		FVector3<Real> edge1 = Triangle.V[1] - Triangle.V[0];
+		FVector3<Real> edge2 = Triangle.V[2] - Triangle.V[0];
+		FVector3<Real> normal = edge1.Cross(edge2);
+
+		// Solve Q + t*D = b1*E1 + b2*E2 (Q = kDiff, D = ray direction,
+		// E1 = kEdge1, E2 = kEdge2, N = Cross(E1,E2)) by
+		//   |Dot(D,N)|*b1 = sign(Dot(D,N))*Dot(D,Cross(Q,E2))
+		//   |Dot(D,N)|*b2 = sign(Dot(D,N))*Dot(D,Cross(E1,Q))
+		//   |Dot(D,N)|*t = -sign(Dot(D,N))*Dot(Q,N)
+		Real DdN = Ray.Direction.Dot(normal);
+		Real sign;
+		if (DdN > TMathUtil<Real>::ZeroTolerance)
+		{
+			sign = (Real)1;
+		}
+		else if (DdN < -TMathUtil<Real>::ZeroTolerance)
+		{
+			sign = (Real)-1;
+			DdN = -DdN;
+		}
+		else
+		{
+			// Ray and triangle are parallel, call it a "no intersection"
+			// even if the ray does intersect.
+			IntersectionType = EIntersectionType::Empty;
+			return false;
+		}
+
+		Real DdQxE2 = sign * Ray.Direction.Dot(diff.Cross(edge2));
+		if (DdQxE2 >= (Real)0)
+		{
+			Real DdE1xQ = sign * Ray.Direction.Dot(edge1.Cross(diff));
+			if (DdE1xQ >= (Real)0)
+			{
+				if (DdQxE2 + DdE1xQ <= DdN)
+				{
+					// Line intersects triangle, check if ray does.
+					Real QdN = -sign * diff.Dot(normal);
+					if (QdN >= (Real)0)
+					{
+						// Ray intersects triangle.
+						Real inv = ((Real)1) / DdN;
+						RayParameter = QdN * inv;
+						TriangleBaryCoords.Y = DdQxE2 * inv;
+						TriangleBaryCoords.Z = DdE1xQ * inv;
+						TriangleBaryCoords.X = (Real)1 - TriangleBaryCoords.Y - TriangleBaryCoords.Z;
+						IntersectionType = EIntersectionType::Point;
+						return true;
+					}
+					// else: t < 0, no intersection
+				}
+				// else: b1+b2 > 1, no intersection
+			}
+			// else: b2 < 0, no intersection
+		}
+		// else: b1 < 0, no intersection
+
+		IntersectionType = EIntersectionType::Empty;
+		return false;
+	}
+
+};
+
+typedef TIntrRay3Triangle3<float> FIntrRay3Triangle3f;
+typedef TIntrRay3Triangle3<double> FIntrRay3Triangle3d;
+
+} // end namespace UE::Geometry
+} // end namespace UE
diff --git a/Engine/Source/Runtime/GeometryCore/Public/MathUtil.h b/Engine/Source/Runtime/GeometryCore/Public/MathUtil.h
new file mode 100644
--- /dev/null
+++ b/Engine/Source/Runtime/GeometryCore/Public/MathUtil.h
@@ -0,0 +1,386 @@
+// Copyright Epic Games, Inc. All Rights Reserved.
+
+#pragma once
+
+#include "GeometryBase.h"
+#include <cmath>
+#include <cfloat>
+
+
+/**
+ * Math constants and utility functions, templated on float/double type
+ */
+template<typename RealType>
+struct TMathUtilConstants;
+
+template<>
+struct TMathUtilConstants<float>
+{
+	/** Machine Epsilon - float approx 1e-7, double approx 2e-16 */
+	GEOMETRYCORE_API static constexpr float Epsilon = FLT_EPSILON;
+	/** Zero tolerance for math operations (eg like parallel tests) - float 1e-6, double 1e-8 */
+	GEOMETRYCORE_API static constexpr float ZeroTolerance = 1e-06f;
+
+	/** largest possible number for type */
+	GEOMETRYCORE_API static constexpr float MaxReal = FLT_MAX;
+
+	/** 3.14159... */
+	GEOMETRYCORE_API static constexpr float Pi = 3.1415926535897932384626433832795f;
+	GEOMETRYCORE_API static constexpr float FourPi = 4.0f * Pi;
+	GEOMETRYCORE_API static constexpr float TwoPi = 2.0f*Pi;
+	GEOMETRYCORE_API static constexpr float HalfPi = 0.5f*Pi;
+
+	/** 1.0 / Pi */
+	GEOMETRYCORE_API static constexpr float InvPi = 1.0f / Pi;
+	/** 1.0 / (2*Pi) */
+	GEOMETRYCORE_API static constexpr float InvTwoPi = 1.0f / TwoPi;
+
+	/** pi / 180 */
+	GEOMETRYCORE_API static constexpr float DegToRad = Pi / 180.0f;
+	/** 180 / pi */
+	GEOMETRYCORE_API static constexpr float RadToDeg = 180.0f / Pi;
+
+	//static constexpr float LN_2;
+	//static constexpr float LN_10;
+	//static constexpr float INV_LN_2;
+	//static constexpr float INV_LN_10;
+
+	GEOMETRYCORE_API static constexpr float Sqrt2 = 1.4142135623730950488016887242097f;
+	GEOMETRYCORE_API static constexpr float InvSqrt2 = 1.0f / Sqrt2;
+	GEOMETRYCORE_API static constexpr float Sqrt3 = 1.7320508075688772935274463415059f;
+	GEOMETRYCORE_API static constexpr float InvSqrt3 = 1.0f / Sqrt3;
+};
+
+template<>
+struct TMathUtilConstants<double>
+{
+	/** Machine Epsilon - float approx 1e-7, double approx 2e-16 */
+	GEOMETRYCORE_API static constexpr double Epsilon = DBL_EPSILON;
+	/** Zero tolerance for math operations (eg like parallel tests) - float 1e-6, double 1e-8 */
+	GEOMETRYCORE_API static constexpr double ZeroTolerance = 1e-08;
+
+	/** largest possible number for type */
+	GEOMETRYCORE_API static constexpr double MaxReal = DBL_MAX;
+
+	/** 3.14159... */
+	GEOMETRYCORE_API static constexpr double Pi = 3.1415926535897932384626433832795;
+	GEOMETRYCORE_API static constexpr double FourPi = 4.0 * Pi;
+	GEOMETRYCORE_API static constexpr double TwoPi = 2.0 * Pi;
+	GEOMETRYCORE_API static constexpr double HalfPi = 0.5 * Pi;
+
+	/** 1.0 / Pi */
+	GEOMETRYCORE_API static constexpr double InvPi = 1.0 / Pi;
+	/** 1.0 / (2*Pi) */
+	GEOMETRYCORE_API static constexpr double InvTwoPi = 1.0 / TwoPi;
+
+	/** pi / 180 */
+	GEOMETRYCORE_API static constexpr double DegToRad = Pi / 180.0;
+	/** 180 / pi */
+	GEOMETRYCORE_API static constexpr double RadToDeg = 180.0 / Pi;
+
+	//static constexpr double LN_2;
+	//static constexpr double LN_10;
+	//static constexpr double INV_LN_2;
+	//static constexpr double INV_LN_10;
+
+	GEOMETRYCORE_API static constexpr double Sqrt2 = 1.4142135623730950488016887242097;
+	GEOMETRYCORE_API static constexpr double InvSqrt2 = 1.0 / Sqrt2;
+	GEOMETRYCORE_API static constexpr double Sqrt3 = 1.7320508075688772935274463415059;
+	GEOMETRYCORE_API static constexpr double InvSqrt3 = 1.0 / Sqrt3;
+};
+
+
+// we use TMathUtil<int> so we need to define these nonsense constants
+template<>
+struct TMathUtilConstants<int32>
+{
+	GEOMETRYCORE_API static constexpr int32 Epsilon = 0;
+	GEOMETRYCORE_API static constexpr int32 ZeroTolerance = 0;
+	GEOMETRYCORE_API static constexpr int32 MaxReal = ((int32)0x7fffffff);
+	GEOMETRYCORE_API static constexpr int32 Pi = 3;
+	GEOMETRYCORE_API static constexpr int32 FourPi = 4 * Pi;
+	GEOMETRYCORE_API static constexpr int32 TwoPi = 2 * Pi;
+	GEOMETRYCORE_API static constexpr int32 HalfPi = 1;
+	GEOMETRYCORE_API static constexpr int32 InvPi = 1;
+	GEOMETRYCORE_API static constexpr int32 InvTwoPi = 1;
+	GEOMETRYCORE_API static constexpr int32 DegToRad = 1;
+	GEOMETRYCORE_API static constexpr int32 RadToDeg = 1;
+	GEOMETRYCORE_API static constexpr int32 Sqrt2 = 1;
+	GEOMETRYCORE_API static constexpr int32 InvSqrt2 = 1;
+	GEOMETRYCORE_API static constexpr int32 Sqrt3 = 2;
+	GEOMETRYCORE_API static constexpr int32 InvSqrt3 = 1;
+};
+
+
+// we use TMathUtil<int> so we need to define these nonsense constants
+template<>
+struct TMathUtilConstants<int64>
+{
+	GEOMETRYCORE_API static constexpr int64 Epsilon = 0;
+	GEOMETRYCORE_API static constexpr int64 ZeroTolerance = 0;
+	GEOMETRYCORE_API static constexpr int64 MaxReal = ((int64)0x7fffffffffffffff);
+	GEOMETRYCORE_API static constexpr int64 Pi = 3;
+	GEOMETRYCORE_API static constexpr int64 FourPi = 4 * Pi;
+	GEOMETRYCORE_API static constexpr int64 TwoPi = 2 * Pi;
+	GEOMETRYCORE_API static constexpr int64 HalfPi = 1;
+	GEOMETRYCORE_API static constexpr int64 InvPi = 1;
+	GEOMETRYCORE_API static constexpr int64 InvTwoPi = 1;
+	GEOMETRYCORE_API static constexpr int64 DegToRad = 1;
+	GEOMETRYCORE_API static constexpr int64 RadToDeg = 1;
+	GEOMETRYCORE_API static constexpr int64 Sqrt2 = 1;
+	GEOMETRYCORE_API static constexpr int64 InvSqrt2 = 1;
+	GEOMETRYCORE_API static constexpr int64 Sqrt3 = 2;
+	GEOMETRYCORE_API static constexpr int64 InvSqrt3 = 1;
+};
+
+
+
+template<typename RealType>
+class TMathUtil : public TMathUtilConstants<RealType>
+{
+public:
+	static inline bool IsNaN(const RealType Value);
+	static inline bool IsFinite(const RealType Value);
+	static inline RealType Abs(const RealType Value);
+	static inline RealType Clamp(const RealType Value, const RealType ClampMin, const RealType ClampMax);
+	static inline RealType Sign(const RealType Value);
+	static inline RealType SignNonZero(const RealType Value);
+	static inline RealType Max(const RealType A, const RealType B);
+	static inline RealType Max3(const RealType A, const RealType B, const RealType C);
+	static inline int32 Max3Index(const RealType A, const RealType B, const RealType C);
+	static inline RealType Min(const RealType A, const RealType B);
+	static inline RealType Min3(const RealType A, const RealType B, const RealType C);
+	static inline int32 Min3Index(const RealType A, const RealType B, const RealType C);
+	/** compute min and max of a,b,c with max 3 comparisons (sometimes 2) */
+	static inline void MinMax(RealType A, RealType B, RealType C, RealType& MinOut, RealType& MaxOut);
+	static inline RealType Sqrt(const RealType Value);
+	static inline RealType Tan(const RealType Value);
+	static inline RealType Atan2(const RealType ValueY, const RealType ValueX);
+	static inline RealType Sin(const RealType Value);
+	static inline RealType Cos(const RealType Value);
+	static inline RealType ACos(const RealType Value);
+	static inline RealType Floor(const RealType Value);
+	static inline RealType Ceil(const RealType Value);
+	static inline RealType Round(const RealType Value);
+	static inline RealType Pow(const RealType Value, const RealType Power);
+	static inline RealType Exp(const RealType Power);
+	static inline RealType Log(const RealType Value);
+	static inline RealType Lerp(const RealType A, const RealType B, RealType Alpha);
+
+
+	/**
+	 * @return result of Atan2 shifted to [0,2pi]  (normal ATan2 returns in range [-pi,pi])
+	 */
+	static inline RealType Atan2Positive(const RealType Y, const RealType X);
+
+private:
+	TMathUtil() = delete;
+};
+typedef TMathUtil<float> FMathf;
+typedef TMathUtil<double> FMathd;
+
+
+template<typename RealType>
+bool TMathUtil<RealType>::IsNaN(const RealType Value)
+{
+	return std::isnan(Value);
+}
+
+
+template<typename RealType>
+bool TMathUtil<RealType>::IsFinite(const RealType Value)
+{
+	return std::isfinite(Value);
+}
+
+
+template<typename RealType>
+RealType TMathUtil<RealType>::Abs(const RealType Value)
+{
+	return (Value >= (RealType)0) ? Value : -Value;
+}
+
+
+template<typename RealType>
+RealType TMathUtil<RealType>::Clamp(const RealType Value, const RealType ClampMin, const RealType ClampMax)
+{
+	return (Value < ClampMin) ? ClampMin : ((Value > ClampMax) ? ClampMax : Value);
+}
+
+
+template<typename RealType>
+RealType TMathUtil<RealType>::Sign(const RealType Value)
+{
+	return (Value > (RealType)0) ? (RealType)1 : ((Value < (RealType)0) ? (RealType)-1 : (RealType)0);
+}
+
+template<typename RealType>
+RealType TMathUtil<RealType>::SignNonZero(const RealType Value)
+{
+	return (Value < (RealType)0) ? (RealType)-1 : (RealType)1;
+}
+
+template<typename RealType>
+RealType TMathUtil<RealType>::Max(const RealType A, const RealType B)
+{
+	return (A >= B) ? A : B;
+}
+
+template<typename RealType>
+RealType TMathUtil<RealType>::Max3(const RealType A, const RealType B, const RealType C)
+{
+	return Max(Max(A, B), C);
+}
+
+template<typename RealType>
+int32 TMathUtil<RealType>::Max3Index(const RealType A, const RealType B, const RealType C)
+{
+	if (A >= B) 
+	{
+		return (A >= C) ? 0 : 2;
+	}
+	else
+	{
+		return (B >= C) ? 1 : 2;
+	}
+}
+
+template<typename RealType>
+RealType TMathUtil<RealType>::Min(const RealType A, const RealType B)
+{
+	return (A <= B) ? A : B;
+}
+
+template<typename RealType>
+RealType TMathUtil<RealType>::Min3(const RealType A, const RealType B, const RealType C)
+{
+	return Min(Min(A, B), C);
+}
+
+template<typename RealType>
+int32 TMathUtil<RealType>::Min3Index(const RealType A, const RealType B, const RealType C)
+{
+	if (A <= B) 
+	{
+		return (A <= C) ? 0 : 2;
+	}
+	else 
+	{
+		return (B <= C) ? 1 : 2;
+	}
+}
+
+// compute min and max of a,b,c with max 3 comparisons (sometimes 2)
+template<typename RealType>
+void TMathUtil<RealType>::MinMax(RealType A, RealType B, RealType C, RealType& MinOut, RealType& MaxOut)
+{
+	if (A < B) {
+		if (A < C) {
+			MinOut = A; MaxOut = TMathUtil<RealType>::Max(B, C);
+		}
+		else {
+			MinOut = C; MaxOut = B;
+		}
+	}
+	else {
+		if (A > C) {
+			MaxOut = A; MinOut = TMathUtil<RealType>::Min(B, C);
+		}
+		else {
+			MinOut = B; MaxOut = C;
+		}
+	}
+}
+
+template<typename RealType>
+RealType TMathUtil<RealType>::Sqrt(const RealType Value)
+{
+	return sqrt(Value);
+}
+
+template<typename RealType>
+RealType TMathUtil<RealType>::Tan(const RealType Value)
+{
+	return tan(Value);
+}
+
+template<typename RealType>
+RealType TMathUtil<RealType>::Atan2(const RealType ValueY, const RealType ValueX)
+{
+	return atan2(ValueY, ValueX);
+}
+
+template<typename RealType>
+RealType TMathUtil<RealType>::Sin(const RealType Value)
+{
+	return sin(Value);
+}
+
+template<typename RealType>
+RealType TMathUtil<RealType>::Cos(const RealType Value)
+{
+	return cos(Value);
+}
+
+template<typename RealType>
+RealType TMathUtil<RealType>::ACos(const RealType Value)
+{
+	return acos(Value);
+}
+
+template<typename RealType>
+RealType TMathUtil<RealType>::Floor(const RealType Value)
+{
+	return floor(Value);
+}
+
+template<typename RealType>
+RealType TMathUtil<RealType>::Ceil(const RealType Value)
+{
+	return ceil(Value);
+}
+
+template<typename RealType>
+RealType TMathUtil<RealType>::Round(const RealType Value)
+{
+	return round(Value);
+}
+
+template<typename RealType>
+RealType TMathUtil<RealType>::Pow(const RealType Value, const RealType Power)
+{
+	return pow(Value, Power);
+}
+
+template<typename RealType>
+RealType TMathUtil<RealType>::Exp(const RealType Power)
+{
+	return exp(Power);
+}
+
+template<typename RealType>
+RealType TMathUtil<RealType>::Log(const RealType Power)
+{
+	return log(Power);
+}
+
+
+
+template<typename RealType>
+RealType TMathUtil<RealType>::Lerp(const RealType A, const RealType B, RealType Alpha)
+{
+	Alpha = Clamp(Alpha, (RealType)0, (RealType)1);
+	return ((RealType)1 - Alpha)*A + (Alpha)*B;
+}
+
+template<typename RealType>
+RealType TMathUtil<RealType>::Atan2Positive(const RealType Y, const RealType X)
+{
+	// @todo this is a float atan2 !!
+	RealType Theta = TMathUtil<RealType>::Atan2(Y, X);
+	if (Theta < 0)
+	{
+		return ((RealType)2 * TMathUtil<RealType>::Pi) + Theta;
+	}
+	return Theta;
+}
diff --git a/Engine/Source/Runtime/GeometryCore/Public/MeshAdapter.h b/Engine/Source/Runtime/GeometryCore/Public/MeshAdapter.h
new file mode 100644
--- /dev/null
+++ b/Engine/Source/Runtime/GeometryCore/Public/MeshAdapter.h
@@ -0,0 +1,231 @@
+// Copyright Epic Games, Inc. All Rights Reserved.
+
+#pragma once
+
+#include "TriangleTypes.h"
+#include "VectorTypes.h"
+
+#include "Math/IntVector.h"
+#include "Math/Vector.h"
+#include "IndexTypes.h"
+
+namespace UE
+{
+namespace Geometry
+{
+
+/**
+ * Most generic / lazy example of a triangle mesh adapter; possibly useful for prototyping / building on top of (but slower than making a more specific-case adapter)
+ */
+template <typename RealType>
+struct TTriangleMeshAdapter
+{
+	TFunction<bool(int32 index)> IsTriangle;
+	TFunction<bool(int32 index)> IsVertex;
+	TFunction<int32()> MaxTriangleID;
+	TFunction<int32()> MaxVertexID;
+	TFunction<int32()> TriangleCount;
+	TFunction<int32()> VertexCount;
+	TFunction<int32()> GetShapeTimestamp;
+	TFunction<FIndex3i(int32)> GetTriangle;
+	TFunction<FVector3<RealType>(int32)> GetVertex;
+
+	inline void GetTriVertices(int TID, UE::Math::TVector<RealType>& V0, UE::Math::TVector<RealType>& V1, UE::Math::TVector<RealType>& V2) const
+	{
+		FIndex3i TriIndices = GetTriangle(TID);
+		V0 = GetVertex(TriIndices.A);
+		V1 = GetVertex(TriIndices.B);
+		V2 = GetVertex(TriIndices.C);
+	}
+};
+
+typedef TTriangleMeshAdapter<double> FTriangleMeshAdapterd;
+typedef TTriangleMeshAdapter<float> FTriangleMeshAdapterf;
+
+/**
+ * Example function to generate a generic mesh adapter from arrays
+ * @param Vertices Array of mesh vertices
+ * @param Triangles Array of int-vectors, one per triangle, indexing into the vertices array
+ */
+inline FTriangleMeshAdapterd GetArrayMesh(TArray<FVector>& Vertices, TArray<FIntVector>& Triangles)
+{
+	return {
+		[&](int) { return true; },
+		[&](int) { return true; },
+		[&]() { return Triangles.Num(); },
+		[&]() { return Vertices.Num(); },
+		[&]() { return Triangles.Num(); },
+		[&]() { return Vertices.Num(); },
+		[&]() { return 0; },
+		[&](int Idx) { return FIndex3i(Triangles[Idx]); },
+		[&](int Idx) { return FVector3d(Vertices[Idx]); }};
+}
+
+
+/**
+ * Faster adapter specifically for the common index mesh case
+ */
+template<typename IndexType, typename OutRealType, typename InVectorType=FVector>
+struct TIndexMeshArrayAdapter
+{
+	const TArray<InVectorType>* SourceVertices;
+	const TArray<IndexType>* SourceTriangles;
+
+	void SetSources(const TArray<InVectorType>* SourceVerticesIn, const TArray<IndexType>* SourceTrianglesIn)
+	{
+		SourceVertices = SourceVerticesIn;
+		SourceTriangles = SourceTrianglesIn;
+	}
+
+	TIndexMeshArrayAdapter() : SourceVertices(nullptr), SourceTriangles(nullptr)
+	{
+	}
+
+	TIndexMeshArrayAdapter(const TArray<InVectorType>* SourceVerticesIn, const TArray<IndexType>* SourceTrianglesIn) : SourceVertices(SourceVerticesIn), SourceTriangles(SourceTrianglesIn)
+	{
+	}
+
+	FORCEINLINE bool IsTriangle(int32 Index) const
+	{
+		return SourceTriangles->IsValidIndex(Index * 3);
+	}
+
+	FORCEINLINE bool IsVertex(int32 Index) const
+	{
+		return SourceVertices->IsValidIndex(Index);
+	}
+
+	FORCEINLINE int32 MaxTriangleID() const
+	{
+		return SourceTriangles->Num() / 3;
+	}
+
+	FORCEINLINE int32 MaxVertexID() const
+	{
+		return SourceVertices->Num();
+	}
+
+	// Counts are same as MaxIDs, because these are compact meshes
+	FORCEINLINE int32 TriangleCount() const
+	{
+		return SourceTriangles->Num() / 3;
+	}
+
+	FORCEINLINE int32 VertexCount() const
+	{
+		return SourceVertices->Num();
+	}
+
+	FORCEINLINE int32 GetShapeTimestamp() const
+	{
+		return 0; // source data doesn't have a timestamp concept
+	}
+
+	FORCEINLINE FIndex3i GetTriangle(int32 Index) const
+	{
+		int32 Start = Index * 3;
+		return FIndex3i((int)(*SourceTriangles)[Start], (int)(*SourceTriangles)[Start+1], (int)(*SourceTriangles)[Start+2]);
+	}
+
+	FORCEINLINE FVector3<OutRealType> GetVertex(int32 Index) const
+	{
+		return FVector3<OutRealType>((*SourceVertices)[Index]);
+	}
+
+	FORCEINLINE void GetTriVertices(int32 TriIndex, UE::Math::TVector<OutRealType>& V0, UE::Math::TVector<OutRealType>& V1, UE::Math::TVector<OutRealType>& V2) const
+	{
+		int32 Start = TriIndex * 3;
+		V0 = FVector3<OutRealType>((*SourceVertices)[(*SourceTriangles)[Start]]);
+		V1 = FVector3<OutRealType>((*SourceVertices)[(*SourceTriangles)[Start+1]]);
+		V2 = FVector3<OutRealType>((*SourceVertices)[(*SourceTriangles)[Start+2]]);
+	}
+
+};
+
+
+/**
+ * Second version of the above faster adapter
+ *  -- for the case where triangle indices are packed into an integer vector type instead of flat
+ */
+template<typename IndexVectorType, typename OutRealType, typename InVectorType = FVector>
+struct TIndexVectorMeshArrayAdapter
+{
+	const TArray<InVectorType>* SourceVertices;
+	const TArray<IndexVectorType>* SourceTriangles;
+
+	void SetSources(const TArray<InVectorType>* SourceVerticesIn, const TArray<IndexVectorType>* SourceTrianglesIn)
+	{
+		SourceVertices = SourceVerticesIn;
+		SourceTriangles = SourceTrianglesIn;
+	}
+
+	TIndexVectorMeshArrayAdapter() : SourceVertices(nullptr), SourceTriangles(nullptr)
+	{
+	}
+
+	TIndexVectorMeshArrayAdapter(const TArray<InVectorType>* SourceVerticesIn, const TArray<IndexVectorType>* SourceTrianglesIn) : SourceVertices(SourceVerticesIn), SourceTriangles(SourceTrianglesIn)
+	{
+	}
+
+	FORCEINLINE bool IsTriangle(int32 Index) const
+	{
+		return SourceTriangles->IsValidIndex(Index);
+	}
+
+	FORCEINLINE bool IsVertex(int32 Index) const
+	{
+		return SourceVertices->IsValidIndex(Index);
+	}
+
+	FORCEINLINE int32 MaxTriangleID() const
+	{
+		return SourceTriangles->Num();
+	}
+
+	FORCEINLINE int32 MaxVertexID() const
+	{
+		return SourceVertices->Num();
+	}
+
+	// Counts are same as MaxIDs, because these are compact meshes
+	FORCEINLINE int32 TriangleCount() const
+	{
+		return SourceTriangles->Num();
+	}
+
+	FORCEINLINE int32 VertexCount() const
+	{
+		return SourceVertices->Num();
+	}
+
+	FORCEINLINE int32 GetShapeTimestamp() const
+	{
+		return 0; // source data doesn't have a timestamp concept
+	}
+
+	FORCEINLINE FIndex3i GetTriangle(int32 Index) const
+	{
+		const IndexVectorType& Tri = (*SourceTriangles)[Index];
+		return FIndex3i((int)Tri[0], (int)Tri[1], (int)Tri[2]);
+	}
+
+	FORCEINLINE FVector3<OutRealType> GetVertex(int32 Index) const
+	{
+		return FVector3<OutRealType>((*SourceVertices)[Index]);
+	}
+
+	FORCEINLINE void GetTriVertices(int32 TriIndex, UE::Math::TVector<OutRealType>& V0, UE::Math::TVector<OutRealType>& V1, UE::Math::TVector<OutRealType>& V2) const
+	{
+		const IndexVectorType& Tri = (*SourceTriangles)[TriIndex];
+		V0 = FVector3<OutRealType>((*SourceVertices)[Tri[0]]);
+		V1 = FVector3<OutRealType>((*SourceVertices)[Tri[1]]);
+		V2 = FVector3<OutRealType>((*SourceVertices)[Tri[2]]);
+	}
+
+};
+
+typedef TIndexMeshArrayAdapter<uint32, double> FIndexMeshArrayAdapterd;
+
+
+} // end namespace UE::Geometry
+} // end namespace UE
\ No newline at end of file
diff --git a/Engine/Source/Runtime/GeometryCore/Public/OrientedBoxTypes.h b/Engine/Source/Runtime/GeometryCore/Public/OrientedBoxTypes.h
new file mode 100644
--- /dev/null
+++ b/Engine/Source/Runtime/GeometryCore/Public/OrientedBoxTypes.h
@@ -0,0 +1,282 @@
+// Copyright Epic Games, Inc. All Rights Reserved.
+
+// ported from geometry3Sharp Box3
+
+#pragma once
+
+#include "VectorTypes.h"
+#include "BoxTypes.h"
+#include "FrameTypes.h"
+
+namespace UE
+{
+namespace Geometry
+{
+
+
+/**
+ * TOrientedBox3 is a non-axis-aligned 3D box defined by a 3D frame and extents along the axes of that frame
+ * The frame is at the center of the box.
+ */
+template<typename RealType>
+struct TOrientedBox3
+{
+	// available for porting: ContainPoint, MergeBoxes()
+
+
+	/** 3D position (center) and orientation (axes) of the box */
+	TFrame3<RealType> Frame;
+	/** Half-dimensions of box measured along the three axes */
+	FVector3<RealType> Extents;
+
+	TOrientedBox3() : Extents(1,1,1) {}
+
+	/**
+	 * Create axis-aligned box with given Origin and Extents
+	 */
+	TOrientedBox3(const FVector3<RealType>& Origin, const FVector3<RealType> & ExtentsIn)
+		: Frame(Origin), Extents(ExtentsIn)
+	{
+	}
+
+	/**
+	 * Create oriented box with given Frame and Extents
+	 */
+	TOrientedBox3(const TFrame3<RealType>& FrameIn, const FVector3<RealType> & ExtentsIn)
+		: Frame(FrameIn), Extents(ExtentsIn)
+	{
+	}
+
+	/**
+	 * Create oriented box from axis-aligned box
+	 */
+	TOrientedBox3(const TAxisAlignedBox3<RealType>& AxisBox)
+		: Frame(AxisBox.Center()), Extents((RealType)0.5 * AxisBox.Diagonal())
+	{
+	}
+
+
+	/** @return box with unit dimensions centered at origin */
+	static TOrientedBox3<RealType> UnitZeroCentered() { return TOrientedBox3<RealType>(FVector3<RealType>::Zero(), (RealType)0.5*FVector3<RealType>::One()); }
+
+	/** @return box with unit dimensions where minimum corner is at origin */
+	static TOrientedBox3<RealType> UnitPositive() { return TOrientedBox3<RealType>((RealType)0.5*FVector3<RealType>::One(), (RealType)0.5*FVector3<RealType>::One()); }
+
+
+	/** @return center of the box */
+	FVector3<RealType> Center() const { return Frame.Origin; }
+
+	/** @return X axis of the box */
+	FVector3<RealType> AxisX() const { return Frame.X(); }
+
+	/** @return Y axis of the box */
+	FVector3<RealType> AxisY() const { return Frame.Y(); }
+
+	/** @return Z axis of the box */
+	FVector3<RealType> AxisZ() const { return Frame.Z(); }
+
+	/** @return an axis of the box */
+	FVector3<RealType> GetAxis(int AxisIndex) const { return Frame.GetAxis(AxisIndex); }
+
+	/** @return maximum extent of box */
+	inline RealType MaxExtent() const
+	{
+		return Extents.MaxAbs();
+	}
+
+	/** @return minimum extent of box */
+	inline RealType MinExtent() const
+	{
+		return Extents.MinAbs();
+	}
+
+	/** @return vector from minimum-corner to maximum-corner of box */
+	inline FVector3<RealType> Diagonal() const
+	{
+		return Frame.PointAt(Extents.X, Extents.Y, Extents.Z) - Frame.PointAt(-Extents.X, -Extents.Y, -Extents.Z);
+	}
+
+	/** @return volume of box */
+	inline RealType Volume() const
+	{
+		return (RealType)8 * (Extents.X) * (Extents.Y) * (Extents.Z);
+	}
+
+	/** @return true if box contains point */
+	inline bool Contains(const FVector3<RealType>& Point) const
+	{
+		FVector3<RealType> InFramePoint = Frame.ToFramePoint(Point);
+		return (TMathUtil<RealType>::Abs(InFramePoint.X) <= Extents.X) &&
+			(TMathUtil<RealType>::Abs(InFramePoint.Y) <= Extents.Y) &&
+			(TMathUtil<RealType>::Abs(InFramePoint.Z) <= Extents.Z);
+	}
+
+
+	// corners [ (-x,-y), (x,-y), (x,y), (-x,y) ], -z, then +z
+	//
+	//   7---6     +z       or        3---2     -z
+	//   |\  |\                       |\  |\
+    //   4-\-5 \                      0-\-1 \
+    //    \ 3---2                      \ 7---6   
+	//     \|   |                       \|   |
+	//      0---1  -z                    4---5  +z
+	//
+	// @todo does this ordering make sense for UE? we are in LHS instead of RHS here
+	// if this is modified, likely need to update IndexUtil::BoxFaces and BoxFaceNormals
+
+	/**
+	 * @param Index corner index in range 0-7
+	 * @return Corner point on the box identified by the given index. See diagram in OrientedBoxTypes.h for index/corner mapping.
+	 */
+	FVector3<RealType> GetCorner(int Index) const
+	{
+		check(Index >= 0 && Index <= 7);
+		RealType dx = (((Index & 1) != 0) ^ ((Index & 2) != 0)) ? (Extents.X) : (-Extents.X);
+		RealType dy = ((Index / 2) % 2 == 0) ? (-Extents.Y) : (Extents.Y);
+		RealType dz = (Index < 4) ? (-Extents.Z) : (Extents.Z);
+		return Frame.PointAt(dx, dy, dz);
+	}
+
+	/**
+	 * Call CornerPointFunc(FVector3) for each of the 8 box corners. Order is the same as GetCorner(X).
+	 * This is more efficient than calling GetCorner(X) because the Rotation matrix is only computed once.
+	 */
+	template<typename PointFuncType>
+	void EnumerateCorners(PointFuncType CornerPointFunc) const
+	{
+		TMatrix3<RealType> RotMatrix = Frame.Rotation.ToRotationMatrix();
+		RealType X = Extents.X, Y = Extents.Y, Z = Extents.Z;
+		CornerPointFunc( RotMatrix*FVector3<RealType>(-X,-Y,-Z) + Frame.Origin );
+		CornerPointFunc( RotMatrix*FVector3<RealType>( X,-Y,-Z) + Frame.Origin );
+		CornerPointFunc( RotMatrix*FVector3<RealType>( X, Y,-Z) + Frame.Origin );
+		CornerPointFunc( RotMatrix*FVector3<RealType>(-X, Y,-Z) + Frame.Origin );
+		CornerPointFunc( RotMatrix*FVector3<RealType>(-X,-Y, Z) + Frame.Origin );
+		CornerPointFunc( RotMatrix*FVector3<RealType>( X,-Y, Z) + Frame.Origin );
+		CornerPointFunc( RotMatrix*FVector3<RealType>( X, Y, Z) + Frame.Origin );
+		CornerPointFunc( RotMatrix*FVector3<RealType>(-X, Y, Z) + Frame.Origin );
+	}
+
+
+	/**
+	 * Call CornerPointPredicate(FVector3) for each of the 8 box corners, with early-out if any call returns false
+	 * @return true if all tests pass
+	 */
+	template<typename PointPredicateType>
+	bool TestCorners(PointPredicateType CornerPointPredicate) const
+	{
+		TMatrix3<RealType> RotMatrix = Frame.Rotation.ToRotationMatrix();
+		RealType X = Extents.X, Y = Extents.Y, Z = Extents.Z;
+		return CornerPointPredicate(RotMatrix * FVector3<RealType>(-X, -Y, -Z) + Frame.Origin) &&
+			CornerPointPredicate(RotMatrix * FVector3<RealType>(X, -Y, -Z) + Frame.Origin) &&
+			CornerPointPredicate(RotMatrix * FVector3<RealType>(X, Y, -Z) + Frame.Origin) &&
+			CornerPointPredicate(RotMatrix * FVector3<RealType>(-X, Y, -Z) + Frame.Origin) &&
+			CornerPointPredicate(RotMatrix * FVector3<RealType>(-X, -Y, Z) + Frame.Origin) &&
+			CornerPointPredicate(RotMatrix * FVector3<RealType>(X, -Y, Z) + Frame.Origin) &&
+			CornerPointPredicate(RotMatrix * FVector3<RealType>(X, Y, Z) + Frame.Origin) &&
+			CornerPointPredicate(RotMatrix * FVector3<RealType>(-X, Y, Z) + Frame.Origin);
+	}
+
+
+
+	/**
+	 * Get whether the corner at Index (see diagram in GetCorner documentation comment) is in the negative or positive direction for each axis
+	 * @param Index corner index in range 0-7
+	 * @return Index3i with 0 or 1 for each axis, 0 if corner is in the negative direction for that axis, 1 if in the positive direction
+	 */
+	static FIndex3i GetCornerSide(int Index)
+	{
+		check(Index >= 0 && Index <= 7);
+		return FIndex3i(
+			(((Index & 1) != 0) ^ ((Index & 2) != 0)) ? 1 : 0,
+			((Index / 2) % 2 == 0) ? 0 : 1,
+			(Index < 4) ? 0 : 1
+		);
+	}
+
+
+
+	/**
+	 * Find squared distance to box.
+	 * @param Point input point
+	 * @return squared distance from point to box, or 0 if point is inside box
+	 */
+	RealType DistanceSquared(FVector3<RealType> Point)
+	{
+		 // Ported from WildMagic5 Wm5DistPoint3Box3.cpp
+			 
+		 // Work in the box's coordinate system.
+		 Point -= Frame.Origin;
+
+		// Compute squared distance and closest point on box.
+		RealType sqrDistance = 0;
+		RealType delta;
+		FVector3<RealType> closest(0, 0, 0);
+		for (int i = 0; i < 3; ++i) 
+		{
+			closest[i] = Point.Dot(GetAxis(i));
+			if (closest[i] < -Extents[i]) 
+			{
+				delta = closest[i] + Extents[i];
+				sqrDistance += delta * delta;
+				closest[i] = -Extents[i];
+			}
+			else if (closest[i] > Extents[i])
+			{
+				delta = closest[i] - Extents[i];
+				sqrDistance += delta * delta;
+				closest[i] = Extents[i];
+			}
+		}
+
+		return sqrDistance;
+	}
+
+
+
+
+	 /**
+	  * Find closest point on box
+	  * @param Point input point
+	  * @return closest point on box. Input point is returned if it is inside box.
+	  */
+	 FVector3<RealType> ClosestPoint(FVector3<RealType> Point)
+	 {
+		 // Work in the box's coordinate system.
+		 Point -= Frame.Origin;
+
+		 // Compute squared distance and closest point on box.
+		 RealType sqrDistance = 0;
+		 RealType delta;
+		 FVector3<RealType> closest;
+		 FVector3<RealType> Axes[3];
+		 for (int i = 0; i < 3; ++i) 
+		 {
+			 Axes[i] = GetAxis(i);
+			 closest[i] = Point.Dot(Axes[i]);
+			 RealType extent = Extents[i];
+			 if (closest[i] < -extent) 
+			 {
+				 delta = closest[i] + extent;
+				 sqrDistance += delta * delta;
+				 closest[i] = -extent;
+			 }
+			 else if (closest[i] > extent) 
+			 {
+				 delta = closest[i] - extent;
+				 sqrDistance += delta * delta;
+				 closest[i] = extent;
+			 }
+		 }
+
+		 return Frame.Origin + closest.X*Axes[0] + closest.Y*Axes[1] + closest.Z*Axes[2];
+	 }
+
+
+
+};
+
+typedef TOrientedBox3<float> FOrientedBox3f;
+typedef TOrientedBox3<double> FOrientedBox3d;
+
+} // end namespace UE::Geometry
+} // end namespace UE
