GeometryCore: create new GeometryCore Module. Move GeometricObjects code from GeometryProcessing plugin to this new module, as well as core FDynamicMesh3 classes (attributes, aabbtree/octree, change t

diff --git a/Engine/Source/Runtime/GeometryCore/Public/Polygon2.h b/Engine/Source/Runtime/GeometryCore/Public/Polygon2.h
new file mode 100644
--- /dev/null
+++ b/Engine/Source/Runtime/GeometryCore/Public/Polygon2.h
@@ -0,0 +1,1148 @@
+// Copyright Epic Games, Inc. All Rights Reserved.
+
+// port of geometry3Sharp Polygon
+
+#pragma once
+
+#include "Templates/UnrealTemplate.h"
+#include "Math/UnrealMath.h"
+#include "VectorTypes.h"
+#include "BoxTypes.h"
+#include "SegmentTypes.h"
+#include "LineTypes.h"
+#include "MathUtil.h"
+#include "Intersection/IntrSegment2Segment2.h"
+#include "Util/DynamicVector.h"
+
+namespace UE
+{
+namespace Geometry
+{
+
+/**
+ * TPolygon2 is a 2D polygon represented as a list of Vertices.
+ * 
+ * @todo move operators
+ */
+template<typename T>
+class TPolygon2
+{
+protected:
+	/** The list of vertices/corners of the polygon */
+	TArray<FVector2<T>> Vertices;
+
+	/** A counter that is incremented every time the polygon vertices are modified */
+	int Timestamp;
+
+public:
+
+	TPolygon2() : Timestamp(0)
+	{
+	}
+
+	/**
+	 * Construct polygon that is a copy of another polygon
+	 */
+	TPolygon2(const TPolygon2& Copy) : Vertices(Copy.Vertices), Timestamp(Copy.Timestamp)
+	{
+	}
+
+	/**
+	 * Construct polygon with given list of vertices
+	 */
+	TPolygon2(const TArray<FVector2<T>>& VertexList) : Vertices(VertexList), Timestamp(0)
+	{
+	}
+
+	/**
+	 * Construct polygon with given indices into a vertex array
+	 */
+	TPolygon2(TArrayView<const FVector2<T>> VertexArray, TArrayView<const int32> VertexIndices) : Timestamp(0)
+	{
+		Vertices.SetNum(VertexIndices.Num());
+		for (int32 Idx = 0; Idx < VertexIndices.Num(); Idx++)
+		{
+			Vertices[Idx] = VertexArray[VertexIndices[Idx]];
+		}
+	}
+
+	/** @return the Timestamp for the polygon, which is updated every time the polygon is modified */
+	int GetTimestamp() const 
+	{
+		return Timestamp;
+	}
+
+	/**
+	 * Get the vertex at a given index
+	 */
+	const FVector2<T>& operator[](int Index) const
+	{
+		return Vertices[Index];
+	}
+
+	/**
+	 * Get the vertex at a given index
+	 * @warning changing the vertex via this operator does not update Timestamp!
+	 */
+	FVector2<T>& operator[](int Index)
+	{
+		return Vertices[Index];
+	}
+
+
+	/**
+	 * @return first vertex of Polygon
+	 */
+	const FVector2<T>& Start() const
+	{
+		return Vertices[0];
+	}
+
+	/**
+	 * @return list of Vertices of Polygon
+	 */
+	const TArray<FVector2<T>>& GetVertices() const
+	{
+		return Vertices;
+	}
+
+	/**
+	 * @return number of Vertices in Polygon
+	 */
+	int VertexCount() const
+	{
+		return Vertices.Num();
+	}
+
+	/**
+	 * Add a vertex to the Polygon
+	 */
+	void AppendVertex(const FVector2<T>& Position)
+	{
+		Vertices.Add(Position);
+		Timestamp++;
+	}
+
+	/**
+	 * Add a list of Vertices to the Polygon
+	 */
+	void AppendVertices(const TArray<FVector2<T>>& NewVertices)
+	{
+		Vertices.Append(NewVertices);
+		Timestamp++;
+	}
+
+	/**
+	 * Set vertex at given index to a new Position
+	 */
+	void Set(int VertexIndex, const FVector2<T>& Position)
+	{
+		Vertices[VertexIndex] = Position;
+		Timestamp++;
+	}
+
+	/**
+	 * Remove a vertex of the Polygon (existing Vertices are shifted)
+	 */
+	void RemoveVertex(int VertexIndex)
+	{
+		Vertices.RemoveAt(VertexIndex);
+		Timestamp++;
+	}
+
+	/**
+	 * Replace the list of Vertices with a new list
+	 */
+	void SetVertices(const TArray<FVector2<T>>& NewVertices)
+	{
+		Vertices = NewVertices;
+		Timestamp++;
+	}
+
+
+	/**
+	 * Reverse the order of the Vertices in the Polygon (ie switch between Clockwise and CounterClockwise)
+	 */
+	void Reverse()
+	{
+		int32 j = Vertices.Num()-1;
+		for (int32 VertexIndex = 0; VertexIndex < j; VertexIndex++, j--)
+		{
+			Swap(Vertices[VertexIndex], Vertices[j]);
+		}
+		Timestamp++;
+	}
+
+
+	/**
+	 * Get the tangent vector at a vertex of the polygon, which is the normalized
+	 * vector from the previous vertex to the next vertex
+	 */
+	FVector2<T> GetTangent(int VertexIndex) const
+	{
+		FVector2<T> next = Vertices[(VertexIndex + 1) % Vertices.Num()];
+		FVector2<T> prev = Vertices[VertexIndex == 0 ? Vertices.Num() - 1 : VertexIndex - 1];
+		return Normalized(next - prev);
+	}
+
+
+	/**
+	 * Get the normal vector at a vertex of the polygon, which is perpendicular to GetTangent()
+	 * Points "inward" for a Clockwise Polygon, and outward for CounterClockwise
+	 */
+	FVector2<T> GetNormal(int VertexIndex) const
+	{
+		return PerpCW(GetTangent(VertexIndex));
+	}
+
+
+	/**
+	 * Construct a normal at a vertex of the Polygon by averaging the adjacent face normals.
+	 * This vector is independent of the lengths of the adjacent segments.
+	 * Points "inward" for a Clockwise Polygon, and outward for CounterClockwise
+	 */
+	FVector2<T> GetNormal_FaceAvg(int VertexIndex) const
+	{
+		FVector2<T> next = Vertices[(VertexIndex + 1) % Vertices.Num()];
+		FVector2<T> prev = Vertices[VertexIndex == 0 ? Vertices.Num() - 1 : VertexIndex - 1];
+		next -= Vertices[VertexIndex]; Normalize(next);
+		prev -= Vertices[VertexIndex]; Normalize(prev);
+
+		FVector2<T> n = (PerpCW(next) - PerpCW(prev));
+		T len = Normalize(n);
+		if (len == 0) 
+		{
+			return Normalized(next + prev);   // this gives right direction for degenerate angle
+		}
+		else 
+		{
+			return n;
+		}
+	}
+
+
+	/**
+	 * @return the bounding box of the Polygon Vertices
+	 */
+	TAxisAlignedBox2<T> Bounds() const
+	{
+		TAxisAlignedBox2<T> box = TAxisAlignedBox2<T>::Empty();
+		box.Contain(Vertices);
+		return box;
+	}
+
+	
+	/**
+	 * SegmentIterator is used to iterate over the TSegment2<T> segments of the polygon
+	 */
+	class SegmentIterator 
+	{
+	public:
+		inline bool operator!()
+		{
+			return i < polygon->VertexCount();
+		}
+		inline TSegment2<T> operator*() const
+		{
+			check(polygon != nullptr && i < polygon->VertexCount());
+			return TSegment2<T>(polygon->Vertices[i], polygon->Vertices[(i+1) % polygon->VertexCount()]);
+		}
+		//inline TSegment2<T> & operator*();
+		inline SegmentIterator & operator++() 		// prefix
+		{
+			i++;
+			return *this;
+		}
+		inline SegmentIterator operator++(int) 		// postfix
+		{
+			SegmentIterator copy(*this);
+			i++;
+			return copy;
+		}
+		inline bool operator==(const SegmentIterator & i2) { return i2.polygon == polygon && i2.i == i; }
+		inline bool operator!=(const SegmentIterator & i2) { return i2.polygon != polygon || i2.i != i; }
+	protected:
+		const TPolygon2 * polygon;
+		int i;
+		inline SegmentIterator(const TPolygon2 * p, int iCur) : polygon(p), i(iCur) {}
+		friend class TPolygon2;
+	};
+	friend class SegmentIterator;
+
+	SegmentIterator SegmentItr() const 
+	{
+		return SegmentIterator(this, 0);
+	}
+
+	/**
+	 * Wrapper around SegmentIterator that has begin() and end() suitable for range-based for loop
+	 */
+	class SegmentEnumerable
+	{
+	public:
+		const TPolygon2<T>* polygon;
+		SegmentEnumerable() : polygon(nullptr) {}
+		SegmentEnumerable(const TPolygon2<T> * p) : polygon(p) {}
+		SegmentIterator begin() { return polygon->SegmentItr(); }
+		SegmentIterator end() { return SegmentIterator(polygon, polygon->VertexCount()); }
+	};
+
+	/**
+	 * @return an object that can be used in a range-based for loop to iterate over the Segments of the Polygon
+	 */
+	SegmentEnumerable Segments() const
+	{
+		return SegmentEnumerable(this);
+	}
+
+
+	/**
+	 * @return true if the Polygon Vertices have Clockwise winding order / orientation  (signed area is negative)
+	 */
+	bool IsClockwise() const
+	{
+		return SignedArea() < 0;
+	}
+
+	/**
+	 * @return the signed area of the Polygon
+	 */
+	T SignedArea() const
+	{
+		T fArea = 0;
+		int N = Vertices.Num();
+		if (N == 0)
+		{
+			return 0;
+		}
+		for (int i = 0; i < N; ++i) 
+		{
+			const FVector2<T>& v1 = Vertices[i];
+			const FVector2<T>& v2 = Vertices[(i + 1) % N];
+			fArea += v1.X * v2.Y - v1.Y * v2.X;
+		}
+		return fArea * 0.5;
+	}
+
+	/**
+	 * @return the unsigned area of the Polygon
+	 */
+	T Area() const
+	{
+		return TMathUtil<T>::Abs(SignedArea());
+	}
+
+	/**
+	 * @return the total perimeter length of the Polygon
+	 */
+	T Perimeter() const
+	{
+		T fPerim = 0;
+		int N = Vertices.Num();
+		for (int i = 0; i < N; ++i)
+		{
+			fPerim += Vertices[i].Distance(Vertices[(i + 1) % N]);
+		}
+		return fPerim;
+	}
+
+
+	/**
+	 * Get the previous and next vertex positions for a given vertex of the Polygon
+	 */
+	void NeighbourPoints(int iVertex, FVector2<T> &PrevNbrOut, FVector2<T> &NextNbrOut) const
+	{
+		int N = Vertices.Num();
+		PrevNbrOut = Vertices[(iVertex == 0) ? N - 1 : iVertex - 1];
+		NextNbrOut = Vertices[(iVertex + 1) % N];
+	}
+
+
+	/**
+	 * Get the vectors from a given vertex to the previous and next Vertices, optionally normalized
+	 */
+	void NeighbourVectors(int iVertex, FVector2<T> &ToPrevOut, FVector2<T> &ToNextOut, bool bNormalize = false) const
+	{
+		int N = Vertices.Num();
+		ToPrevOut = Vertices[(iVertex == 0) ? N - 1 : iVertex - 1] - Vertices[iVertex];
+		ToNextOut = Vertices[(iVertex + 1) % N] - Vertices[iVertex];
+		if (bNormalize) 
+		{
+			Normalize(ToPrevOut);
+			Normalize(ToNextOut);
+		}
+	}
+
+
+	/**
+	 * @return the opening angle in degrees at a vertex of the Polygon
+	 */
+	T OpeningAngleDeg(int iVertex) const
+	{
+		FVector2<T> e0, e1;
+		NeighbourVectors(iVertex, e0, e1, true);
+		return AngleD(e0, e1);
+	}
+
+
+	/**
+	 * @return analytic winding integral for this Polygon at an arbitrary point
+	 */
+	T WindingIntegral(const FVector2<T>& QueryPoint) const
+	{
+		T sum = 0;
+		int N = Vertices.Num();
+		FVector2<T> a = Vertices[0] - QueryPoint, b = FVector2<T>::Zero();
+		for (int i = 0; i < N; ++i) 
+		{
+			b = Vertices[(i + 1) % N] - QueryPoint;
+			sum += TMathUtil<T>::Atan2(a.X * b.Y - a.Y * b.X, a.X * b.X + a.Y * b.Y);
+			a = b;
+		}
+		return sum / FMathd::TwoPi;
+	}
+
+
+	/**
+	 * @return true if the given query point is inside the Polygon, based on the winding integral
+	 */
+	bool Contains(const FVector2<T>& QueryPoint) const
+	{
+		int nWindingNumber = 0;
+
+		int N = Vertices.Num();
+		FVector2<T> a = Vertices[0], b = FVector2<T>::Zero();
+		for (int i = 0; i < N; ++i) 
+		{
+			b = Vertices[(i + 1) % N];
+
+			if (a.Y <= QueryPoint.Y)     // y <= P.Y (below)
+			{
+				if (b.Y > QueryPoint.Y)									// an upward crossing
+				{
+					if (Orient(a, b, QueryPoint) > 0)  // P left of edge
+						++nWindingNumber;                       // have a valid up intersect
+				}
+			}
+			else     // y > P.Y  (above)
+			{
+				if (b.Y <= QueryPoint.Y)									// a downward crossing
+				{
+					if (Orient(a, b, QueryPoint) < 0)  // P right of edge
+					{
+						--nWindingNumber;						// have a valid down intersect
+					}
+				}
+			}
+			a = b;
+		}
+		return nWindingNumber != 0;
+	}
+
+	/**
+	 * Check for polygon overlap, aka solid intersection.  (In contrast, note that the "Intersects" method checks for edge intersection)
+	 *
+	 * @return true if the Polygon overlaps the OtherPolygon.
+	 */
+	bool Overlaps(const TPolygon2<T>& OtherPoly) const
+	{
+		if (!Bounds().Intersects(OtherPoly.Bounds()))
+		{
+			return false;
+		}
+
+		for (int i = 0, N = OtherPoly.VertexCount(); i < N; ++i) 
+		{
+			if (Contains(OtherPoly[i]))
+			{
+				return true;
+			}
+		}
+
+		for (int i = 0, N = VertexCount(); i < N; ++i) 
+		{
+			if (OtherPoly.Contains(Vertices[i]))
+			{
+				return true;
+			}
+		}
+
+		for (TSegment2<T> seg : Segments()) 
+		{
+			for (TSegment2<T> oseg : OtherPoly.Segments()) 
+			{
+				if (seg.Intersects(oseg))
+				{
+					return true;
+				}
+			}
+		}
+		return false;
+	}
+
+	/**
+	 * @return true if the Polygon fully contains the OtherPolygon
+	 */
+	bool Contains(const TPolygon2<T>& OtherPoly) const
+	{
+		// @todo fast bbox check?
+
+		int N = OtherPoly.VertexCount();
+		for (int i = 0; i < N; ++i) 
+		{
+			if (Contains(OtherPoly[i]) == false)
+			{
+				return false;
+			}
+		}
+
+		if (Intersects(OtherPoly))
+		{
+			return false;
+		}
+
+		return true;
+	}
+
+
+	/**
+	 * @return true if the Segment is fully contained inside the Polygon
+	 */
+	bool Contains(const TSegment2<T>& Segment) const
+	{
+		// [TODO] Add bbox check
+		if (Contains(Segment.StartPoint()) == false || Contains(Segment.EndPoint()) == false)
+		{
+			return false;
+		}
+
+		for (TSegment2<T> seg : Segments())
+		{
+			if (seg.Intersects(Segment))
+			{
+				return false;
+			}
+		}
+		return true;
+	}
+
+
+	/**
+	 * @return true if at least one edge of the OtherPolygon intersects the Polygon
+	 */
+	bool Intersects(const TPolygon2<T>& OtherPoly) const
+	{
+		if (!Bounds().Intersects(OtherPoly.Bounds()))
+		{
+			return false;
+		}
+
+		for (TSegment2<T> seg : Segments()) 
+		{
+			for (TSegment2<T> oseg : OtherPoly.Segments()) 
+			{
+				if (seg.Intersects(oseg))
+				{
+					return true;
+				}
+			}
+		}
+		return false;
+	}
+
+
+	/**
+	 * @return true if the Segment intersects an edge of the Polygon
+	 */
+	bool Intersects(const TSegment2<T>& Segment) const
+	{
+		// [TODO] Add bbox check
+		if (Contains(Segment.StartPoint()) == true || Contains(Segment.EndPoint()) == true)
+		{
+			return true;
+		}
+
+		// [TODO] Add bbox check
+		for (TSegment2<T> seg : Segments())
+		{
+			if (seg.Intersects(Segment))
+			{
+				return true;
+			}
+		}
+		return false;
+	}
+
+
+	/**
+	 * Find all the points where an edge of the Polygon intersects an edge of the OtherPolygon
+	 * @param OtherPoly polygon to test against
+	 * @param OutArray intersection points are stored here
+	 * @return true if any intersections were found
+	 */
+	bool FindIntersections(const TPolygon2<T>& OtherPoly, TArray<FVector2<T>>& OutArray) const
+	{
+		if (!Bounds().Intersects(OtherPoly.Bounds()))
+		{
+			return false;
+		}
+
+		bool bFoundIntersections = false;
+		for (TSegment2<T> seg : Segments()) 
+		{
+			for (TSegment2<T> oseg : OtherPoly.Segments())
+			{
+				// this computes test twice for intersections, but seg.intersects doesn't
+				// create any new objects so it should be much faster for majority of segments (should profile!)
+				if (seg.Intersects(oseg)) 
+				{
+					//@todo can we replace with something like seg.intersects?
+					TIntrSegment2Segment2<T> intr(seg, oseg);
+					if (intr.Find()) 
+					{
+						bFoundIntersections = true;
+						OutArray.Add(intr.Point0);
+						if (intr.Quantity == 2) 
+						{
+							OutArray.Add(intr.Point1);
+						}
+					}
+				}
+			}
+		}
+
+		return bFoundIntersections;
+	}
+
+
+	/**
+	 * @return edge of the polygon starting at vertex SegmentIndex
+	 */
+	TSegment2<T> Segment(int SegmentIndex) const
+	{
+		return TSegment2<T>(Vertices[SegmentIndex], Vertices[(SegmentIndex + 1) % Vertices.Num()]);
+	}
+
+	/**
+	 * @param SegmentIndex index of first vertex of the edge
+	 * @param SegmentParam parameter in range [-Extent,Extent] along segment
+	 * @return point on the segment at the given parameter value
+	 */
+	FVector2<T> GetSegmentPoint(int SegmentIndex, T SegmentParam) const
+	{
+		TSegment2<T> seg(Vertices[SegmentIndex], Vertices[(SegmentIndex + 1) % Vertices.Num()]);
+		return seg.PointAt(SegmentParam);
+	}
+
+	/**
+	 * @param SegmentIndex index of first vertex of the edge
+	 * @param SegmentParam parameter in range [0,1] along segment
+	 * @return point on the segment at the given parameter value
+	 */
+	FVector2<T> GetSegmentPointUnitParam(int SegmentIndex, T SegmentParam) const
+	{
+		TSegment2<T> seg(Vertices[SegmentIndex], Vertices[(SegmentIndex + 1) % Vertices.Num()]);
+		return seg.PointBetween(SegmentParam);
+	}
+
+
+	/**
+	 * @param SegmentIndex index of first vertex of the edge
+	 * @param SegmentParam parameter in range [0,1] along segment
+	 * @return interpolated normal to the segment at the given parameter value
+	 */
+	FVector2<T> GetNormal(int iSeg, T SegmentParam) const
+	{
+		TSegment2<T> seg(Vertices[iSeg], Vertices[(iSeg + 1) % Vertices.Num()]);
+		T t = ((SegmentParam / seg.Extent) + 1.0) / 2.0;
+
+		FVector2<T> n0 = GetNormal(iSeg);
+		FVector2<T> n1 = GetNormal((iSeg + 1) % Vertices.Num());
+		return Normalized((T(1) - t) * n0 + t * n1);
+	}
+
+
+
+	/**
+	 * Calculate the squared distance from a point to the polygon
+	 * @param QueryPoint the query point
+	 * @param NearestSegIndexOut The index of the nearest segment
+	 * @param NearestSegParamOut the parameter value of the nearest point on the segment
+	 * @return squared distance to the polygon
+	 */
+	T DistanceSquared(const FVector2<T>& QueryPoint, int& NearestSegIndexOut, T& NearestSegParamOut) const
+	{
+		NearestSegIndexOut = -1;
+		NearestSegParamOut = TNumericLimits<T>::Max();
+		T dist = TNumericLimits<T>::Max();
+		int N = Vertices.Num();
+		for (int vi = 0; vi < N; ++vi) 
+		{
+			// @todo can't we just use segment function here now?
+			TSegment2<T> seg = TSegment2<T>(Vertices[vi], Vertices[(vi + 1) % N]);
+			T t = (QueryPoint - seg.Center).Dot(seg.Direction);
+			T d = TNumericLimits<T>::Max();
+			if (t >= seg.Extent)
+			{
+				d = seg.EndPoint().DistanceSquared(QueryPoint);
+			}
+			else if (t <= -seg.Extent)
+			{
+				d = seg.StartPoint().DistanceSquared(QueryPoint);
+			}
+			else
+			{
+				d = (seg.PointAt(t) - QueryPoint).SquaredLength();
+			}
+			if (d < dist) 
+			{
+				dist = d;
+				NearestSegIndexOut = vi;
+				NearestSegParamOut = TMathUtil<T>::Clamp(t, -seg.Extent, seg.Extent);
+			}
+		}
+		return dist;
+	}
+
+
+	/**
+	 * Calculate the squared distance from a point to the polygon
+	 * @param QueryPoint the query point
+	 * @return squared distance to the polygon
+	 */
+	T DistanceSquared(const FVector2<T>& QueryPoint) const
+	{
+		int seg; T segt;
+		return DistanceSquared(QueryPoint, seg, segt);
+	}
+
+
+	/**
+	 * @return average edge length of all the edges of the Polygon
+	 */
+	T AverageEdgeLength() const
+	{
+		T avg = 0; int N = Vertices.Num();
+		for (int i = 1; i < N; ++i) {
+			avg += Vertices[i].Distance(Vertices[i - 1]);
+		}
+		avg += Vertices[N - 1].Distance(Vertices[0]);
+		return avg / N;
+	}
+
+
+	/**
+	 * Translate the polygon
+	 * @returns the Polygon, so that you can chain calls like Translate().Scale()
+	 */
+	TPolygon2<T>& Translate(const FVector2<T>& Translate) 
+	{
+		int N = Vertices.Num();
+		for (int i = 0; i < N; ++i)
+		{
+			Vertices[i] += Translate;
+		}
+		Timestamp++;
+		return *this;
+	}
+
+	/**
+	 * Scale the Polygon relative to a given point
+	 * @returns the Polygon, so that you can chain calls like Translate().Scale()
+	 */
+	TPolygon2<T>& Scale(const FVector2<T>& Scale, const FVector2<T>& Origin)
+	{
+		int N = Vertices.Num();
+		for (int i = 0; i < N; ++i)
+		{
+			Vertices[i] = Scale * (Vertices[i] - Origin) + Origin;
+		}
+		Timestamp++;
+		return *this;
+	}
+
+
+	/**
+	 * Apply an arbitrary transformation to the Polygon
+	 * @returns the Polygon, so that you can chain calls like Translate().Scale()
+	 */
+	TPolygon2<T>& Transform(const TFunction<FVector2<T> (const FVector2<T>&)>& TransformFunc)
+	{
+		int N = Vertices.Num();
+		for (int i = 0; i < N; ++i) 
+		{
+			Vertices[i] = TransformFunc(Vertices[i]);
+		}
+		Timestamp++;
+		return *this;
+	}
+
+
+
+
+	/**
+	 * Offset each point by the given Distance along vertex "normal" direction
+	 * @param OffsetDistance the distance to offset
+	 * @param bUseFaceAvg if true, we offset by the average-face normal instead of the perpendicular-tangent normal
+	 */
+	void VtxNormalOffset(T OffsetDistance, bool bUseFaceAvg = false)
+	{
+		TArray<FVector2<T>> NewVertices;
+		NewVertices.SetNumUninitialized(Vertices.Num());
+		if (bUseFaceAvg) 
+		{
+			for (int k = 0; k < Vertices.Num(); ++k)
+			{
+				NewVertices[k] = Vertices[k] + OffsetDistance * GetNormal_FaceAvg(k);
+			}
+		}
+		else 
+		{
+			for (int k = 0; k < Vertices.Num(); ++k)
+			{
+				NewVertices[k] = Vertices[k] + OffsetDistance * GetNormal(k);
+			}
+		}
+		for (int k = 0; k < Vertices.Num(); ++k)
+		{
+			Vertices[k] = NewVertices[k];
+		}
+
+		Timestamp++;
+	}
+
+
+	/**
+	 * Offset polygon by fixed distance, by offsetting and intersecting edges.
+	 * CounterClockWise Polygon offsets "outwards", ClockWise "inwards".
+	 */
+	void PolyOffset(T OffsetDistance)
+	{
+		// [TODO] possibly can do with half as many normalizes if we do w/ sequential edges,
+		//  rather than centering on each v?
+		TArray<FVector2<T>> NewVertices;
+		NewVertices.SetNumUninitialized(Vertices.Num());
+		for (int k = 0; k < Vertices.Num(); ++k) 
+		{
+			FVector2<T> v = Vertices[k];
+			FVector2<T> next = Vertices[(k + 1) % Vertices.Num()];
+			FVector2<T> prev = Vertices[k == 0 ? Vertices.Num() - 1 : k - 1];
+			FVector2<T> dn = Normalized(next - v);
+			FVector2<T> dp = Normalized(prev - v);
+			TLine2<T> ln(v + OffsetDistance * PerpCW(dn), dn);
+			TLine2<T> lp(v - OffsetDistance * PerpCW(dp), dp);
+
+			bool bIntersectsAtPoint = ln.IntersectionPoint(lp, NewVertices[k]);
+			if (!bIntersectsAtPoint)  // lines were parallel
+			{
+				NewVertices[k] = Vertices[k] + OffsetDistance * GetNormal_FaceAvg(k);
+			}
+		}
+		for (int k = 0; k < Vertices.Num(); ++k)
+		{
+			Vertices[k] = NewVertices[k];
+		}
+
+		Timestamp++;
+	}
+
+
+private:
+	// Polygon simplification
+	// code adapted from: http://softsurfer.com/Archive/algorithm_0205/algorithm_0205.htm
+	// simplifyDP():
+	//  This is the Douglas-Peucker recursive simplification routine
+	//  It just marks Vertices that are part of the simplified polyline
+	//  for approximating the polyline subchain v[j] to v[k].
+	//    Input:  tol = approximation tolerance
+	//            v[] = polyline array of vertex points
+	//            j,k = indices for the subchain v[j] to v[k]
+	//    Output: mk[] = array of markers matching vertex array v[]
+	static void SimplifyDouglasPeucker(T Tolerance, const TArray<FVector2<T>>& Vertices, int j, int k, TArray<bool>& Marked)
+	{
+		Marked.SetNum(Vertices.Num());
+		if (k <= j + 1) // there is nothing to simplify
+			return;
+
+		// check for adequate approximation by segment S from v[j] to v[k]
+		int maxi = j;          // index of vertex farthest from S
+		T maxd2 = 0;         // distance squared of farthest vertex
+		T tol2 = Tolerance * Tolerance;  // tolerance squared
+		TSegment2<T> S = TSegment2<T>(Vertices[j], Vertices[k]);    // segment from v[j] to v[k]
+
+		// test each vertex v[i] for max distance from S
+		// Note: this works in any dimension (2D, 3D, ...)
+		for (int i = j + 1; i < k; i++)
+		{
+			T dv2 = S.DistanceSquared(Vertices[i]);
+			if (dv2 <= maxd2)
+				continue;
+			// v[i] is a max vertex
+			maxi = i;
+			maxd2 = dv2;
+		}
+		if (maxd2 > tol2)       // error is worse than the tolerance
+		{
+			// split the polyline at the farthest vertex from S
+			Marked[maxi] = true;      // mark v[maxi] for the simplified polyline
+			// recursively simplify the two subpolylines at v[maxi]
+			SimplifyDouglasPeucker(Tolerance, Vertices, j, maxi, Marked);  // polyline v[j] to v[maxi]
+			SimplifyDouglasPeucker(Tolerance, Vertices, maxi, k, Marked);  // polyline v[maxi] to v[k]
+		}
+		// else the approximation is OK, so ignore intermediate Vertices
+		return;
+	}
+
+
+public:
+
+	/**
+	 * Simplify the Polygon to reduce the vertex count
+	 * @param ClusterTolerance Vertices closer than this distance will be merged into a single vertex
+	 * @param LineDeviationTolerance Vertices are allowed to deviate this much from the input polygon lines
+	 */
+	void Simplify(T ClusterTolerance = 0.0001, T LineDeviationTolerance = 0.01)
+	{
+		int n = Vertices.Num();
+		if (n < 3)
+		{
+			return;
+		}
+
+		int i, k, pv;            // misc counters
+		TArray<FVector2<T>> NewVertices;
+		NewVertices.SetNumUninitialized(n + 1);  // vertex buffer
+		TArray<bool> Marked;
+		Marked.SetNumUninitialized(n + 1);
+		for (i = 0; i < n + 1; ++i)		// marker buffer
+		{
+			Marked[i] = false;
+		}
+
+		// STAGE 1.  Vertex Reduction within tolerance of prior vertex cluster
+		T clusterTol2 = ClusterTolerance * ClusterTolerance;
+		NewVertices[0] = Vertices[0];              // start at the beginning
+		for (i = 1, k = 1, pv = 0; i < n; i++) 
+		{
+			if (Vertices[i].DistanceSquared(Vertices[pv]) < clusterTol2)
+			{
+				continue;
+			}
+			NewVertices[k++] = Vertices[i];
+			pv = i;
+		}
+		bool skip_dp = false;
+		if (k == 1) 
+		{
+			NewVertices[k++] = Vertices[1];
+			NewVertices[k++] = Vertices[2];
+			skip_dp = true;
+		}
+		else if (k == 2) 
+		{
+			NewVertices[k++] = Vertices[0];
+			skip_dp = true;
+		}
+
+		// push on start vertex again, because simplifyDP is for polylines, not polygons
+		NewVertices[k++] = Vertices[0];
+
+		// STAGE 2.  Douglas-Peucker polyline simplification
+		int nv = 0;
+		if (skip_dp == false && LineDeviationTolerance > 0) 
+		{
+			Marked[0] = Marked[k - 1] = true;       // mark the first and last Vertices
+			SimplifyDouglasPeucker(LineDeviationTolerance, NewVertices, 0, k - 1, Marked);
+			for (i = 0; i < k - 1; ++i) 
+			{
+				if (Marked[i])
+				{
+					nv++;
+				}
+			}
+		}
+		else 
+		{
+			for (i = 0; i < k; ++i)
+			{
+				Marked[i] = true;
+			}
+			nv = k - 1;
+		}
+
+		// polygon requires at least 3 Vertices
+		if (nv == 2) 
+		{
+			for (i = 1; i < k - 1; ++i) 
+			{
+				if (Marked[1] == false)
+				{
+					Marked[1] = true;
+				}
+				else if (Marked[k - 2] == false)
+				{
+					Marked[k - 2] = true;
+				}
+			}
+			nv++;
+		}
+		else if (nv == 1) 
+		{
+			Marked[1] = true;
+			Marked[2] = true;
+			nv += 2;
+		}
+
+		// copy marked Vertices back to this polygon
+		Vertices.Reset();
+		for (i = 0; i < k - 1; ++i)    // last vtx is copy of first, and definitely marked
+		{
+			if (Marked[i]) 
+			{
+				Vertices.Add(NewVertices[i]);
+			}
+		}
+
+		Timestamp++;
+		return;
+	}
+
+
+
+	/**
+	 * Chamfer each vertex corner of the Polygon
+	 * @param ChamferDist offset distance from corner that we cut at
+	 */
+	void Chamfer(T ChamferDist, T MinConvexAngleDeg = 30, T MinConcaveAngleDeg = 30)
+	{
+		check(IsClockwise());
+
+		TArray<FVector2<T>> OldV = Vertices;
+		int N = OldV.Num();
+		Vertices.Reset();
+
+		int iCur = 0;
+		do {
+			FVector2<T> center = OldV[iCur];
+
+			int iPrev = (iCur == 0) ? N - 1 : iCur - 1;
+			FVector2<T> prev = OldV[iPrev];
+			int iNext = (iCur + 1) % N;
+			FVector2<T> next = OldV[iNext];
+
+			FVector2<T> cp = prev - center;
+			T cpdist = Normalize(cp);
+			FVector2<T> cn = next - center;
+			T cndist = Normalize(cn);
+
+			// if degenerate, skip this vert
+			if (cpdist < TMathUtil<T>::ZeroTolerance || cndist < TMathUtil<T>::ZeroTolerance) 
+			{
+				iCur = iNext;
+				continue;
+			}
+
+			T angle = AngleD(cp, cn);
+			// TODO document what this means sign-wise
+			// TODO re-test post Unreal port that this DotPerp is doing the right thing
+			T sign = DotPerp(cn, cp);
+			bool bConcave = (sign > 0);
+
+			T thresh = (bConcave) ? MinConcaveAngleDeg : MinConvexAngleDeg;
+
+			// ok not too sharp
+			if (angle > thresh) 
+			{
+				Vertices.Add(center);
+				iCur = iNext;
+				continue;
+			}
+
+
+			T prev_cut_dist = TMathUtil<T>::Min(ChamferDist, cpdist*0.5);
+			FVector2<T> prev_cut = center + cp * prev_cut_dist;
+			T next_cut_dist = TMathUtil<T>::Min(ChamferDist, cndist * 0.5);
+			FVector2<T> next_cut = center + cn * next_cut_dist;
+
+			Vertices.Add(prev_cut);
+			Vertices.Add(next_cut);
+			iCur = iNext;
+		} while (iCur != 0);
+
+		Timestamp++;
+	}
+
+
+
+
+	/**
+	 * Construct a four-vertex rectangle polygon
+	 */
+	static TPolygon2<T> MakeRectangle(const FVector2<T>& Center, T Width, T Height)
+	{
+		TPolygon2<T> Rectangle;
+		Rectangle.Vertices.SetNumUninitialized(4);
+		Rectangle.Set(0, FVector2<T>(Center.X - Width / 2, Center.Y - Height / 2));
+		Rectangle.Set(1, FVector2<T>(Center.X + Width / 2, Center.Y - Height / 2));
+		Rectangle.Set(2, FVector2<T>(Center.X + Width / 2, Center.Y + Height / 2));
+		Rectangle.Set(3, FVector2<T>(Center.X - Width / 2, Center.Y + Height / 2));
+		return Rectangle;
+	}
+
+	/**
+	* Construct a rounded rectangle polygon
+	*/
+	static TPolygon2<T> MakeRoundedRectangle(const FVector2<T>& Center, T Width, T Height, T Corner, int CornerSteps)
+	{
+		TPolygon2<T> RRectangle;
+		CornerSteps = FMath::Max(2, CornerSteps);
+		RRectangle.Vertices.SetNumUninitialized(CornerSteps * 4);
+		
+		FVector2<T> InnerExtents(Width*.5 - Corner, Height*.5 - Corner);
+		FVector2<T> InnerCorners[4]
+		{
+			Center - InnerExtents,
+			FVector2<T>(Center.X + InnerExtents.X, Center.Y - InnerExtents.Y),
+			Center + InnerExtents,
+			FVector2<T>(Center.X - InnerExtents.X, Center.Y + InnerExtents.Y),
+		};
+
+		T IdxToAng = TMathUtil<T>::HalfPi / T(CornerSteps - 1);
+		for (int StepIdx = 0; StepIdx < CornerSteps; StepIdx++)
+		{
+			T Angle = StepIdx * IdxToAng;
+			T OffC = TMathUtil<T>::Cos(Angle) * Corner;
+			T OffS = TMathUtil<T>::Sin(Angle) * Corner;
+			RRectangle.Set(StepIdx + CornerSteps * 0, InnerCorners[0] + FVector2<T>(-OffC, -OffS));
+			RRectangle.Set(StepIdx + CornerSteps * 1, InnerCorners[1] + FVector2<T>(OffS, -OffC));
+			RRectangle.Set(StepIdx + CornerSteps * 2, InnerCorners[2] + FVector2<T>(OffC, OffS));
+			RRectangle.Set(StepIdx + CornerSteps * 3, InnerCorners[3] + FVector2<T>(-OffS, OffC));
+		}
+		return RRectangle;
+	}
+
+
+	/**
+	 * Construct a circular polygon
+	 */
+	static TPolygon2<T> MakeCircle(T Radius, int Steps, T AngleShiftRadians = 0)
+	{
+		TPolygon2<T> Circle;
+		Circle.Vertices.SetNumUninitialized(Steps);
+
+		for (int i = 0; i < Steps; ++i)
+		{
+			T t = (T)i / (T)Steps;
+			T a = TMathUtil<T>::TwoPi * t + AngleShiftRadians;
+			Circle.Set(i, FVector2<T>(Radius * TMathUtil<T>::Cos(a), Radius * TMathUtil<T>::Sin(a)));
+		}
+
+		return Circle;
+	}
+};
+
+typedef TPolygon2<double> FPolygon2d;
+typedef TPolygon2<float> FPolygon2f;
+
+} // end namespace UE::Geometry
+} // end namespace UE
\ No newline at end of file
diff --git a/Engine/Source/Runtime/GeometryCore/Public/RayTypes.h b/Engine/Source/Runtime/GeometryCore/Public/RayTypes.h
new file mode 100644
--- /dev/null
+++ b/Engine/Source/Runtime/GeometryCore/Public/RayTypes.h
@@ -0,0 +1,131 @@
+// Copyright Epic Games, Inc. All Rights Reserved.
+
+// port of geometry3Sharp Ray3
+
+#pragma once
+
+#include "Math/UnrealMath.h"
+#include "VectorTypes.h"
+#include "Math/Ray.h"
+
+namespace UE
+{
+namespace Geometry
+{
+
+/*
+ * 3D Ray stored as Origin point and normalized Direction vector
+ */
+template<typename RealType>
+class TRay3
+{
+public:
+	/** Origin point */
+	FVector3<RealType> Origin;
+	/** Direction vector, always normalized */
+	FVector3<RealType> Direction;
+
+	/** Construct ray at origin pointed down Z axis */
+	TRay3()
+	{
+		Origin = FVector3<RealType>::Zero();
+		Direction = FVector3<RealType>::UnitZ();
+	}
+
+	/**
+	 * Construct ray from origin and direction
+	 * @param Origin origin point
+	 * @param Direction direction vector. Must be normalized if bIsNormalized=true
+	 * @param bIsNormalized if true, Direction will not be re-normalized
+	 */
+	TRay3(const FVector3<RealType>& Origin, const FVector3<RealType>& Direction, bool bIsNormalized = false)
+	{
+		this->Origin = Origin;
+		this->Direction = Direction;
+		if (bIsNormalized == false)
+		{
+			Normalize(this->Direction);
+		}
+	}
+
+	/**
+	 * @return point on ray at given (signed) Distance from the ray Origin
+	 */
+	FVector3<RealType> PointAt(RealType Distance) const
+	{
+		return Origin + Distance * Direction;
+	}
+
+
+	/**
+	 * @return ray parameter (ie positive distance from Origin) at nearest point on ray to QueryPoint
+	 */
+	inline RealType Project(const FVector3<RealType>& QueryPoint) const
+	{
+		RealType LineParam = (QueryPoint - Origin).Dot(Direction);
+		return (LineParam < 0) ? 0 : LineParam;
+	}
+
+	/**
+	 * @return smallest squared distance from ray to QueryPoint
+	 */
+	inline RealType DistanceSquared(const FVector3<RealType>& QueryPoint) const
+	{
+		RealType LineParam = (QueryPoint - Origin).Dot(Direction);
+		if (LineParam < 0)
+		{
+			return Origin.DistanceSquared(QueryPoint);
+		}
+		else
+		{
+			FVector3<RealType> NearestPt = Origin + LineParam * Direction;
+			return NearestPt.DistanceSquared(QueryPoint);
+		}
+	}
+
+	/**
+	 * @return smallest squared distance from ray to QueryPoint
+	 */
+	inline RealType Distance(const FVector3<RealType>& QueryPoint) const
+	{
+		return TMathUtil<RealType>::Sqrt(DistanceSquared(QueryPoint));
+	}
+
+	/**
+	 * @return nearest point on line to QueryPoint
+	 */
+	inline FVector3<RealType> NearestPoint(const FVector3<RealType>& QueryPoint) const
+	{
+		RealType LineParam = (QueryPoint - Origin).Dot(Direction);
+		if (LineParam < 0)
+		{
+			return Origin;
+		}
+		else
+		{
+			return Origin + LineParam * Direction;
+		}
+	}
+
+
+
+	// conversion operators
+
+	explicit inline operator FRay() const
+	{
+		return FRay((FVector)Origin, (FVector)Direction);
+	}
+	inline TRay3(const FRay & RayIn)
+	{
+		Origin = FVector3<RealType>(RayIn.Origin);
+		Direction = FVector3<RealType>(RayIn.Direction);
+	}
+
+
+};
+typedef TRay3<float> FRay3f;
+typedef TRay3<double> FRay3d;
+
+} // end namespace UE::Geometry
+} // end namespace UE
+
diff --git a/Engine/Source/Runtime/GeometryCore/Public/Spatial/DenseGrid2.h b/Engine/Source/Runtime/GeometryCore/Public/Spatial/DenseGrid2.h
new file mode 100644
--- /dev/null
+++ b/Engine/Source/Runtime/GeometryCore/Public/Spatial/DenseGrid2.h
@@ -0,0 +1,211 @@
+// Copyright Epic Games, Inc. All Rights Reserved.
+
+// Port of geometry3Sharp DenseGrid2
+
+#pragma once
+
+#include "CoreMinimal.h"
+#include "BoxTypes.h"
+#include "IntBoxTypes.h"
+
+#include "HAL/PlatformAtomics.h"
+
+namespace UE
+{
+namespace Geometry
+{
+
+/**
+ * 2D dense grid of floating-point scalar values.
+ */
+template<typename ElemType>
+class TDenseGrid2
+{
+protected:
+	/** grid of allocated elements */
+	TArray64<ElemType> Buffer;
+
+	/** dimensions per axis */
+	int64 DimensionX, DimensionY;
+
+public:
+	/**
+	 * Create empty grid
+	 */
+	TDenseGrid2() : DimensionX(0), DimensionY(0)
+	{
+	}
+
+	TDenseGrid2(int32 DimX, int32 DimY, ElemType InitialValue)
+	{
+		Resize(DimX, DimY);
+		Assign(InitialValue);
+	}
+
+	int64 Size() const
+	{
+		return DimensionX * DimensionY;
+	}
+
+	int64 Width() const
+	{
+		return DimensionX;
+	}
+
+	int64 Height() const
+	{
+		return DimensionY;
+	}
+
+	FVector2i GetDimensions() const
+	{
+		return FVector2i((int32)DimensionX, (int32)DimensionY);
+	}
+
+	void Resize(int32 DimX, int32 DimY, bool bAllowShrinking = true)
+	{
+		DimensionX = (int64)DimX;
+		DimensionY = (int64)DimY;
+		Buffer.SetNumUninitialized( DimensionX*DimensionY, bAllowShrinking);
+	}
+
+	void AssignAll(ElemType Value)
+	{
+		for (int64 i = 0, Num = Size(); i < Num; ++i)
+		{
+			Buffer[i] = Value;
+		}
+	}
+
+	void SetMin(const FVector2i& IJK, ElemType NewValue)
+	{
+		int64 Idx = (int64)IJK.X + (DimensionX * (int64)IJK.Y);
+		if (NewValue < Buffer[Idx])
+		{
+			Buffer[Idx] = NewValue;
+		}
+	}
+
+	void SetMax(const FVector2i& IJK, ElemType NewValue)
+	{
+		int64 Idx = (int64)IJK.X + (DimensionX * (int64)IJK.Y);
+		if (NewValue > Buffer[Idx])
+		{
+			Buffer[Idx] = NewValue;
+		}
+	}
+
+	constexpr ElemType& operator[](int64 Idx)
+	{
+		return Buffer[Idx];
+	}
+	constexpr const ElemType& operator[](int64 Idx) const
+	{
+		return Buffer[Idx];
+	}
+
+	constexpr ElemType& operator[](FVector2i Idx)
+	{
+		return Buffer[(int64)Idx.X + DimensionX * (int64)Idx.Y];
+	}
+	constexpr const ElemType& operator[](FVector2i Idx) const
+	{
+		return Buffer[(int64)Idx.X + DimensionX * (int64)Idx.Y];
+	}
+	constexpr ElemType& At(int32 X, int32 Y)
+	{
+		return Buffer[(int64)X + DimensionX * (int64)Y];
+	}
+	constexpr const ElemType& At(int32 X, int32 Y) const
+	{
+		return Buffer[(int64)X + DimensionX * (int64)Y];
+	}
+	constexpr ElemType& At(int64 X, int64 Y)
+	{
+		return Buffer[X + DimensionX * Y];
+	}
+	constexpr const ElemType& At(int64 X, int64 Y) const
+	{
+		return Buffer[X + DimensionX * Y];
+	}
+
+	const TArray64<ElemType>& GridValues() const
+	{
+		return Buffer;
+	}
+
+
+	void GetXPair(int32 X0, int32 Y, ElemType& AOut, ElemType& BOut) const
+	{
+		int64 Offset = DimensionX * (int64)Y;
+		AOut = Buffer[Offset + X0];
+		BOut = Buffer[Offset + X0 + 1];
+	}
+
+	void Apply(TFunctionRef<ElemType(ElemType)> ApplyFunc)
+	{
+		for (int64 Idx = 0, Num = Size(); Idx < Num; Idx++)
+		{
+			Buffer[Idx] = F(Buffer[Idx]);
+		}
+	}
+
+	FAxisAlignedBox2i Bounds() const
+	{
+		return FAxisAlignedBox2i({ 0, 0}, { (int32)DimensionX, (int32)DimensionY});
+	}
+	FAxisAlignedBox2i BoundsInclusive() const
+	{
+		return FAxisAlignedBox2i({ 0, 0}, { (int32)(DimensionX - 1), (int32)(DimensionY - 1)});
+	}
+
+	FVector2i GetCoords(int64 LinearIndex) const
+	{
+		checkSlow(LinearIndex >= 0);
+		return FVector2i((int32)(LinearIndex % (int64)DimensionX), (int32)(LinearIndex / (int64)DimensionX));
+	}
+	int64 GetIndex(int32 X, int32 Y) const
+	{
+		return X + DimensionX * Y;
+	}
+	int64 GetIndex(const FVector2i& IJK) const
+	{
+		return IJK.X + DimensionX * IJK.Y;
+	}
+};
+
+
+typedef TDenseGrid2<float> FDenseGrid2f;
+typedef TDenseGrid2<double> FDenseGrid2d;
+typedef TDenseGrid2<int32> FDenseGrid2i;
+
+
+// additional utility functions
+namespace DenseGrid
+{
+
+	inline void AtomicIncrement(FDenseGrid2i& Grid, int32 i, int32 j)
+	{
+		FPlatformAtomics::InterlockedIncrement(&Grid.At(i, j));
+	}
+
+	inline void AtomicDecrement(FDenseGrid2i& Grid, int32 i, int32 j)
+	{
+		FPlatformAtomics::InterlockedDecrement(&Grid.At(i, j));
+	}
+
+	inline void AtomicIncDec(FDenseGrid2i& Grid, int32 i, int32 j, bool bDecrement = false)
+	{
+		if (bDecrement)
+		{
+			FPlatformAtomics::InterlockedDecrement(&Grid.At(i, j));
+		}
+		else
+		{
+			FPlatformAtomics::InterlockedIncrement(&Grid.At(i, j));
+		}
+	}
+}
+
+} // end namespace UE::Geometry
+} // end namespace UE
diff --git a/Engine/Source/Runtime/GeometryCore/Public/Spatial/DenseGrid3.h b/Engine/Source/Runtime/GeometryCore/Public/Spatial/DenseGrid3.h
new file mode 100644
--- /dev/null
+++ b/Engine/Source/Runtime/GeometryCore/Public/Spatial/DenseGrid3.h
@@ -0,0 +1,234 @@
+// Copyright Epic Games, Inc. All Rights Reserved.
+
+// Port of geometry3Sharp DenseGrid3
+
+#pragma once
+
+#include "CoreMinimal.h"
+#include "BoxTypes.h"
+#include "IntBoxTypes.h"
+
+#include "HAL/PlatformAtomics.h"
+
+namespace UE
+{
+namespace Geometry
+{
+
+/**
+ * 3D dense grid of floating-point scalar values. 
+ */
+template<typename ElemType>
+class TDenseGrid3
+{
+protected:
+	/** grid of allocated elements */
+	TArray<ElemType> Buffer;
+
+	/** dimensions per axis */
+	FVector3i Dimensions;
+
+public:
+	/**
+	 * Create empty grid
+	 */
+	TDenseGrid3() : Dimensions(0,0,0)
+	{
+	}
+
+	TDenseGrid3(int DimX, int DimY, int DimZ, ElemType InitialValue)
+	{
+		Resize(DimX, DimY, DimZ);
+		Assign(InitialValue);
+	}
+
+	int Size() const
+	{
+		return Dimensions.X * Dimensions.Y * Dimensions.Z;
+	}
+
+	bool IsValidIndex(const FVector3i& Index) const
+	{
+		return Index.X >= 0 && Index.Y >= 0 && Index.Z >= 0
+			&& Index.X < Dimensions.X && Index.Y < Dimensions.Y && Index.Z < Dimensions.Z;
+	}
+
+	const FVector3i& GetDimensions() const
+	{
+		return Dimensions;
+	}
+
+	void Resize(int DimX, int DimY, int DimZ, bool bAllowShrinking = true)
+	{
+		check((int64)DimX * (int64)DimY * (int64)DimZ < INT_MAX);
+		Buffer.SetNumUninitialized(DimX * DimY * DimZ, bAllowShrinking);
+		Dimensions = FVector3i(DimX, DimY, DimZ);
+	}
+
+	void Assign(ElemType Value)
+	{
+		for (int i = 0; i < Buffer.Num(); ++i)
+		{
+			Buffer[i] = Value;
+		}
+	}
+
+	void SetMin(const FVector3i& IJK, ElemType F)
+	{
+		int Idx = IJK.X + Dimensions.X * (IJK.Y + Dimensions.Y * IJK.Z);
+		if (F < Buffer[Idx])
+		{
+			Buffer[Idx] = F;
+		}
+	}
+	void SetMax(const FVector3i& IJK, ElemType F)
+	{
+		int Idx = IJK.X + Dimensions.X * (IJK.Y + Dimensions.Y * IJK.Z);
+		if (F > Buffer[Idx])
+		{
+			Buffer[Idx] = F;
+		}
+	}
+
+	constexpr ElemType& operator[](int Idx)
+	{
+		return Buffer[Idx];
+	}
+	constexpr const ElemType& operator[](int Idx) const
+	{
+		return Buffer[Idx];
+	}
+
+	constexpr ElemType& operator[](FVector3i Idx)
+	{
+		return Buffer[Idx.X + Dimensions.X * (Idx.Y + Dimensions.Y * Idx.Z)];
+	}
+	constexpr const ElemType& operator[](FVector3i Idx) const
+	{
+		return Buffer[Idx.X + Dimensions.X * (Idx.Y + Dimensions.Y * Idx.Z)];
+	}
+	constexpr ElemType& At(int I, int J, int K)
+	{
+		return Buffer[I + Dimensions.X * (J + Dimensions.Y * K)];
+	}
+	constexpr const ElemType& At(int I, int J, int K) const
+	{
+		return Buffer[I + Dimensions.X * (J + Dimensions.Y * K)];
+	}
+
+
+	void GetXPair(int X0, int Y, int Z, ElemType& AOut, ElemType& BOut) const
+	{
+		int Offset = Dimensions.X * (Y + Dimensions.Y * Z);
+		AOut = Buffer[Offset + X0];
+		BOut = Buffer[Offset + X0 + 1];
+	}
+
+	void Apply(TFunctionRef<ElemType(ElemType)> F)
+	{
+		for (int Idx = 0, Num = Size(); Idx < Num; Idx++)
+		{
+			Buffer[Idx] = F(Buffer[Idx]);
+		}
+	}
+
+	// TODO: implement TDenseGrid2<ElemType> and then implement this
+	//TDenseGrid2<ElemType> get_slice(int slice_i, int dimension)
+	//{
+	//	TDenseGrid2<ElemType> slice;
+	//	if (dimension == 0) {
+	//		slice = TDenseGrid2<ElemType>(Dimensions.Y, Dimensions.Z, 0);
+	//		for (int k = 0; k < Dimensions.Z; ++k)
+	//			for (int j = 0; j < Dimensions.Y; ++j)
+	//				slice[j, k] = Buffer[slice_i + Dimensions.X * (j + Dimensions.Y * k)];
+	//	} else if (dimension == 1) {
+	//		slice = TDenseGrid2<ElemType>(Dimensions.X, Dimensions.Z, 0);
+	//		for (int k = 0; k < Dimensions.Z; ++k)
+	//			for (int i = 0; i < Dimensions.X; ++i)
+	//				slice[i, k] = Buffer[i + Dimensions.X * (slice_i + Dimensions.Y * k)];
+	//	} else {
+	//		slice = TDenseGrid2<ElemType>(Dimensions.X, Dimensions.Y, 0);
+	//		for (int j = 0; j < Dimensions.Y; ++j)
+	//			for (int i = 0; i < Dimensions.X; ++i)
+	//				slice[i, j] = Buffer[i + Dimensions.X * (j + Dimensions.Y * slice_i)];
+	//	}
+	//	return slice;
+	//}
+	//void set_slice(TDenseGrid2<ElemType> slice, int slice_i, int dimension)
+	//{
+	//	if (dimension == 0) {
+	//		for (int k = 0; k < Dimensions.Z; ++k)
+	//			for (int j = 0; j < Dimensions.Y; ++j)
+	//				Buffer[slice_i + Dimensions.X * (j + Dimensions.Y * k)] = slice[j, k];
+	//	} else if (dimension == 1) {
+	//		for (int k = 0; k < Dimensions.Z; ++k)
+	//			for (int i = 0; i < Dimensions.X; ++i)
+	//				Buffer[i + Dimensions.X * (slice_i + Dimensions.Y * k)] = slice[i, k];
+	//	} else {
+	//		for (int j = 0; j < Dimensions.Y; ++j)
+	//			for (int i = 0; i < Dimensions.X; ++i)
+	//				Buffer[i + Dimensions.X * (j + Dimensions.Y * slice_i)] = slice[i, j];
+	//	}
+	//}
+
+
+	FAxisAlignedBox3i Bounds() const
+	{
+		return FAxisAlignedBox3i({0, 0, 0},{Dimensions.X, Dimensions.Y, Dimensions.Z});
+	}
+	FAxisAlignedBox3i BoundsInclusive() const
+	{
+		return FAxisAlignedBox3i({0, 0, 0},{Dimensions.X-1, Dimensions.Y-1, Dimensions.Z-1});
+	}
+
+	FVector3i ToIndex(int Idx) const
+	{
+		int x = Idx % Dimensions.X;
+		int y = (Idx / Dimensions.X) % Dimensions.Y;
+		int z = Idx / (Dimensions.X * Dimensions.Y);
+		return FVector3i(x, y, z);
+	}
+	int ToLinear(int X, int Y, int Z) const
+	{
+		return X + Dimensions.X * (Y + Dimensions.Y * Z);
+	}
+	int ToLinear(const FVector3i& IJK) const
+	{
+		return IJK.X + Dimensions.X * (IJK.Y + Dimensions.Y * IJK.Z);
+	}
+};
+
+
+typedef TDenseGrid3<float> FDenseGrid3f;
+typedef TDenseGrid3<double> FDenseGrid3d;
+typedef TDenseGrid3<int> FDenseGrid3i;
+
+// additional utility functions
+namespace DenseGrid
+{
+
+	inline void AtomicIncrement(FDenseGrid3i& Grid, int i, int j, int k)
+	{
+		FPlatformAtomics::InterlockedIncrement(&Grid.At(i,j,k));
+	}
+
+	inline void AtomicDecrement(FDenseGrid3i& Grid, int i, int j, int k)
+	{
+		FPlatformAtomics::InterlockedDecrement(&Grid.At(i, j, k));
+	}
+
+	inline void AtomicIncDec(FDenseGrid3i& Grid, int i, int j, int k, bool bDecrement = false)
+	{
+		if (bDecrement)
+		{
+			FPlatformAtomics::InterlockedDecrement(&Grid.At(i, j, k));
+		}
+		else
+		{
+			FPlatformAtomics::InterlockedIncrement(&Grid.At(i, j, k));
+		}
+	}
+}
+
+} // end namespace UE::Geometry
+} // end namespace UE
diff --git a/Engine/Source/Runtime/GeometryCore/Public/TransformTypes.h b/Engine/Source/Runtime/GeometryCore/Public/TransformTypes.h
new file mode 100644
--- /dev/null
+++ b/Engine/Source/Runtime/GeometryCore/Public/TransformTypes.h
@@ -0,0 +1,342 @@
+// Copyright Epic Games, Inc. All Rights Reserved.
+
+#pragma once
+
+#include "CoreMinimal.h"
+#include "VectorTypes.h"
+#include "Quaternion.h"
+
+namespace UE
+{
+namespace Geometry
+{
+
+/**
+ * TTransform3 is a double/float templated version of standard UE FTransform.
+ */
+template<typename RealType>
+class TTransform3
+{
+protected:
+	TQuaternion<RealType> Rotation;
+	FVector3<RealType> Translation;
+	FVector3<RealType> Scale3D;
+
+public:
+
+	TTransform3()
+	{
+		Rotation = TQuaternion<RealType>::Identity();
+		Translation = FVector3<RealType>::Zero();
+		Scale3D = FVector3<RealType>::One();
+	}
+
+	TTransform3(const TQuaternion<RealType>& RotationIn, const UE::Math::TVector<RealType>& TranslationIn, const UE::Math::TVector<RealType>& ScaleIn)
+	{
+		Rotation = RotationIn;
+		Translation = TranslationIn;
+		Scale3D = ScaleIn;
+	}
+
+	explicit TTransform3(const TQuaternion<RealType>& RotationIn, const UE::Math::TVector<RealType>& TranslationIn)
+	{
+		Rotation = RotationIn;
+		Translation = TranslationIn;
+		Scale3D = FVector3<RealType>::One();
+	}
+
+	explicit TTransform3(const UE::Math::TVector<RealType>& TranslationIn)
+	{
+		Rotation = TQuaternion<RealType>::Identity();
+		Translation = TranslationIn;
+		Scale3D = FVector3<RealType>::One();
+	}
+
+	explicit TTransform3(const FTransform& Transform)
+	{
+		Rotation = TQuaternion<RealType>(Transform.GetRotation());
+		Translation = FVector3<RealType>(Transform.GetTranslation());
+		Scale3D = FVector3<RealType>(Transform.GetScale3D());
+	}
+
+	/**
+	 * @return identity transform, IE no rotation, zero origin, unit scale
+	 */
+	static TTransform3<RealType> Identity()
+	{
+		return TTransform3<RealType>(TQuaternion<RealType>::Identity(), FVector3<RealType>::Zero(), FVector3<RealType>::One());
+	}
+
+	/**
+	 * @return this transform converted to FTransform 
+	 */
+	explicit operator FTransform() const
+	{
+		return FTransform((FQuat)Rotation, (FVector)Translation, (FVector)Scale3D);
+	}
+
+
+	/**
+	 * @return Rotation portion of Transform, as Quaternion
+	 */
+	const TQuaternion<RealType>& GetRotation() const 
+	{ 
+		return Rotation; 
+	}
+
+	/** 
+	 * Set Rotation portion of Transform 
+	 */
+	void SetRotation(const TQuaternion<RealType>& RotationIn)
+	{
+		Rotation = RotationIn;
+	}
+
+	/**
+	 * @return Translation portion of transform
+	 */
+	const FVector3<RealType>& GetTranslation() const
+	{
+		return Translation;
+	}
+
+	/**
+	 * set Translation portion of transform
+	 */
+	void SetTranslation(const UE::Math::TVector<RealType>& TranslationIn)
+	{
+		Translation = TranslationIn;
+	}
+
+	/**
+	 * @return Scale portion of transform
+	 */
+	const FVector3<RealType>& GetScale() const
+	{
+		return Scale3D;
+	}
+
+	/**
+	 * set Scale portion of transform
+	 */
+	void SetScale(const UE::Math::TVector<RealType>& ScaleIn)
+	{
+		Scale3D = ScaleIn;
+	}
+
+	/**
+	 * @return determinant of scale matrix
+	 */
+	RealType GetDeterminant() const
+	{
+		return Scale3D.X * Scale3D.Y * Scale3D.Z;
+	}
+
+	/**
+	 * @return true if scale is nonuniform, within tolerance
+	 */
+	bool HasNonUniformScale(RealType Tolerance = TMathUtil<RealType>::ZeroTolerance) const
+	{
+		return (TMathUtil<RealType>::Abs(Scale3D.X - Scale3D.Y) > Tolerance)
+			|| (TMathUtil<RealType>::Abs(Scale3D.X - Scale3D.Z) > Tolerance)
+			|| (TMathUtil<RealType>::Abs(Scale3D.Y - Scale3D.Z) > Tolerance);
+	}
+
+	/**
+	 * @return inverse of this transform
+	 */
+	TTransform3<RealType> Inverse() const
+	{
+		TQuaternion<RealType> InvRotation = Rotation.Inverse();
+		FVector3<RealType> InvScale3D = GetSafeScaleReciprocal(Scale3D);
+		FVector3<RealType> InvTranslation = InvRotation * (InvScale3D * -Translation);
+		return TTransform3<RealType>(InvRotation, InvTranslation, InvScale3D);
+	}
+
+
+	/**
+	 * @return input point with QST transformation applied, ie QST(P) = Rotate(Scale*P) + Translate
+	 */
+	FVector3<RealType> TransformPosition(const FVector3<RealType>& P) const
+	{
+		//Transform using QST is following
+		//QST(P) = Q.Rotate(S*P) + T where Q = quaternion, S = scale, T = translation
+		return Rotation * (Scale3D*P) + Translation;
+	}
+
+	/**
+	 * @return input vector with QS transformation applied, ie QS(V) = Rotate(Scale*V)
+	 */
+	FVector3<RealType> TransformVector(const UE::Math::TVector<RealType>& V) const
+	{
+		return Rotation * (Scale3D*V);
+	}
+
+	/**
+	 * @return input vector with Q transformation applied, ie Q(V) = Rotate(V)
+	 */
+	FVector3<RealType> TransformVectorNoScale(const UE::Math::TVector<RealType>& V) const
+	{
+		return Rotation * V;
+	}
+
+	/**
+	 * Surface Normals are special, their transform is Rotate( Normalize( (1/Scale) * Normal) ) ).
+	 * However 1/Scale requires special handling in case any component is near-zero.
+	 * @return input surface normal with transform applied.
+	 */
+	FVector3<RealType> TransformNormal(const UE::Math::TVector<RealType>& Normal) const
+	{
+		// transform normal by a safe inverse scale + normalize, and a standard rotation
+		const FVector3<RealType>& S = Scale3D;
+		RealType DetSign = FMathd::SignNonZero(S.X * S.Y * S.Z); // we only need to multiply by the sign of the determinant, rather than divide by it, since we normalize later anyway
+		FVector3<RealType> SafeInvS(S.Y*S.Z*DetSign, S.X*S.Z*DetSign, S.X*S.Y*DetSign);
+		return TransformVectorNoScale( Normalized(SafeInvS*Normal) );
+	}
+
+
+	/**
+	 * @return input vector with inverse-QST transformation applied, ie QSTinv(P) = InverseScale(InverseRotate(P - Translate))
+	 */
+	FVector3<RealType> InverseTransformPosition(const UE::Math::TVector<RealType> &P) const
+	{
+		return GetSafeScaleReciprocal(Scale3D) * Rotation.InverseMultiply(P - Translation);
+	}
+
+	/**
+	 * @return input vector with inverse-QS transformation applied, ie QSinv(V) = InverseScale(InverseRotate(V))
+	 */
+	FVector3<RealType> InverseTransformVector(const UE::Math::TVector<RealType> &V) const
+	{
+		return GetSafeScaleReciprocal(Scale3D) * Rotation.InverseMultiply(V);
+	}
+
+
+	/**
+	 * @return input vector with inverse-Q transformation applied, ie Qinv(V) = InverseRotate(V)
+	 */
+	FVector3<RealType> InverseTransformVectorNoScale(const UE::Math::TVector<RealType> &V) const
+	{
+		return Rotation.InverseMultiply(V);
+	}
+
+
+	/**
+	 * Surface Normals are special, their inverse transform is InverseRotate( Normalize(Scale * Normal) ) )
+	 * @return input surface normal with inverse transform applied.
+	 */
+	FVector3<RealType> InverseTransformNormal(const UE::Math::TVector<RealType>& Normal) const
+	{
+		return InverseTransformVectorNoScale( Normalized(Scale3D*Normal) );
+	}
+
+
+
+	/**
+	 * Clamp all scale components to a minimum value. Sign of scale components is preserved.
+	 * This is used to remove uninvertible zero/near-zero scaling.
+	 */
+	void ClampMinimumScale(RealType MinimumScale = TMathUtil<RealType>::ZeroTolerance)
+	{
+		for (int j = 0; j < 3; ++j)
+		{
+			RealType Value = Scale3D[j];
+			if (TMathUtil<RealType>::Abs(Value) < MinimumScale)
+			{
+				Value = MinimumScale * TMathUtil<RealType>::SignNonZero(Value);
+				Scale3D[j] = Value;
+			}
+		}
+	}
+
+	
+
+	static FVector3<RealType> GetSafeScaleReciprocal(const FVector3<RealType>& InScale, RealType Tolerance = TMathUtil<RealType>::ZeroTolerance)
+	{
+		FVector3<RealType> SafeReciprocalScale;
+		if (TMathUtil<RealType>::Abs(InScale.X) <= Tolerance)
+		{
+			SafeReciprocalScale.X = (RealType)0;
+		}
+		else
+		{
+			SafeReciprocalScale.X = (RealType)1 / InScale.X;
+		}
+
+		if (TMathUtil<RealType>::Abs(InScale.Y) <= Tolerance)
+		{
+			SafeReciprocalScale.Y = (RealType)0;
+		}
+		else
+		{
+			SafeReciprocalScale.Y = (RealType)1 / InScale.Y;
+		}
+
+		if (TMathUtil<RealType>::Abs(InScale.Z) <= Tolerance)
+		{
+			SafeReciprocalScale.Z = (RealType)0;
+		}
+		else
+		{
+			SafeReciprocalScale.Z = (RealType)1 / InScale.Z;
+		}
+
+		return SafeReciprocalScale;
+	}
+
+
+
+
+	// vector-type-conversion variants. This allows applying a float transform to double vector and vice-versa.
+	// Whether this should be allowed is debatable. However in practice it is extremely rare to convert an
+	// entire float transform to a double transform in order to apply to a double vector, which is the only
+	// case where this conversion is an issue
+
+	template<typename RealType2, TEMPLATE_REQUIRES(std::is_same<RealType, RealType2>::value == false)>
+	FVector3<RealType2> TransformPosition(const UE::Math::TVector<RealType2>& P) const
+	{
+		return FVector3<RealType2>(TransformPosition(FVector3<RealType>(P)));
+	}
+	template<typename RealType2, TEMPLATE_REQUIRES(std::is_same<RealType, RealType2>::value == false)>
+	FVector3<RealType2> TransformVector(const UE::Math::TVector<RealType2>& V) const
+	{
+		return FVector3<RealType2>(TransformVector(FVector3<RealType>(V)));
+	}
+	template<typename RealType2, TEMPLATE_REQUIRES(std::is_same<RealType, RealType2>::value == false)>
+	FVector3<RealType2> TransformVectorNoScale(const UE::Math::TVector<RealType2>& V) const
+	{
+		return FVector3<RealType2>(TransformVectorNoScale(FVector3<RealType>(V)));
+	}
+	template<typename RealType2, TEMPLATE_REQUIRES(std::is_same<RealType, RealType2>::value == false)>
+	FVector3<RealType2> TransformNormal(const UE::Math::TVector<RealType2>& V) const
+	{
+		return FVector3<RealType2>(TransformNormal(FVector3<RealType>(V)));
+	}
+	template<typename RealType2, TEMPLATE_REQUIRES(std::is_same<RealType, RealType2>::value == false)>
+	FVector3<RealType2> InverseTransformPosition(const UE::Math::TVector<RealType2>& P) const
+	{
+		return FVector3<RealType2>(InverseTransformPosition(FVector3<RealType>(P)));
+	}
+	template<typename RealType2, TEMPLATE_REQUIRES(std::is_same<RealType, RealType2>::value == false)>
+	FVector3<RealType2> InverseTransformVector(const UE::Math::TVector<RealType2>& V) const
+	{
+		return FVector3<RealType2>(InverseTransformVector(FVector3<RealType>(V)));
+	}
+	template<typename RealType2, TEMPLATE_REQUIRES(std::is_same<RealType, RealType2>::value == false)>
+	FVector3<RealType2> InverseTransformVectorNoScale(const UE::Math::TVector<RealType2>& V) const
+	{
+		return FVector3<RealType2>(InverseTransformVectorNoScale(FVector3<RealType>(V)));
+	}
+	template<typename RealType2, TEMPLATE_REQUIRES(std::is_same<RealType, RealType2>::value == false)>
+	FVector3<RealType2> InverseTransformNormal(const UE::Math::TVector<RealType2>& V) const
+	{
+		return FVector3<RealType2>(InverseTransformNormal(FVector3<RealType>(V)));
+	}
+
+};
+typedef TTransform3<float> FTransform3f;
+typedef TTransform3<double> FTransform3d;
+
+
+} // end namespace UE::Geometry
+} // end namespace UE
\ No newline at end of file
diff --git a/Engine/Source/Runtime/GeometryCore/Public/TriangleTypes.h b/Engine/Source/Runtime/GeometryCore/Public/TriangleTypes.h
new file mode 100644
--- /dev/null
+++ b/Engine/Source/Runtime/GeometryCore/Public/TriangleTypes.h
@@ -0,0 +1,245 @@
+// Copyright Epic Games, Inc. All Rights Reserved.
+
+#pragma once
+
+#include "VectorTypes.h"
+#include "VectorUtil.h"
+
+namespace UE
+{
+namespace Geometry
+{
+
+
+/**
+ * Triangle utility functions
+ */
+namespace TriangleUtil
+{
+	/**
+	 * @return the edge length of an equilateral/regular triangle with the given area
+	 */
+	template<typename RealType>
+	RealType EquilateralEdgeLengthForArea(RealType TriArea)
+	{
+		return TMathUtil<RealType>::Sqrt(((RealType)4 * TriArea) / TMathUtil<RealType>::Sqrt3);
+	}
+
+};
+
+
+
+
+template<typename RealType>
+struct TTriangle2
+{
+	FVector2<RealType> V[3];
+
+	TTriangle2() {}
+
+	TTriangle2(const FVector2<RealType>& V0, const FVector2<RealType>& V1, const FVector2<RealType>& V2)
+	{
+		V[0] = V0;
+		V[1] = V1;
+		V[2] = V2;
+	}
+
+	TTriangle2(const FVector2<RealType> VIn[3])
+	{
+		V[0] = VIn[0];
+		V[1] = VIn[1];
+		V[2] = VIn[2];
+	}
+
+	FVector2<RealType> BarycentricPoint(RealType Bary0, RealType Bary1, RealType Bary2) const
+	{
+		return Bary0 * V[0] + Bary1 * V[1] + Bary2 * V[2];
+	}
+
+	FVector2<RealType> BarycentricPoint(const FVector3<RealType>& BaryCoords) const
+	{
+		return BaryCoords[0] * V[0] + BaryCoords[1] * V[1] + BaryCoords[2] * V[2];
+	}
+
+	FVector3<RealType> GetBarycentricCoords(const FVector2<RealType>& Point) const
+	{
+		return VectorUtil::BarycentricCoords(Point, V[0], V[1], V[2]);
+	}
+
+	/**
+	 * @param A first vertex of triangle
+	 * @param B second vertex of triangle
+	 * @param C third vertex of triangle
+	 * @return signed area of triangle
+	 */
+	static RealType SignedArea(const FVector2<RealType>& A, const FVector2<RealType>& B, const FVector2<RealType>& C)
+	{
+		return ((RealType)0.5) * ((A.X*B.Y - A.Y*B.X) + (B.X*C.Y - B.Y*C.X) + (C.X*A.Y - C.Y*A.X));
+	}
+
+	/** @return signed area of triangle */
+	RealType SignedArea() const
+	{
+		return SignedArea(V[0], V[1], V[2]);
+	}
+
+
+	/**
+	 * @param A first vertex of triangle
+	 * @param B second vertex of triangle
+	 * @param C third vertex of triangle
+	 * @param QueryPoint test point
+	 * @return true if QueryPoint is inside triangle
+	 */
+	static bool IsInside(const FVector2<RealType>& A, const FVector2<RealType>& B, const FVector2<RealType>& C, const FVector2<RealType>& QueryPoint)
+	{
+		RealType Sign1 = Orient(A, B, QueryPoint);
+		RealType Sign2 = Orient(B, C, QueryPoint);
+		RealType Sign3 = Orient(C, A, QueryPoint);
+		return (Sign1*Sign2 > 0) && (Sign2*Sign3 > 0) && (Sign3*Sign1 > 0);
+	}
+
+	/** @return true if QueryPoint is inside triangle */
+	bool IsInside(const FVector2<RealType>& QueryPoint) const
+	{
+		return IsInside(V[0], V[1], V[2], QueryPoint);
+	}
+
+
+	/** @return true if QueryPoint is inside triangle or on edge */
+	static bool IsInsideOrOn(const FVector2<RealType>& A, const FVector2<RealType>& B, const FVector2<RealType>& C, const FVector2<RealType>& QueryPoint)
+	{
+		RealType Sign1 = Orient(A, B, QueryPoint);
+		RealType Sign2 = Orient(B, C, QueryPoint);
+		RealType Sign3 = Orient(C, A, QueryPoint);
+		return (Sign1*Sign2 >= 0) && (Sign2*Sign3 >= 0) && (Sign3*Sign1 >= 0);
+	}
+
+	/** @return true if QueryPoint is inside triangle or on edge */
+	bool IsInsideOrOn(const FVector2<RealType>& QueryPoint, RealType Epsilon = 0) const
+	{
+		return IsInsideOrOn(V[0], V[1], V[2], QueryPoint);
+	}
+
+
+	/**
+	 * More robust (because it doesn't multiply orientation test results) inside-triangle test for oriented triangles only
+	 *  (the code early-outs at the first 'outside' edge, which only works if the triangle is oriented as expected)
+	 * @return 1 if outside, -1 if inside, 0 if on boundary
+	 */
+	int IsInsideOrOn_Oriented(const FVector2<RealType>& QueryPoint) const
+	{
+		return IsInsideOrOn_Oriented(V[0], V[1], V[2], QueryPoint);
+	}
+
+	/**
+	 * More robust (because it doesn't multiply orientation test results) inside-triangle test for oriented triangles only
+	 *  (the code early-outs at the first 'outside' edge, which only works if the triangle is oriented as expected)
+	 * @return 1 if outside, -1 if inside, 0 if on boundary
+	 */
+	static int IsInsideOrOn_Oriented(const FVector2<RealType>& A, const FVector2<RealType>& B, const FVector2<RealType>& C, const FVector2<RealType>& QueryPoint)
+	{
+		checkSlow(Orient(A, B, C) <= 0); // TODO: remove this checkSlow; it's just to make sure the orientation is as expected
+
+		RealType Sign1 = Orient(A, B, QueryPoint);
+		if (Sign1 > 0)
+		{
+			return 1;
+		}
+		
+		RealType Sign2 = Orient(B, C, QueryPoint);
+		if (Sign2 > 0)
+		{
+			return 1;
+		}
+		
+		RealType Sign3 = Orient(A, C, QueryPoint);
+		if (Sign3 < 0) // note this edge is queried backwards so the sign test is also backwards
+		{
+			return 1;
+		}
+
+		return (Sign1 != 0 && Sign2 != 0 && Sign3 != 0) ? -1 : 0;
+	}
+
+
+};
+
+typedef TTriangle2<float> FTriangle2f;
+typedef TTriangle2<double> FTriangle2d;
+typedef TTriangle2<int> FTriangle2i;
+
+
+
+
+
+
+template<typename RealType>
+struct TTriangle3
+{
+	FVector3<RealType> V[3];
+
+	TTriangle3() {}
+
+	TTriangle3(const FVector3<RealType>& V0, const FVector3<RealType>& V1, const FVector3<RealType>& V2)
+	{
+		V[0] = V0;
+		V[1] = V1;
+		V[2] = V2;
+	}
+
+	TTriangle3(const FVector3<RealType> VIn[3])
+	{
+		V[0] = VIn[0];
+		V[1] = VIn[1];
+		V[2] = VIn[2];
+	}
+
+	FVector3<RealType> BarycentricPoint(RealType Bary0, RealType Bary1, RealType Bary2) const
+	{
+		return Bary0*V[0] + Bary1*V[1] + Bary2*V[2];
+	}
+
+	FVector3<RealType> BarycentricPoint(const FVector3<RealType> & BaryCoords) const
+	{
+		return BaryCoords[0]*V[0] + BaryCoords[1]*V[1] + BaryCoords[2]*V[2];
+	}
+
+	FVector3<RealType> GetBarycentricCoords(const FVector3<RealType>& Point) const
+	{
+		return VectorUtil::BarycentricCoords(Point, V[0], V[1], V[2]);
+	}
+
+	/** @return vector that is perpendicular to the plane of this triangle */
+	FVector3<RealType> Normal() const
+	{
+		return VectorUtil::Normal(V[0], V[1], V[2]);
+	}
+
+	/** @return centroid of this triangle */
+	FVector3<RealType> Centroid() const
+	{
+		constexpr RealType f = 1.0 / 3.0;
+		return FVector3<RealType>(
+			(V[0].X + V[1].X + V[2].X) * f,
+			(V[0].Y + V[1].Y + V[2].Y) * f,
+			(V[0].Z + V[1].Z + V[2].Z) * f
+		);
+	}
+
+	/** grow the triangle around the centroid */
+	void Expand(RealType Delta)
+	{
+		FVector3<RealType> Centroid(Centroid());
+		V[0] += Delta * ((V[0] - Centroid).Normalized());
+		V[1] += Delta * ((V[1] - Centroid).Normalized());
+		V[2] += Delta * ((V[2] - Centroid).Normalized());
+	}
+};
+
+typedef TTriangle3<float> FTriangle3f;
+typedef TTriangle3<double> FTriangle3d;
+typedef TTriangle3<int> FTriangle3i;
+
+} // end namespace UE::Geometry
+} // end namespace UE
\ No newline at end of file
diff --git a/Engine/Source/Runtime/GeometryCore/Public/Util/CompactMaps.h b/Engine/Source/Runtime/GeometryCore/Public/Util/CompactMaps.h
new file mode 100644
--- /dev/null
+++ b/Engine/Source/Runtime/GeometryCore/Public/Util/CompactMaps.h
@@ -0,0 +1,157 @@
+// Copyright Epic Games, Inc. All Rights Reserved.
+
+// Index remapping struct extracted from Dynamic Mesh for more general use
+
+#pragma once
+
+#include "Containers/Array.h"
+
+#include "IndexTypes.h"
+
+namespace UE
+{
+namespace Geometry
+{
+
+/**
+ * Stores index remapping arrays for different types of elements
+ * Should only be used for compacting, and should maintain invariant that Map*[Idx] <= Idx for all maps
+ */
+struct FCompactMaps
+{
+	constexpr static int InvalidID = IndexConstants::InvalidID;
+
+	TArray<int> MapV;
+	TArray<int> MapT;
+	bool bKeepTriangleMap = true; // if set to false, MapT should be empty
+
+	/**
+	 * Set up maps as identity maps.
+	 * Note triangle map will not be set if bKeepTriangleMap is false.
+	 *
+	 * @param MaxVID Vertex map will be created from 0 to MaxVID
+	 * @param MaxTID If bKeepTriangleMap is true, triangle map will be created from 0 to MaxTID
+	 */
+	void SetIdentity(int MaxVID, int MaxTID)
+	{
+		MapV.SetNumUninitialized(MaxVID);
+		for (int ID = 0; ID < MaxVID; ID++)
+		{
+			MapV[ID] = ID;
+		}
+		if (bKeepTriangleMap)
+		{
+			MapT.SetNumUninitialized(MaxTID);
+			for (int ID = 0; ID < MaxTID; ID++)
+			{
+				MapT[ID] = ID;
+			}
+		}
+		else
+		{
+			MapT.Empty();
+		}
+	}
+
+	/**
+	 * Reset all maps, and initialize with InvalidID
+	 * @param MaxVID Size of post-reset vertex map (filled w/ InvalidID)
+	 * @param MaxTID Size of post-reset triangle map (filled w/ InvalidID)
+	 */
+	void Reset(int MaxVID, int MaxTID = 0)
+	{
+		MapV.Reset();
+		MapT.Reset();
+		MapV.SetNumUninitialized(MaxVID);
+		MapT.SetNumUninitialized(MaxTID);
+		for (int VID = 0; VID < MaxVID; VID++)
+		{
+			MapV[VID] = InvalidID;
+		}
+		for (int TID = 0; TID < MaxTID; TID++)
+		{
+			MapT[TID] = InvalidID;
+		}
+	}
+
+	/** Reset all maps, leaving them empty */
+	void Reset()
+	{
+		MapV.Reset();
+		MapT.Reset();
+	}
+
+	/** Set mapping for a vertex */
+	inline void SetVertex(int FromID, int ToID)
+	{
+		checkSlow(FromID >= ToID);
+		MapV[FromID] = ToID;
+	}
+
+	/** Set mapping for a triangle */
+	inline void SetTriangle(int FromID, int ToID)
+	{
+		checkSlow(bKeepTriangleMap);
+		checkSlow(FromID >= ToID);
+		MapT[FromID] = ToID;
+	}
+
+	/** Get mapping for a vertex */
+	inline int GetVertex(int FromID) const
+	{
+		int ToID = MapV[FromID];
+		checkSlow(ToID != InvalidID);
+		return ToID;
+	}
+
+	/** Get mapping for a triangle */
+	inline int GetTriangle(int FromID) const
+	{
+		checkSlow(bKeepTriangleMap);
+		int ToID = MapT[FromID];
+		checkSlow(ToID != InvalidID);
+		return ToID;
+	}
+
+	/** 
+	 * Clear all triangles
+	 *
+	 * @param bPermanent If true, set flag to not create triangle map
+	 */
+	void ClearTriangleMap(bool bPermanent = true)
+	{
+		if (bPermanent)
+		{
+			bKeepTriangleMap = false;
+		}
+		MapT.Empty();
+	}
+
+	/** Check data for validity; for testing */
+	bool Validate() const
+	{
+		for (int Idx = 0; Idx < MapV.Num(); Idx++)
+		{
+			if (MapV[Idx] > Idx)
+			{
+				return false;
+			}
+		}
+		for (int Idx = 0; Idx < MapT.Num(); Idx++)
+		{
+			if (MapT[Idx] > Idx)
+			{
+				return false;
+			}
+		}
+		if (!bKeepTriangleMap && MapT.Num() > 0)
+		{
+			return false;
+		}
+		return true;
+	}
+};
+
+
+} // end namespace UE::Geometry
+} // end namespace UE
\ No newline at end of file
diff --git a/Engine/Source/Runtime/GeometryCore/Public/Util/DynamicVector.h b/Engine/Source/Runtime/GeometryCore/Public/Util/DynamicVector.h
new file mode 100644
--- /dev/null
+++ b/Engine/Source/Runtime/GeometryCore/Public/Util/DynamicVector.h
@@ -0,0 +1,719 @@
+// Copyright Epic Games, Inc. All Rights Reserved.
+
+// Port of geometry3cpp dvector
+
+#pragma once
+
+#include <CoreMinimal.h>
+#include "Containers/StaticArray.h"
+#include "Serialization/Archive.h"
+#include "VectorTypes.h"
+#include "IndexTypes.h"
+
+namespace UE
+{
+namespace Geometry
+{
+
+/*
+ * Blocked array with fixed, power-of-two sized blocks.
+ *
+ * Iterator functions suitable for use with range-based for are provided
+ */
+template <class Type>
+class TDynamicVector
+{
+public:
+	TDynamicVector()
+	{
+		Blocks.Add(new BlockType());
+	}
+
+	TDynamicVector(const TDynamicVector& Copy) = default;
+	TDynamicVector(TDynamicVector&& Moved)
+	{
+		Blocks = MoveTemp(Moved.Blocks);
+		CurBlock = MoveTemp(Moved.CurBlock);
+		CurBlockUsed = MoveTemp(Moved.CurBlockUsed);
+		Moved.Blocks.Add(new BlockType());
+		Moved.CurBlock = 0;
+		Moved.CurBlockUsed = 0;
+	}
+	TDynamicVector& operator=(const TDynamicVector& Copy) = default;
+	TDynamicVector& operator=(TDynamicVector&& Moved)
+	{
+		Blocks = MoveTemp(Moved.Blocks);
+		CurBlock = MoveTemp(Moved.CurBlock);
+		CurBlockUsed = MoveTemp(Moved.CurBlockUsed);
+		Moved.Blocks.Add(new BlockType());
+		Moved.CurBlock = 0;
+		Moved.CurBlockUsed = 0;
+		return *this;
+	}
+
+	TDynamicVector(const TArray<Type>& Array)
+	{
+		SetNum(Array.Num());
+		for (int Idx = 0; Idx < Array.Num(); Idx++)
+		{
+			(*this)[Idx] = Array[Idx];
+		}
+	}
+	TDynamicVector(TArrayView<const Type> Array)
+	{
+		SetNum(Array.Num());
+		for (int Idx = 0; Idx < Array.Num(); Idx++)
+		{
+			(*this)[Idx] = Array[Idx];
+		}
+	}
+
+	inline void Clear();
+	inline void Fill(const Type& Value);
+	inline void Resize(size_t Count);
+	inline void Resize(size_t Count, const Type& InitValue);
+	inline void SetNum(size_t Count) { Resize(Count); }
+
+	inline bool IsEmpty() const { return CurBlock == 0 && CurBlockUsed == 0; }
+	inline size_t GetLength() const { return CurBlock * BlockSize + CurBlockUsed; }
+	inline size_t Num() const { return CurBlock * BlockSize + CurBlockUsed; }
+	inline int GetBlockSize() const { return BlockSize; }
+	inline size_t GetByteCount() const { return Blocks.Num() * BlockSize * sizeof(Type); }
+
+	inline void Add(const Type& Data);
+	inline void Add(const TDynamicVector& Data);
+	inline void PopBack();
+
+	inline void InsertAt(const Type& Data, unsigned int Index);
+	inline Type& ElementAt(unsigned int Index, Type InitialValue = Type{});
+
+	inline const Type& Front() const { return Blocks[0][0]; }
+	inline const Type& Back() const { return Blocks[CurBlock][CurBlockUsed - 1]; }
+
+	inline Type& operator[](unsigned int Index)
+	{
+		return Blocks[Index >> nShiftBits][Index & BlockIndexBitmask];
+	}
+	inline const Type& operator[](unsigned int Index) const
+	{
+		return Blocks[Index >> nShiftBits][Index & BlockIndexBitmask];
+	}
+
+	// apply f() to each member sequentially
+	template <typename Func>
+	void Apply(const Func& f);
+
+	/**
+	 * Serialization operator for TDynamicVector.
+	 *
+	 * @param Ar Archive to serialize with.
+	 * @param Vec Vector to serialize.
+	 * @returns Passing down serializing archive.
+	 */
+	friend FArchive& operator<<(FArchive& Ar, TDynamicVector& Vec)
+	{
+		Vec.Serialize(Ar);
+		return Ar;
+	}
+
+	/** Serialize vector to an archive */
+	void Serialize(FArchive& Ar)
+	{
+		Ar << CurBlock;
+		Ar << CurBlockUsed;
+		Ar << Blocks;
+	}
+
+public:
+	/*
+	 * FIterator class iterates over values of vector
+	 */
+	class FIterator
+	{
+	public:
+		inline const Type& operator*() const
+		{
+			return (*DVector)[Idx];
+		}
+		inline Type& operator*()
+		{
+			return (*DVector)[Idx];
+		}
+		inline FIterator& operator++()   // prefix
+		{
+			Idx++;
+			return *this;
+		}
+		inline FIterator operator++(int) // postfix
+		{
+			FIterator Copy(*this);
+			Idx++;
+			return Copy;
+		}
+		inline bool operator==(const FIterator& Itr2)
+		{
+			return DVector == Itr2.DVector && Idx == Itr2.Idx;
+		}
+		inline bool operator!=(const FIterator& Itr2)
+		{
+			return DVector != Itr2.DVector || Idx != Itr2.Idx;
+		}
+
+	private:
+		friend class TDynamicVector;
+		FIterator(TDynamicVector* DVectorIn, int IdxIn)
+			: DVector(DVectorIn), Idx(IdxIn){}
+		TDynamicVector* DVector{};
+		int Idx{0};
+	};
+
+	/** @return iterator at beginning of vector */
+	FIterator begin()
+	{
+		return IsEmpty() ? end() : FIterator{this, 0};
+	}
+	/** @return iterator at end of vector */
+	FIterator end()
+	{
+		return FIterator{this, (int)GetLength()};
+	}
+
+	/*
+	 * FConstIterator class iterates over values of vector
+	 */
+	class FConstIterator
+	{
+	public:
+		inline const Type& operator*() const
+		{
+			return (*DVector)[Idx];
+		}
+		inline FConstIterator& operator++()   // prefix
+		{
+			Idx++;
+			return *this;
+		}
+		inline FConstIterator operator++(int) // postfix
+		{
+			FConstIterator Copy(*this);
+			Idx++;
+			return Copy;
+		}
+		inline bool operator==(const FConstIterator& Itr2)
+		{
+			return DVector == Itr2.DVector && Idx == Itr2.Idx;
+		}
+		inline bool operator!=(const FConstIterator& Itr2)
+		{
+			return DVector != Itr2.DVector || Idx != Itr2.Idx;
+		}
+
+	private:
+		friend class TDynamicVector;
+		FConstIterator(const TDynamicVector* DVectorIn, int IdxIn)
+			: DVector(DVectorIn), Idx(IdxIn){}
+		const TDynamicVector* DVector{};
+		int Idx{0};
+	};
+
+	/** @return iterator at beginning of vector */
+	FConstIterator begin() const
+	{
+		return IsEmpty() ? end() : FConstIterator{this, 0};
+	}
+	/** @return iterator at end of vector */
+	FConstIterator end() const
+	{
+		return FConstIterator{this, (int)GetLength()};
+	}
+
+private:
+	static constexpr int nShiftBits = 9;
+	static constexpr int BlockSize = 1 << nShiftBits;
+	static constexpr int BlockIndexBitmask = BlockSize - 1; // low 9 bits
+	static_assert( BlockSize && ((BlockSize & (BlockSize - 1)) == 0), "DynamicVector: BlockSize must be a power of two");
+
+	unsigned int CurBlock{0};
+	unsigned int CurBlockUsed{0};
+
+	using BlockType = TStaticArray<Type, BlockSize>;
+	
+
+	/**
+	 * This is used to store/manage an array of BlockType*. Basically a stripped
+	 * down TIndirectArray that does not do range-checking on operator[] in non-debug builds
+	 */
+	template<typename ArrayType>
+	class TBlockVector
+	{
+	protected:
+		TArray<ArrayType*> Elements;
+
+	public:
+		TBlockVector() = default;
+		TBlockVector(TBlockVector&& Moved) = default;
+
+		TBlockVector(const TBlockVector& Copy)
+		{
+			int32 N = Copy.Num();
+			Elements.Reserve(N);
+			for (int32 k = 0; k < N; ++k)
+			{
+				Elements.Add(new ArrayType(*Copy.Elements[k]));
+			}
+		}
+		TBlockVector& operator=(const TBlockVector& Copy)
+		{
+			if (&Copy != this)
+			{
+				int32 N = Copy.Num();
+				Empty(N);
+				Elements.Reserve(N);
+				for (int32 k = 0; k < N; ++k)
+				{
+					Elements.Add(new ArrayType(*Copy.Elements[k]));
+				}
+			}
+			return *this;
+		}
+		TBlockVector& operator=(TBlockVector&& Moved)
+		{
+			if (&Moved != this)
+			{
+				Empty();
+				Elements = MoveTemp(Moved.Elements);
+			}
+			return *this;
+		}
+		~TBlockVector()
+		{
+			Empty();
+		}
+
+		int32 Num() const { return Elements.Num(); }
+
+#if UE_BUILD_DEBUG
+		ArrayType& operator[](int32 Index) { return *Elements[Index]; }
+		const ArrayType& operator[](int32 Index) const { return *Elements[Index]; }
+#else
+		ArrayType& operator[](int32 Index) { return *Elements.GetData()[Index]; }
+		const ArrayType& operator[](int32 Index) const { return *Elements.GetData()[Index]; }
+#endif
+
+		void Add(ArrayType* NewElement)
+		{
+			Elements.Add(NewElement);
+		}
+
+		void Truncate(int32 NewElementCount, bool bAllowShrinking)
+		{
+			for (int32 k = NewElementCount; k < Elements.Num(); ++k)
+			{
+				delete Elements[k];
+				Elements[k] = nullptr;
+			}
+			Elements.RemoveAt(NewElementCount, Elements.Num() - NewElementCount, bAllowShrinking);
+		}
+
+		void Empty(int32 NewReservedSize = 0)
+		{
+			for (int32 k = 0, N = Elements.Num(); k < N; ++k)
+			{
+				delete Elements[k];
+			}
+			Elements.Empty(NewReservedSize);
+		}
+
+		/**
+		 * Count bytes needed to serialize this array.
+		 *
+		 * @param Ar Archive to count for.
+		 */
+		void CountBytes(FArchive& Ar) const
+		{
+			Elements.CountBytes(Ar);
+		}
+
+		/**
+		 * Serialization operator for TBlockVector.
+		 *
+		 * @param Ar Archive to serialize with.
+		 * @param Vec Vector to serialize.
+		 * @returns Passing down serializing archive.
+		 */
+		friend FArchive& operator<<(FArchive& Ar, TBlockVector& Vec)
+		{
+			Vec.Serialize(Ar);
+			return Ar;
+		}
+
+		/** Serialize TBlockVector to archive. */
+		void Serialize(FArchive& Ar)
+		{
+			CountBytes(Ar);
+			if (Ar.IsLoading())
+			{
+				// Load array.
+				int32 NewNum;
+				Ar << NewNum;
+				Empty(NewNum);
+				for (int32 Index = 0; Index < NewNum; Index++)
+				{
+					ArrayType* NewElement = new ArrayType;
+					// TODO: Consider moving bulk serialization code into TStaticArray.
+					if (TCanBulkSerialize<Type>::Value)
+					{
+						Ar.Serialize(NewElement->GetData(), NewElement->Num() * sizeof(Type));
+					}
+					else
+					{
+						Ar << *NewElement;
+					}
+					Add(NewElement);
+				}
+			}
+			else
+			{
+				// Save array.
+				int32 CurNum = Num();
+				Ar << CurNum;
+				for (int32 Index = 0; Index < CurNum; Index++)
+				{
+					if (TCanBulkSerialize<Type>::Value)
+					{
+						Ar.Serialize((*this)[Index].GetData(), (*this)[Index].Num() * sizeof(Type));
+					}
+					else
+					{
+						Ar << (*this)[Index];
+					}
+				}
+			}
+		}
+	};
+
+	// memory chunks for dynamic vector
+	TBlockVector<BlockType> Blocks;
+};
+
+template <class Type, int N>
+class TDynamicVectorN
+{
+public:
+	TDynamicVectorN() = default;
+	TDynamicVectorN(const TDynamicVectorN& Copy) = default;
+	TDynamicVectorN(TDynamicVectorN&& Moved) = default;
+	TDynamicVectorN& operator=(const TDynamicVectorN& Copy) = default;
+	TDynamicVectorN& operator=(TDynamicVectorN&& Moved) = default;
+
+	inline void Clear()
+	{
+		Data.Clear();
+	}
+	inline void Fill(const Type& Value)
+	{
+		Data.Fill(Value);
+	}
+	inline void Resize(size_t Count)
+	{
+		Data.Resize(Count * N);
+	}
+	inline void Resize(size_t Count, const Type& InitValue)
+	{
+		Data.Resize(Count * N, InitValue);
+	}
+	inline bool IsEmpty() const
+	{
+		return Data.IsEmpty();
+	}
+	inline size_t GetLength() const
+	{
+		return Data.GetLength() / N;
+	}
+	inline int GetBlockSize() const
+	{
+		return Data.GetBlockSize();
+	}
+	inline size_t GetByteCount() const
+	{
+		return Data.GetByteCount();
+	}
+
+	// simple struct to help pass N-dimensional data without presuming a vector type (e.g. just via initializer list)
+	struct ElementVectorN
+	{
+		Type Data[N];
+	};
+
+	inline void Add(const ElementVectorN& AddData)
+	{
+		// todo specialize for N=2,3,4
+		for (int i = 0; i < N; i++)
+		{
+			Data.Add(AddData.Data[i]);
+		}
+	}
+
+	inline void PopBack()
+	{
+		for (int i = 0; i < N; i++)
+		{
+			PopBack();
+		}
+	} // TODO specialize
+
+	inline void InsertAt(const ElementVectorN& AddData, unsigned int Index)
+	{
+		for (int i = 1; i <= N; i++)
+		{
+			Data.InsertAt(AddData.Data[N - i], N * (Index + 1) - i);
+		}
+	}
+
+	inline Type& operator()(unsigned int TopIndex, unsigned int SubIndex)
+	{
+		return Data[TopIndex * N + SubIndex];
+	}
+	inline const Type& operator()(unsigned int TopIndex, unsigned int SubIndex) const
+	{
+		return Data[TopIndex * N + SubIndex];
+	}
+	inline void SetVector2(unsigned int TopIndex, const FVector2<Type>& V)
+	{
+		check(N >= 2);
+		unsigned int i = TopIndex * N;
+		Data[i] = V.X;
+		Data[i + 1] = V.Y;
+	}
+	inline void SetVector3(unsigned int TopIndex, const FVector3<Type>& V)
+	{
+		check(N >= 3);
+		unsigned int i = TopIndex * N;
+		Data[i] = V.X;
+		Data[i + 1] = V.Y;
+		Data[i + 2] = V.Z;
+	}
+	inline FVector2<Type> AsVector2(unsigned int TopIndex) const
+	{
+		check(N >= 2);
+		return FVector2<Type>(
+			Data[TopIndex * N + 0],
+			Data[TopIndex * N + 1]);
+	}
+	inline FVector3<Type> AsVector3(unsigned int TopIndex) const
+	{
+		check(N >= 3);
+		return FVector3<Type>(
+			Data[TopIndex * N + 0],
+			Data[TopIndex * N + 1],
+			Data[TopIndex * N + 2]);
+	}
+	inline FIndex2i AsIndex2(unsigned int TopIndex) const
+	{
+		check(N >= 2);
+		return FIndex2i(
+			Data[TopIndex * N + 0],
+			Data[TopIndex * N + 1]);
+	}
+	inline FIndex3i AsIndex3(unsigned int TopIndex) const
+	{
+		check(N >= 3);
+		return FIndex3i(
+			Data[TopIndex * N + 0],
+			Data[TopIndex * N + 1],
+			Data[TopIndex * N + 2]);
+	}
+	inline FIndex4i AsIndex4(unsigned int TopIndex) const
+	{
+		check(N >= 4);
+		return FIndex4i(
+			Data[TopIndex * N + 0],
+			Data[TopIndex * N + 1],
+			Data[TopIndex * N + 2],
+			Data[TopIndex * N + 3]);
+	}
+
+private:
+	TDynamicVector<Type> Data;
+
+	friend class FIterator;
+};
+
+template class TDynamicVectorN<double, 2>;
+
+using  TDynamicVector3f = TDynamicVectorN<float, 3>;
+using  TDynamicVector2f = TDynamicVectorN<float, 2>;
+using  TDynamicVector3d = TDynamicVectorN<double, 3>;
+using  TDynamicVector2d = TDynamicVectorN<double, 2>;
+using  TDynamicVector3i = TDynamicVectorN<int, 3>;
+using  TDynamicVector2i = TDynamicVectorN<int, 2>;
+
+template <class Type>
+void TDynamicVector<Type>::Clear()
+{
+	Blocks.Truncate(1, false);
+	CurBlock = 0;
+	CurBlockUsed = 0;
+	if (Blocks.Num() == 0)
+	{
+		Blocks.Add(new BlockType());
+	}
+}
+
+template <class Type>
+void TDynamicVector<Type>::Fill(const Type& Value)
+{
+	size_t Count = Blocks.Num();
+	for (unsigned int i = 0; i < Count; ++i)
+	{
+		for (uint32 ElementIndex = 0; ElementIndex < BlockSize; ++ElementIndex)
+		{
+			Blocks[i][ElementIndex] = Value;
+		}
+	}
+}
+
+template <class Type>
+void TDynamicVector<Type>::Resize(size_t Count)
+{
+	if (GetLength() == Count)
+	{
+		return;
+	}
+
+	// figure out how many segments we need
+	int nNumSegs = 1 + (int)Count / BlockSize;
+
+	// figure out how many are currently allocated...
+	int32 nCurCount = Blocks.Num();
+
+	// resize to right number of segments
+	if (nNumSegs >= nCurCount)
+	{
+		// allocate new segments
+		for (int i = (int)nCurCount; i < nNumSegs; ++i)
+		{
+			Blocks.Add(new BlockType());
+		}
+	}
+	else
+	{
+		Blocks.Truncate(nNumSegs, false);
+	}
+
+	// mark last segment
+	CurBlockUsed = (unsigned int)(Count - (nNumSegs - 1) * BlockSize);
+	CurBlock = nNumSegs - 1;
+}
+
+template <class Type>
+void TDynamicVector<Type>::Resize(size_t Count, const Type& InitValue)
+{
+	size_t nCurSize = GetLength();
+	Resize(Count);
+	for (size_t Index = nCurSize; Index < Count; ++Index)
+	{
+		this->operator[](Index) = InitValue;
+	}
+}
+
+template <class Type>
+void TDynamicVector<Type>::Add(const Type& Value)
+{
+	if (CurBlockUsed == BlockSize)
+	{
+		if (CurBlock == Blocks.Num() - 1)
+		{
+			Blocks.Add(new BlockType());
+		}
+		CurBlock++;
+		CurBlockUsed = 0;
+	}
+	Blocks[CurBlock][CurBlockUsed] = Value;
+	CurBlockUsed++;
+}
+
+
+template <class Type>
+void TDynamicVector<Type>::Add(const TDynamicVector<Type>& AddData)
+{
+	// @todo it could be more efficient to use memcopies here...
+	size_t nSize = AddData.Num();
+	for (unsigned int k = 0; k < nSize; ++k)
+	{
+		Add(AddData[k]);
+	}
+}
+
+template <class Type>
+void TDynamicVector<Type>::PopBack()
+{
+	if (CurBlockUsed > 0)
+	{
+		CurBlockUsed--;
+	}
+	if (CurBlockUsed == 0 && CurBlock > 0)
+	{
+		CurBlock--;
+		CurBlockUsed = BlockSize;
+		// remove block ??
+	}
+}
+
+template <class Type>
+Type& TDynamicVector<Type>::ElementAt(unsigned int Index, Type InitialValue)
+{
+	size_t s = GetLength();
+	if (Index == s)
+	{
+		Add(InitialValue);
+	}
+	else if (Index > s)
+	{
+		Resize(Index);
+		Add(InitialValue);
+	}
+	return this->operator[](Index);
+}
+
+template <class Type>
+void TDynamicVector<Type>::InsertAt(const Type& AddData, unsigned int Index)
+{
+	size_t s = GetLength();
+	if (Index == s)
+	{
+		Add(AddData);
+	}
+	else if (Index > s)
+	{
+		Resize(Index);
+		Add(AddData);
+	}
+	else
+	{
+		(*this)[Index] = AddData;
+	}
+}
+
+template <typename Type>
+template <typename Func>
+void TDynamicVector<Type>::Apply(const Func& applyF)
+{
+	for (int bi = 0; bi < CurBlock; ++bi)
+	{
+		auto block = Blocks[bi];
+		for (int k = 0; k < BlockSize; ++k)
+		{
+			applyF(block[k], k);
+		}
+	}
+	auto lastblock = Blocks[CurBlock];
+	for (int k = 0; k < CurBlockUsed; ++k)
+	{
+		applyF(lastblock[k], k);
+	}
+}
+
+
+} // end namespace UE::Geometry
+} // end namespace UE
\ No newline at end of file
