GeometryCore: create new GeometryCore Module. Move GeometricObjects code from GeometryProcessing plugin to this new module, as well as core FDynamicMesh3 classes (attributes, aabbtree/octree, change t

diff --git a/Engine/Source/Runtime/GeometryCore/Public/VectorTypes.h b/Engine/Source/Runtime/GeometryCore/Public/VectorTypes.h
new file mode 100644
--- /dev/null
+++ b/Engine/Source/Runtime/GeometryCore/Public/VectorTypes.h
@@ -0,0 +1,1220 @@
+// Copyright Epic Games, Inc. All Rights Reserved.
+
+#pragma once
+
+#include "Math/Vector.h"
+#include "MathUtil.h"
+#include "Serialization/Archive.h"
+#include "Templates/UnrealTypeTraits.h"
+#include <sstream>
+
+
+namespace UE {
+namespace Geometry {
+
+
+/**
+* Templated 2D Vector. Ported from g3Sharp library, with the intention of
+* maintaining compatibility with existing g3Sharp code. Has an API
+* similar to WildMagic, GTEngine, Eigen, etc.
+*
+* Convenience typedefs for FVector2f/FVector2d are defined, and
+* should be preferentially used over the base template type
+*
+* @todo Possibly can be replaced/merged with Chaos TVector<T,N>
+*/
+template <typename T>
+struct FVector2
+{
+	T X{}, Y{};
+
+	constexpr FVector2()
+		: X(0), Y(0)
+	{
+	}
+
+	constexpr FVector2(T ValX, T ValY)
+		: X(ValX), Y(ValY)
+	{
+	}
+
+	constexpr FVector2(const T* Data)
+		: X(Data[0]), Y(Data[1])
+	{
+	}
+
+	constexpr FVector2(const FVector2& Vec) = default;
+
+	constexpr FVector2(const FVector2D& Vec)
+		: X((T)Vec.X), Y((T)Vec.Y)
+	{
+	}
+
+	template<typename RealType2>
+	explicit constexpr FVector2(const FVector2<RealType2>& Vec)
+		: X((T)Vec.X), Y((T)Vec.Y)
+	{
+	}
+
+	explicit constexpr operator const T*() const
+	{
+		return &X;
+	};
+
+	explicit constexpr operator T*()
+	{
+		return &X;
+	}
+
+	explicit operator FVector2D() const
+	{
+		return FVector2D((float)X, (float)Y);
+	}
+
+	constexpr static FVector2<T> Zero()
+	{
+		return FVector2<T>((T)0, (T)0);
+	}
+	constexpr static FVector2<T> One()
+	{
+		return FVector2<T>((T)1, (T)1);
+	}
+	constexpr static FVector2<T> UnitX()
+	{
+		return FVector2<T>((T)1, (T)0);
+	}
+	constexpr static FVector2<T> UnitY()
+	{
+		return FVector2<T>((T)0, (T)1);
+	}
+
+	constexpr FVector2<T>& operator=(const FVector2<T>& V2)
+	{
+		X = V2.X;
+		Y = V2.Y;
+		return *this;
+	}
+
+	constexpr T& operator[](int Idx)
+	{
+		return (&X)[Idx];
+	}
+	constexpr const T& operator[](int Idx) const
+	{
+		return (&X)[Idx];
+	}
+
+	constexpr T Length() const
+	{
+		return TMathUtil<T>::Sqrt(X * X + Y * Y);
+	}
+	T SquaredLength() const
+	{
+		return X * X + Y * Y;
+	}
+
+	constexpr T Distance(const FVector2<T>& V2) const
+	{
+		T dx = V2.X - X;
+		T dy = V2.Y - Y;
+		return TMathUtil<T>::Sqrt(dx * dx + dy * dy);
+	}
+	constexpr T DistanceSquared(const FVector2<T>& V2) const
+	{
+		T dx = V2.X - X;
+		T dy = V2.Y - Y;
+		return dx * dx + dy * dy;
+	}
+
+	constexpr T Dot(const FVector2<T>& V2) const
+	{
+		return X * V2.X + Y * V2.Y;
+	}
+
+	constexpr  FVector2 operator-() const
+	{
+		return FVector2(-X, -Y);
+	}
+
+	constexpr FVector2 operator+(const FVector2& V2) const
+	{
+		return FVector2(X + V2.X, Y + V2.Y);
+	}
+
+	constexpr FVector2 operator-(const FVector2& V2) const
+	{
+		return FVector2(X - V2.X, Y - V2.Y);
+	}
+
+	constexpr FVector2<T> operator+(const T& Scalar) const
+	{
+		return FVector2<T>(X + Scalar, Y + Scalar);
+	}
+
+	constexpr FVector2<T> operator-(const T& Scalar) const
+	{
+		return FVector2<T>(X - Scalar, Y - Scalar);
+	}
+
+	constexpr FVector2<T> operator*(const T& Scalar) const
+	{
+		return FVector2<T>(X * Scalar, Y * Scalar);
+	}
+
+	template<typename RealType2>
+	constexpr FVector2<T> operator*(const RealType2& Scalar) const
+	{
+		return FVector2<T>(X * (T)Scalar, Y * (T)Scalar);
+	}
+
+	constexpr FVector2<T> operator*(const FVector2<T>& V2) const // component-wise
+	{
+		return FVector2<T>(X * V2.X, Y * V2.Y);
+	}
+
+	constexpr FVector2<T> operator/(const T& Scalar) const
+	{
+		return FVector2<T>(X / Scalar, Y / Scalar);
+	}
+
+	constexpr FVector2<T> operator/(const FVector2<T>& V2) const // component-wise
+	{
+		return FVector2<T>(X / V2.X, Y / V2.Y);
+	}
+
+	constexpr FVector2<T>& operator+=(const FVector2<T>& V2)
+	{
+		X += V2.X;
+		Y += V2.Y;
+		return *this;
+	}
+
+	constexpr FVector2<T>& operator-=(const FVector2<T>& V2)
+	{
+		X -= V2.X;
+		Y -= V2.Y;
+		return *this;
+	}
+
+	constexpr FVector2<T>& operator*=(const T& Scalar)
+	{
+		X *= Scalar;
+		Y *= Scalar;
+		return *this;
+	}
+
+	constexpr FVector2<T>& operator/=(const T& Scalar)
+	{
+		X /= Scalar;
+		Y /= Scalar;
+		return *this;
+	}
+
+	constexpr bool operator==(const FVector2<T>& Other) const
+	{
+		return X == Other.X && Y == Other.Y;
+	}
+
+	constexpr bool operator!=(const FVector2<T>& Other) const
+	{
+		return X != Other.X || Y != Other.Y;
+	}
+
+	/**
+	 * Serialization operator for FVector2.
+	 *
+	 * @param Ar Archive to serialize with.
+	 * @param Vec Vector to serialize.
+	 * @returns Passing down serializing archive.
+	 */
+	friend FArchive& operator<<(FArchive& Ar, FVector2& Vec)
+	{
+		Vec.Serialize(Ar);
+		return Ar;
+	}
+
+	/** Serialize FVector2 to an archive. */
+	void Serialize(FArchive& Ar)
+	{
+		Ar << X;
+		Ar << Y;
+	}
+};
+
+/** @return dot product of V1 with PerpCW(V2), ie V2 rotated 90 degrees clockwise */
+template <typename T>
+constexpr T DotPerp(const FVector2<T>& V1, const FVector2<T>& V2)
+{
+	return V1.X * V2.Y - V1.Y * V2.X;
+}
+
+/** @return right-Perpendicular vector to V, ie V rotated 90 degrees clockwise */
+template <typename T>
+constexpr FVector2<T> PerpCW(const FVector2<T>& V)
+{
+	return FVector2<T>(V.Y, -V.X);
+}
+
+/** @return > 0 if C is to the left of the line from A to B, < 0 if to the right, 0 if on the line */
+template<typename T>
+T Orient(const FVector2<T>& A, const FVector2<T>& B, const FVector2<T>& C)
+{
+	return DotPerp((B - A), (C - A));
+}
+
+
+template <typename T>
+constexpr bool IsNormalized(const FVector2<T>& Vector, const T Tolerance = TMathUtil<T>::ZeroTolerance)
+{
+	return TMathUtil<T>::Abs((Vector.X*Vector.X + Vector.Y*Vector.Y) - 1) < Tolerance;
+}
+
+template <typename T>
+T Normalize(FVector2<T>& Vector, const T Epsilon = 0)
+{
+	 T length = Vector.Length();
+	 if (length > Epsilon)
+	 {
+		 T invLength = ((T)1) / length;
+		 Vector.X *= invLength;
+		 Vector.Y *= invLength;
+		 return length;
+	 }
+	 Vector.X = Vector.Y = (T)0;
+	 return (T)0;
+}
+
+template <typename T>
+FVector2<T> Normalized(const FVector2<T>& Vector, const T Epsilon = 0)
+{
+	T length = Vector.Length();
+	if (length > Epsilon)
+	{
+		T invLength = ((T)1) / length;
+		return FVector2<T>(Vector.X*invLength, Vector.Y*invLength);
+	}
+	return FVector2<T>((T)0, (T)0);
+}
+
+
+// Angle in Degrees
+template <typename T>
+T AngleD(const FVector2<T>& V1, const FVector2<T>& V2)
+{
+	T DotVal = V1.Dot(V2);
+	T ClampedDot = (DotVal < (T)-1) ? (T)-1 : ((DotVal > (T)1) ? (T)1 : DotVal);
+	return TMathUtil<T>::ACos(ClampedDot) * TMathUtil<T>::RadToDeg;
+}
+
+// Angle in Radians
+template <typename T>
+T AngleR(const FVector2<T>& V1, const FVector2<T>& V2)
+{
+	T DotVal = V1.Dot(V2);
+	T ClampedDot = (DotVal < (T)-1) ? (T)-1 : ((DotVal > (T)1) ? (T)1 : DotVal);
+	return TMathUtil<T>::ACos(ClampedDot);
+}
+
+// Angle in Radians
+template <typename T>
+T SignedAngleR(const FVector2<T>& V1, const FVector2<T>& V2)
+{
+	T DotVal = V1.Dot(V2);
+	T ClampedDot = (DotVal < (T)-1) ? (T)-1 : ((DotVal > (T)1) ? (T)1 : DotVal);
+	T Direction = DotPerp(V1, V2);
+	if (Direction * Direction < TMathUtil<T>::ZeroTolerance)
+	{
+		return (DotVal < 0) ? TMathUtil<T>::Pi : (T)0;
+	}
+	else
+	{
+		T Sign = Direction < 0 ? (T)-1 : (T)1;
+		return Sign * TMathUtil<T>::ACos(ClampedDot);
+	}
+}
+
+template <typename T>
+FVector2<T> Lerp(const FVector2<T>& A, const FVector2<T>& B, T Alpha)
+{
+	T OneMinusAlpha = (T)1 - Alpha;
+	return FVector2<T>(OneMinusAlpha * A.X + Alpha * B.X,
+		OneMinusAlpha * A.Y + Alpha * B.Y);
+}
+
+
+/**
+ * Templated 3D Vector. Ported from g3Sharp library, with the intention of 
+ * maintaining compatibility with existing g3Sharp code. Has an API
+ * similar to WildMagic, GTEngine, Eigen, etc. 
+ * 
+ * Convenience typedefs for FVector3f/FVector3d are defined, and
+ * should be preferentially used over the base template type
+ * 
+ * @todo Possibly can be replaced/merged with Chaos TVector<T,N>
+ */
+template <typename T>
+struct FVector3 : public UE::Math::TVector<T>
+{
+	using UE::Math::TVector<T>::TVector;
+	using UE::Math::TVector<T>::X;
+	using UE::Math::TVector<T>::Y;
+	using UE::Math::TVector<T>::Z;
+
+	FVector3() : UE::Math::TVector<T>((T)0)
+	{
+	}
+
+	FVector3(const UE::Math::TVector<T>& Vec)
+	{
+		X = Vec.X;
+		Y = Vec.Y;
+		Z = Vec.Z;
+	}
+
+	template<typename OtherType, TEMPLATE_REQUIRES(!std::is_same<T, OtherType>::value)>
+	explicit FVector3(const UE::Math::TVector<OtherType>& Vec)
+	{
+		X = (T)Vec.X;
+		Y = (T)Vec.Y;
+		Z = (T)Vec.Z;
+	}
+
+	FVector3& operator=(const UE::Math::TVector<T>& Vec)
+	{
+		X = Vec.X;
+		Y = Vec.Y;
+		Z = Vec.Z;
+		return *this;
+	}
+
+	explicit operator FVector3f() const
+	{
+		return FVector3f((float)X, (float)Y, (float)Z);
+	}
+
+	explicit operator FVector3d() const
+	{
+		return FVector3d((double)X, (double)Y, (double)Z);
+	}
+
+
+	explicit operator FLinearColor() const
+	{
+		return FLinearColor((float)X, (float)Y, (float)Z);
+	}
+	FVector3(const FLinearColor& Color)
+	{
+		X = (T)Color.R;
+		Y = (T)Color.G;
+		Z = (T)Color.B;
+	}
+
+
+	explicit FVector3(const T* Data)
+	{
+		X = Data[0];
+		Y = Data[1];
+		Z = Data[2];
+	}
+
+
+	explicit operator const T*() const
+	{
+		return &X;
+	};
+	explicit operator T*()
+	{
+		return &X;
+	}
+
+
+	static FVector3<T> Zero()
+	{
+		return FVector3<T>((T)0, (T)0, (T)0);
+	}
+	static FVector3<T> One()
+	{
+		return FVector3<T>((T)1, (T)1, (T)1);
+	}
+	static FVector3<T> UnitX()
+	{
+		return FVector3<T>((T)1, (T)0, (T)0);
+	}
+	static FVector3<T> UnitY()
+	{
+		return FVector3<T>((T)0, (T)1, (T)0);
+	}
+	static FVector3<T> UnitZ()
+	{
+		return FVector3<T>((T)0, (T)0, (T)1);
+	}
+	static FVector3<T> MaxVector()
+	{
+		return FVector3<T>(TNumericLimits<T>::Max(), TNumericLimits<T>::Max(), TNumericLimits<T>::Max());
+	}
+
+	T Length() const
+	{
+		return TMathUtil<T>::Sqrt(X * X + Y * Y + Z * Z);
+	}
+	T SquaredLength() const
+	{
+		return X * X + Y * Y + Z * Z;
+	}
+
+	T Distance(const UE::Math::TVector<T>& V2) const
+	{
+		T dx = V2.X - X;
+		T dy = V2.Y - Y;
+		T dz = V2.Z - Z;
+		return TMathUtil<T>::Sqrt(dx * dx + dy * dy + dz * dz);
+	}
+	T DistanceSquared(const UE::Math::TVector<T>& V2) const
+	{
+		T dx = V2.X - X;
+		T dy = V2.Y - Y;
+		T dz = V2.Z - Z;
+		return dx * dx + dy * dy + dz * dz;
+	}
+
+	FVector3<T> operator-() const
+	{
+		return FVector3<T>(-X, -Y, -Z);
+	}
+
+	FVector3<T> operator+(const UE::Math::TVector<T>& V2) const
+	{
+		return FVector3<T>(X + V2.X, Y + V2.Y, Z + V2.Z);
+	}
+
+	FVector3<T> operator-(const UE::Math::TVector<T>& V2) const
+	{
+		return FVector3<T>(X - V2.X, Y - V2.Y, Z - V2.Z);
+	}
+
+	FVector3<T> operator+(const T& Scalar) const
+	{
+		return FVector3<T>(X + Scalar, Y + Scalar, Z + Scalar);
+	}
+
+	FVector3<T> operator-(const T& Scalar) const
+	{
+		return FVector3<T>(X - Scalar, Y - Scalar, Z - Scalar);
+	}
+
+	FVector3<T> operator*(const T& Scalar) const
+	{
+		return FVector3<T>(X * Scalar, Y * Scalar, Z * Scalar);
+	}
+
+	FVector3<T> operator*(const UE::Math::TVector<T>& V2) const // component-wise
+	{
+		return FVector3<T>(X * V2.X, Y * V2.Y, Z * V2.Z);
+	}
+
+	FVector3<T> operator/(const T& Scalar) const
+	{
+		return FVector3<T>(X / Scalar, Y / Scalar, Z / Scalar);
+	}
+
+	FVector3<T> operator/(const UE::Math::TVector<T>& V2) const // component-wise
+	{
+		return FVector3<T>(X / V2.X, Y / V2.Y, Z / V2.Z);
+	}
+
+	template<typename RealType2, TEMPLATE_REQUIRES(std::is_floating_point<RealType2>::value)>
+	FVector3<T> operator*(const RealType2& Scalar) const
+	{
+		return FVector3<T>(X * (T)Scalar, Y * (T)Scalar, Z * (T)Scalar);
+	}
+
+	FVector3<T>& operator+=(const UE::Math::TVector<T>& V2)
+	{
+		X += V2.X;
+		Y += V2.Y;
+		Z += V2.Z;
+		return *this;
+	}
+
+	FVector3<T>& operator-=(const UE::Math::TVector<T>& V2)
+	{
+		X -= V2.X;
+		Y -= V2.Y;
+		Z -= V2.Z;
+		return *this;
+	}
+
+	FVector3<T>& operator*=(const T& Scalar)
+	{
+		X *= Scalar;
+		Y *= Scalar;
+		Z *= Scalar;
+		return *this;
+	}
+
+	FVector3<T>& operator/=(const T& Scalar)
+	{
+		X /= Scalar;
+		Y /= Scalar;
+		Z /= Scalar;
+		return *this;
+	}
+
+	T Dot(const UE::Math::TVector<T>& V2) const
+	{
+		return X * V2.X + Y * V2.Y + Z * V2.Z;
+	}
+
+	FVector3<T> Cross(const UE::Math::TVector<T>& V2) const
+	{
+		return FVector3(
+			Y * V2.Z - Z * V2.Y,
+			Z * V2.X - X * V2.Z,
+			X * V2.Y - Y * V2.X);
+	}
+};
+
+
+
+
+
+/** @return unit vector along axis X=0, Y=1, Z=2 */
+template <typename T>
+constexpr UE::Math::TVector<T> MakeUnitVector3(int32 Axis)
+{
+	UE::Math::TVector<T> UnitVec((T)0, (T)0, (T)0);
+	UnitVec[FMath::Clamp(Axis, 0, 2)] = (T)1;
+	return UnitVec;
+}
+
+
+template<typename T>
+T Length(const UE::Math::TVector<T>& V)
+{
+	return TMathUtil<T>::Sqrt(V.X*V.X + V.Y*V.Y + V.Z*V.Z);
+}
+
+template<typename T>
+T SquaredLength(const UE::Math::TVector<T>& V)
+{
+	return V.X*V.X + V.Y*V.Y + V.Z*V.Z;
+}
+
+
+template <typename T>
+constexpr bool IsNormalized(const UE::Math::TVector<T>& Vector, const T Tolerance = TMathUtil<T>::ZeroTolerance)
+{
+	return TMathUtil<T>::Abs((Vector.X*Vector.X + Vector.Y*Vector.Y + Vector.Z*Vector.Z) - 1) < Tolerance;
+}
+
+
+template<typename T>
+T Normalize(UE::Math::TVector<T>& Vector, const T Epsilon = 0)
+{
+	T length = Vector.Length();
+	if (length > Epsilon)
+	{
+		T invLength = ((T)1) / length;
+		Vector.X *= invLength;
+		Vector.Y *= invLength;
+		Vector.Z *= invLength;
+		return length;
+	}
+	Vector.X = Vector.Y = Vector.Z = (T)0;
+	return (T)0;
+}
+
+template<typename T>
+constexpr UE::Math::TVector<T> Normalized(const UE::Math::TVector<T>& Vector, const T Epsilon = 0)
+{
+	T length = Vector.Length();
+	if (length > Epsilon)
+	{
+		T invLength = ((T)1) / length;
+		return UE::Math::TVector<T>(Vector.X*invLength, Vector.Y*invLength, Vector.Z*invLength);
+	}
+	return UE::Math::TVector<T>((T)0, (T)0, (T)0);
+}
+
+template<typename T>
+T Distance(const UE::Math::TVector<T>& V1, const UE::Math::TVector<T>& V2)
+{
+	T dx = V2.X - V1.X;
+	T dy = V2.Y - V1.Y;
+	T dz = V2.Z - V1.Z;
+	return TMathUtil<T>::Sqrt(dx * dx + dy * dy + dz * dz);
+}
+
+template<typename T>
+T DistanceSquared(const UE::Math::TVector<T>& V1, const UE::Math::TVector<T>& V2)
+{
+	T dx = V2.X - V1.X;
+	T dy = V2.Y - V1.Y;
+	T dz = V2.Z - V1.Z;
+	return dx * dx + dy * dy + dz * dz;
+}
+
+
+
+template<typename T>
+T Dot(const UE::Math::TVector<T>& V1, const UE::Math::TVector<T>& V2)
+{
+	return V1.X * V2.X + V1.Y * V2.Y + V1.Z * V2.Z;
+}
+
+template<typename T>
+UE::Math::TVector<T> Cross(const UE::Math::TVector<T>& V1, const UE::Math::TVector<T>& V2)
+{
+	return UE::Math::TVector<T>(
+		V1.Y * V2.Z - V1.Z * V2.Y,
+		V1.Z * V2.X - V1.X * V2.Z,
+		V1.X * V2.Y - V1.Y * V2.X);
+}
+
+template<typename T>
+UE::Math::TVector<T> UnitCross(const UE::Math::TVector<T>& V1, const UE::Math::TVector<T>& V2)
+{
+	UE::Math::TVector<T> N = V1.Cross(V2);
+	return Normalized(N);
+}
+
+/**
+ * Computes the Angle between V1 and V2, assuming they are already normalized
+ * @return the (positive) angle between V1 and V2 in degrees
+ */
+template <typename T>
+T AngleD(const UE::Math::TVector<T>& V1, const UE::Math::TVector<T>& V2)
+{
+	T DotVal = V1.Dot(V2);
+	T ClampedDot = (DotVal < (T)-1) ? (T)-1 : ((DotVal > (T)1) ? (T)1 : DotVal);
+	return TMathUtil<T>::ACos(ClampedDot) * TMathUtil<T>::RadToDeg;
+}
+
+/**
+ * Computes the Angle between V1 and V2, assuming they are already normalized
+ * @return the (positive) angle between V1 and V2 in radians
+ */
+template <typename T>
+T AngleR(const UE::Math::TVector<T>& V1, const UE::Math::TVector<T>& V2)
+{
+	T DotVal = V1.Dot(V2);
+	T ClampedDot = (DotVal < (T)-1) ? (T)-1 : ((DotVal > (T)1) ? (T)1 : DotVal);
+	return TMathUtil<T>::ACos(ClampedDot);
+}
+
+template <typename T>
+constexpr FVector2<T> GetXY(const UE::Math::TVector<T>& V)
+{
+	return FVector2<T>(V.X, V.Y);
+}
+
+template <typename T>
+constexpr FVector2<T> GetXZ(const UE::Math::TVector<T>& V)
+{
+	return FVector2<T>(V.X, V.Z);
+}
+
+template <typename T>
+constexpr FVector2<T> GetYZ(const UE::Math::TVector<T>& V)
+{
+	return FVector2<T>(V.Y, V.Z);
+}
+
+
+template<typename T>
+constexpr FVector3<T> Min(const UE::Math::TVector<T>& V0, const UE::Math::TVector<T>& V1)
+{
+	return FVector3<T>(TMathUtil<T>::Min(V0.X, V1.X),
+		TMathUtil<T>::Min(V0.Y, V1.Y),
+		TMathUtil<T>::Min(V0.Z, V1.Z));
+}
+
+template<typename T>
+constexpr FVector3<T> Max(const UE::Math::TVector<T>& V0, const UE::Math::TVector<T>& V1)
+{
+	return FVector3<T>(TMathUtil<T>::Max(V0.X, V1.X),
+		TMathUtil<T>::Max(V0.Y, V1.Y),
+		TMathUtil<T>::Max(V0.Z, V1.Z));
+}
+
+
+template<typename T>
+constexpr T MaxElement(const UE::Math::TVector<T>& Vector)
+{
+	return TMathUtil<T>::Max3(Vector.X, Vector.Y, Vector.Z);
+}
+
+/** @return 0/1/2 index of maximum element */
+template<typename T>
+constexpr int32 MaxElementIndex(const UE::Math::TVector<T>& Vector)
+{
+	return TMathUtil<T>::Max3Index(Vector.X, Vector.Y, Vector.Z);
+}
+
+template<typename T>
+constexpr T MinElement(const UE::Math::TVector<T>& Vector)
+{
+	return TMathUtil<T>::Min3(Vector.X, Vector.Y, Vector.Z);
+}
+
+/** @return 0/1/2 index of minimum element */
+template<typename T>
+constexpr int32 MinElementIndex(const UE::Math::TVector<T>& Vector)
+{
+	return TMathUtil<T>::Min3Index(Vector.X, Vector.Y, Vector.Z);
+}
+
+template<typename T>
+constexpr T MaxAbsElement(const UE::Math::TVector<T>& Vector)
+{
+	return TMathUtil<T>::Max3(TMathUtil<T>::Abs(Vector.X), TMathUtil<T>::Abs(Vector.Y), TMathUtil<T>::Abs(Vector.Z));
+}
+
+/** @return 0/1/2 index of maximum absolute-value element */
+template<typename T>
+constexpr T MaxAbsElementIndex(const UE::Math::TVector<T>& Vector)
+{
+	return TMathUtil<T>::Max3Index(TMathUtil<T>::Abs(Vector.X), TMathUtil<T>::Abs(Vector.Y), TMathUtil<T>::Abs(Vector.Z));
+}
+
+template<typename T>
+constexpr T MinAbsElement(const UE::Math::TVector<T>& Vector)
+{
+	return TMathUtil<T>::Min3(TMathUtil<T>::Abs(Vector.X), TMathUtil<T>::Abs(Vector.Y), TMathUtil<T>::Abs(Vector.Z));
+}
+
+/** @return 0/1/2 index of minimum absolute-value element */
+template<typename T>
+constexpr T MinAbsElementIndex(const UE::Math::TVector<T>& Vector)
+{
+	return TMathUtil<T>::Min3Index(TMathUtil<T>::Abs(Vector.X), TMathUtil<T>::Abs(Vector.Y), TMathUtil<T>::Abs(Vector.Z));
+}
+
+template<typename T>
+constexpr FColor ToFColor(const UE::Math::TVector<T>& Vector)
+{
+	return FColor(
+		FMathf::Clamp((int)((float)Vector.X * 255.0f), 0, 255),
+		FMathf::Clamp((int)((float)Vector.Y * 255.0f), 0, 255),
+		FMathf::Clamp((int)((float)Vector.Z * 255.0f), 0, 255));
+}
+
+template<typename T>
+constexpr FLinearColor ToLinearColor(const UE::Math::TVector<T>& Vector)
+{
+	return FLinearColor((float)Vector.X, (float)Vector.Y, (float)Vector.Z);
+}
+
+template<typename T>
+UE::Math::TVector<T> Lerp(const UE::Math::TVector<T>& A, const UE::Math::TVector<T>& B, T Alpha)
+{
+	T OneMinusAlpha = (T)1 - Alpha;
+	return UE::Math::TVector<T>(OneMinusAlpha * A.X + Alpha * B.X,
+		OneMinusAlpha * A.Y + Alpha * B.Y,
+		OneMinusAlpha * A.Z + Alpha * B.Z);
+}
+
+template<typename T>
+UE::Math::TVector<T> Blend3(const UE::Math::TVector<T>& A, const UE::Math::TVector<T>& B, const UE::Math::TVector<T>& C, const T& WeightA, const T& WeightB, const T& WeightC)
+{
+	return UE::Math::TVector<T>(
+		WeightA * A.X + WeightB * B.X + WeightC * C.X,
+		WeightA * A.Y + WeightB * B.Y + WeightC * C.Y,
+		WeightA * A.Z + WeightB * B.Z + WeightC * C.Z);
+}
+
+template <typename RealType>
+inline FVector3<RealType> operator*(RealType Scalar, const FVector3<RealType>& V)
+{
+	return FVector3<RealType>(Scalar * V.X, Scalar * V.Y, Scalar * V.Z);
+}
+
+// allow float*Vector3<double> and double*Vector3<float>
+template <typename RealType, typename RealType2, TEMPLATE_REQUIRES(std::is_floating_point<RealType2>::value)>
+inline FVector3<RealType> operator*(RealType2 Scalar, const FVector3<RealType>& V)
+{
+	return FVector3<RealType>((RealType)Scalar * V.X, (RealType)Scalar * V.Y, (RealType)Scalar * V.Z);
+}
+
+template <typename RealType>
+std::ostream& operator<<(std::ostream& os, const FVector3<RealType>& Vec)
+{
+	os << Vec.X << " " << Vec.Y << " " << Vec.Z;
+	return os;
+}
+
+template <typename T>
+FORCEINLINE uint32 GetTypeHash(const FVector3<T>& Vector)
+{
+	// (this is how FIntVector and all the other FVectors do their hash functions)
+	// Note: this assumes there's no padding that could contain uncompared data.
+	return FCrc::MemCrc_DEPRECATED(&Vector, sizeof(FVector3<T>));
+}
+
+
+
+
+
+template <typename T>
+FORCEINLINE uint32 GetTypeHash(const FVector2<T>& Vector)
+{
+	// (this is how FIntVector and all the other FVectors do their hash functions)
+	// Note: this assumes there's no padding that could contain uncompared data.
+	return FCrc::MemCrc_DEPRECATED(&Vector, sizeof(FVector2<T>));
+}
+
+template <typename RealType>
+inline FVector2<RealType> operator*(RealType Scalar, const FVector2<RealType>& V)
+{
+	return FVector2<RealType>(Scalar * V.X, Scalar * V.Y);
+}
+
+// allow float*Vector2<double> and double*Vector2<float>
+template <typename RealType, typename RealType2>
+inline FVector2<RealType> operator*(RealType2 Scalar, const FVector2<RealType>& V)
+{
+	return FVector2<RealType>((RealType)Scalar * V.X, (RealType)Scalar * V.Y);
+}
+
+template <typename RealType>
+std::ostream& operator<<(std::ostream& os, const FVector2<RealType>& Vec)
+{
+	os << Vec.X << " " << Vec.Y;
+	return os;
+}
+
+
+
+
+
+
+
+
+template <typename T>
+struct TVector4
+{
+	T X{}, Y{}, Z{}, W{};
+
+	constexpr TVector4()
+		: X(0), Y(0), Z(0), W(0)
+	{
+	}
+
+	constexpr TVector4(T ValX, T ValY, T ValZ, T ValW)
+		: X(ValX), Y(ValY), Z(ValZ), W(ValW)
+	{
+	}
+
+	constexpr TVector4(const T* Data)
+		: X(Data[0]), Y(Data[1]), Z(Data[2]), W(Data[3])
+	{
+	}
+
+	constexpr TVector4(const TVector4& Vec) = default;
+
+	template<typename RealType2>
+	explicit constexpr TVector4(const TVector4<RealType2>& Vec)
+		: X((T)Vec.X), Y((T)Vec.Y), Z((T)Vec.Z), W((T)Vec.W)
+	{
+	}
+
+	explicit constexpr operator const T*() const
+	{
+		return &X;
+	};
+	explicit constexpr operator T*()
+	{
+		return &X;
+	}
+
+	explicit constexpr operator FLinearColor() const
+	{
+		return FLinearColor((float)X, (float)Y, (float)Z, (float)W);
+	}
+	constexpr TVector4(const FLinearColor& Color)
+		: X((T)Color.R), Y((T)Color.G), Z((T)Color.B), W((T)Color.A)
+	{
+	}
+
+	explicit constexpr operator FColor() const
+	{
+		return FColor(
+			FMathf::Clamp((int)((float)X*255.0f), 0, 255),
+			FMathf::Clamp((int)((float)Y*255.0f), 0, 255),
+			FMathf::Clamp((int)((float)Z*255.0f), 0, 255),
+			FMathf::Clamp((int)((float)W*255.0f), 0, 255));
+	}
+
+
+	static TVector4<T> Zero()
+	{
+		return TVector4<T>((T)0, (T)0, (T)0, (T)0);
+	}
+	static TVector4<T> One()
+	{
+		return TVector4<T>((T)1, (T)1, (T)1, (T)1);
+	}
+
+
+	TVector4<T>& operator=(const TVector4<T>& V2)
+	{
+		X = V2.X;
+		Y = V2.Y;
+		Z = V2.Z;
+		W = V2.W;
+		return *this;
+	}
+
+	T& operator[](int Idx)
+	{
+		return (&X)[Idx];
+	}
+	const T& operator[](int Idx) const
+	{
+		return (&X)[Idx];
+	}
+
+	T Length() const
+	{
+		return TMathUtil<T>::Sqrt(X * X + Y * Y + Z * Z + W * W);
+	}
+	T SquaredLength() const
+	{
+		return X * X + Y * Y + Z * Z + W * W;
+	}
+
+
+	constexpr TVector4<T> operator-() const
+	{
+		return TVector4<T>(-X, -Y, -Z, -W);
+	}
+
+	constexpr TVector4<T> operator+(const TVector4<T>& V2) const
+	{
+		return TVector4<T>(X + V2.X, Y + V2.Y, Z + V2.Z, W + V2.W);
+	}
+
+	constexpr TVector4<T> operator-(const TVector4<T>& V2) const
+	{
+		return TVector4<T>(X - V2.X, Y - V2.Y, Z - V2.Z, W - V2.W);
+	}
+
+	constexpr TVector4<T> operator+(const T& Scalar) const
+	{
+		return TVector4<T>(X + Scalar, Y + Scalar, Z + Scalar, W + Scalar);
+	}
+
+	constexpr TVector4<T> operator-(const T& Scalar) const
+	{
+		return TVector4<T>(X - Scalar, Y - Scalar, Z - Scalar, W - Scalar);
+	}
+
+	constexpr TVector4<T> operator*(const T& Scalar) const
+	{
+		return TVector4<T>(X * Scalar, Y * Scalar, Z * Scalar, W * Scalar);
+	}
+
+	template<typename RealType2>
+	constexpr TVector4<T> operator*(const RealType2& Scalar) const
+	{
+		return TVector4<T>(X * (T)Scalar, Y * (T)Scalar, Z * (T)Scalar, W * (T)Scalar);
+	}
+
+	constexpr TVector4<T> operator*(const TVector4<T>& V2) const // component-wise
+	{
+		return TVector4<T>(X * V2.X, Y * V2.Y, Z * V2.Z, W * V2.W);
+	}
+
+	constexpr TVector4<T> operator/(const T& Scalar) const
+	{
+		return TVector4<T>(X / Scalar, Y / Scalar, Z / Scalar, W / Scalar);
+	}
+
+	constexpr TVector4<T> operator/(const TVector4<T>& V2) const // component-wise
+	{
+		return TVector4<T>(X / V2.X, Y / V2.Y, Z / V2.Z, W / V2.W);
+	}
+
+	constexpr TVector4<T>& operator+=(const TVector4<T>& V2)
+	{
+		X += V2.X;
+		Y += V2.Y;
+		Z += V2.Z;
+		W += V2.W;
+		return *this;
+	}
+
+	constexpr TVector4<T>& operator-=(const TVector4<T>& V2)
+	{
+		X -= V2.X;
+		Y -= V2.Y;
+		Z -= V2.Z;
+		W -= V2.W;
+		return *this;
+	}
+
+	constexpr TVector4<T>& operator*=(const T& Scalar)
+	{
+		X *= Scalar;
+		Y *= Scalar;
+		Z *= Scalar;
+		W *= Scalar;
+		return *this;
+	}
+
+	constexpr TVector4<T>& operator/=(const T& Scalar)
+	{
+		X /= Scalar;
+		Y /= Scalar;
+		Z /= Scalar;
+		W /= Scalar;
+		return *this;
+	}
+
+	T Dot(const TVector4<T>& V2) const
+	{
+		return X * V2.X + Y * V2.Y + Z * V2.Z + W * V2.W;
+	}
+
+	constexpr bool operator==(const TVector4<T>& Other) const
+	{
+		return X == Other.X && Y == Other.Y && Z == Other.Z && W == Other.W;
+	}
+
+	constexpr bool operator!=(const TVector4<T>& Other) const
+	{
+		return X != Other.X || Y != Other.Y || Z != Other.Z || W != Other.W;
+	}
+
+	/**
+	 * Serialization operator for TVector4.
+	 *
+	 * @param Ar Archive to serialize with.
+	 * @param Vec Vector to serialize.
+	 * @returns Passing down serializing archive.
+	 */
+	friend FArchive& operator<<(FArchive& Ar, TVector4& Vec)
+	{
+		Vec.Serialize(Ar);
+		return Ar;
+	}
+
+	/** Serialize TVector4 to an archive. */
+	void Serialize(FArchive& Ar)
+	{
+		Ar << X;
+		Ar << Y;
+		Ar << Z;
+		Ar << W;
+	}
+};
+
+template <typename T>
+constexpr bool IsNormalized(const TVector4<T>& Vector, const T Tolerance = TMathUtil<T>::ZeroTolerance)
+{
+	return TMathUtil<T>::Abs((Vector.X*Vector.X + Vector.Y*Vector.Y + Vector.Z*Vector.Z + Vector.W*Vector.W) - 1) < Tolerance;
+}
+
+template<typename T>
+T Normalize(TVector4<T>& Vector, const T Epsilon = 0)
+{
+	T length = Vector.Length();
+	if (length > Epsilon)
+	{
+		T invLength = ((T)1) / length;
+		Vector.X *= invLength;
+		Vector.Y *= invLength;
+		Vector.Z *= invLength;
+		Vector.W *= invLength;
+		return length;
+	}
+	Vector.X = Vector.Y = Vector.Z = Vector.W = (T)0;
+	return (T)0;
+}
+
+template<typename T>
+TVector4<T> Normalized(const TVector4<T>& Vector, const T Epsilon = 0)
+{
+	T length = Vector.Length();
+	if (length > Epsilon)
+	{
+		T invLength = ((T)1) / length;
+		return TVector4<T>(Vector.X*invLength, Vector.Y*invLength, Vector.Z*invLength, Vector.W*invLength);
+	}
+	return TVector4<T>::Zero();
+}
+
+template<typename T>
+constexpr FVector3<T> GetXYZ(const TVector4<T>& V)
+{
+	return FVector3<T>(V.X, V.Y, V.Z);
+}
+
+template<typename T>
+TVector4<T> Lerp(const TVector4<T>& A, const TVector4<T>& B, T Alpha)
+{
+	T OneMinusAlpha = (T)1 - Alpha;
+	return TVector4<T>(OneMinusAlpha * A.X + Alpha * B.X,
+		OneMinusAlpha * A.Y + Alpha * B.Y,
+		OneMinusAlpha * A.Z + Alpha * B.Z,
+		OneMinusAlpha * A.W + Alpha * B.W);
+}
+
+template<typename T>
+TVector4<T> Blend3(const TVector4<T>& A, const TVector4<T>& B, const TVector4<T>& C, const T& WeightA, const T& WeightB, const T& WeightC)
+{
+	return TVector4<T>(
+		WeightA * A.X + WeightB * B.X + WeightC * C.X,
+		WeightA * A.Y + WeightB * B.Y + WeightC * C.Y,
+		WeightA * A.Z + WeightB * B.Z + WeightC * C.Z,
+		WeightA * A.W + WeightB * B.W + WeightC * C.W);
+}
+
+template <typename RealType>
+inline TVector4<RealType> operator*(RealType Scalar, const TVector4<RealType>& V)
+{
+	return TVector4<RealType>(Scalar * V.X, Scalar * V.Y, Scalar * V.Z, Scalar * V.W);
+}
+
+// allow float*Vector4<double> and double*Vector4<float>
+template <typename RealType, typename RealType2>
+inline TVector4<RealType> operator*(RealType2 Scalar, const TVector4<RealType>& V)
+{
+	return TVector4<RealType>((RealType)Scalar * V.X, (RealType)Scalar * V.Y, (RealType)Scalar * V.Z, (RealType)Scalar * V.W);
+}
+
+template <typename RealType>
+std::ostream& operator<<(std::ostream& os, const TVector4<RealType>& Vec)
+{
+	os << Vec.X << " " << Vec.Y << " " << Vec.Z << " " << Vec.W;
+	return os;
+}
+
+
+template <typename T>
+FORCEINLINE uint32 GetTypeHash(const TVector4<T>& Vector)
+{
+	// (this is how FIntVector and all the other FVectors do their hash functions)
+	// Note: this assumes there's no padding that could contain uncompared data.
+	return FCrc::MemCrc_DEPRECATED(&Vector, sizeof(TVector4<T>));
+}
+
+
+
+
+} // end namespace UE::Geometry
+} // end namespace UE
+
+typedef UE::Geometry::FVector2<float> FVector2f;
+typedef UE::Geometry::FVector2<double> FVector2d;
+
+//typedef UE::Geometry::FVector3<float> FVector3f;
+//typedef UE::Geometry::FVector3<double> FVector3d;
+
+typedef UE::Geometry::TVector4<float> FVector4f;
+typedef UE::Geometry::TVector4<double> FVector4d;
+template<> struct TCanBulkSerialize<FVector2f> { enum { Value = true }; };
+template<> struct TCanBulkSerialize<FVector2d> { enum { Value = true }; };
+
+template<> struct TCanBulkSerialize<FVector4f> { enum { Value = true }; };
+template<> struct TCanBulkSerialize<FVector4d> { enum { Value = true }; };
diff --git a/Engine/Source/Runtime/GeometryCore/Public/VectorUtil.h b/Engine/Source/Runtime/GeometryCore/Public/VectorUtil.h
new file mode 100644
--- /dev/null
+++ b/Engine/Source/Runtime/GeometryCore/Public/VectorUtil.h
@@ -0,0 +1,454 @@
+// Copyright Epic Games, Inc. All Rights Reserved.
+
+#pragma once
+
+#include "MathUtil.h"
+#include "VectorTypes.h"
+
+
+enum class EIntersectionResult
+{
+	NotComputed,
+	Intersects,
+	NoIntersection,
+	InvalidQuery
+};
+
+enum class EIntersectionType
+{
+	Empty,
+	Point,
+	Segment,
+	Line,
+	Polygon,
+	Plane,
+	MultiSegment,
+	Unknown
+};
+
+namespace VectorUtil
+{
+	using namespace UE::Geometry;
+
+	/**
+	 * @return true if all components of V are finite
+	 */
+	template <typename RealType>
+	inline bool IsFinite(const FVector2<RealType>& V)
+	{
+		return TMathUtil<RealType>::IsFinite(V.X) && TMathUtil<RealType>::IsFinite(V.Y);
+	}
+
+	/**
+	 * @return true if all components of V are finite
+	 */
+	template <typename RealType>
+	inline bool IsFinite(const UE::Math::TVector<RealType>& V)
+	{
+		return TMathUtil<RealType>::IsFinite(V.X) && TMathUtil<RealType>::IsFinite(V.Y) && TMathUtil<RealType>::IsFinite(V.Z);
+	}
+
+
+	/**
+	 * @return input Value clamped to range [MinValue, MaxValue]
+	 */
+	template <typename RealType>
+	inline RealType Clamp(RealType Value, RealType MinValue, RealType MaxValue)
+	{
+		return (Value < MinValue) ? MinValue : ((Value > MaxValue) ? MaxValue : Value);
+	}
+
+	/**
+	 * @return normalized vector that is perpendicular to triangle V0,V1,V2  (triangle normal)
+	 */
+	template <typename RealType>
+	inline FVector3<RealType> Normal(const UE::Math::TVector<RealType>& V0, const UE::Math::TVector<RealType>& V1, const UE::Math::TVector<RealType>& V2)
+	{
+		FVector3<RealType> edge1(V1 - V0);
+		FVector3<RealType> edge2(V2 - V0);
+		Normalize(edge1);
+		Normalize(edge2);
+		// Unreal has Left-Hand Coordinate System so we need to reverse this cross-product to get proper triangle normal
+		FVector3<RealType> vCross(edge2.Cross(edge1));
+		//FVector3<RealType> vCross(edge1.Cross(edge2));
+		return Normalized(vCross);
+	}
+
+	/**
+	 * @return un-normalized direction that is parallel to normal of triangle V0,V1,V2
+	 */
+	template <typename RealType>
+	inline FVector3<RealType> NormalDirection(const UE::Math::TVector<RealType>& V0, const UE::Math::TVector<RealType>& V1, const UE::Math::TVector<RealType>& V2)
+	{
+		// Unreal has Left-Hand Coordinate System so we need to reverse this cross-product to get proper triangle normal
+		return (V2 - V0).Cross(V1 - V0);
+		//return (V1 - V0).Cross(V2 - V0);
+	}
+
+	/**
+	 * @return area of 3D triangle V0,V1,V2
+	 */
+	template <typename RealType>
+	inline RealType Area(const UE::Math::TVector<RealType>& V0, const UE::Math::TVector<RealType>& V1, const UE::Math::TVector<RealType>& V2)
+	{
+		FVector3<RealType> Edge1(V1 - V0);
+		FVector3<RealType> Edge2(V2 - V0);
+		FVector3<RealType> Cross = Edge2.Cross(Edge1);
+		return (RealType)0.5 * Cross.Length();
+	}
+
+	/**
+	 * @return area of 2D triangle V0,V1,V2
+	 */
+	template <typename RealType>
+	inline RealType Area(const FVector2<RealType>& V0, const FVector2<RealType>& V1, const FVector2<RealType>& V2)
+	{
+		FVector2<RealType> Edge1(V1 - V0);
+		FVector2<RealType> Edge2(V2 - V0);
+		RealType CrossZ = DotPerp(Edge1, Edge2);
+		return (RealType)0.5 * TMathUtil<RealType>::Abs(CrossZ);
+	}
+
+	/**
+	 * @return true if triangle V1,V2,V3 is obtuse
+	 */
+	template <typename RealType>
+	inline bool IsObtuse(const UE::Math::TVector<RealType>& V1, const UE::Math::TVector<RealType>& V2, const UE::Math::TVector<RealType>& V3)
+	{
+		RealType a2 = DistanceSquared(V1, V2);
+		RealType b2 = DistanceSquared(V1, V3);
+		RealType c2 = DistanceSquared(V2, V3);
+		return (a2 + b2 < c2) || (b2 + c2 < a2) || (c2 + a2 < b2);
+	}
+
+	/**
+	 * Calculate Normal and Area of triangle V0,V1,V2
+	 * @return triangle normal
+	 */
+	template <typename RealType>
+	inline FVector3<RealType> NormalArea(const UE::Math::TVector<RealType>& V0, const UE::Math::TVector<RealType>& V1, const UE::Math::TVector<RealType>& V2, RealType& AreaOut)
+	{
+		FVector3<RealType> edge1(V1 - V0);
+		FVector3<RealType> edge2(V2 - V0);
+		// Unreal has Left-Hand Coordinate System so we need to reverse this cross-product to get proper triangle normal
+		FVector3d vCross = edge2.Cross(edge1);
+		//FVector3d vCross = edge1.Cross(edge2);
+		AreaOut = RealType(0.5) * Normalize(vCross);
+		return vCross;
+	}
+
+	/** @return true if Abs(A-B) is less than Epsilon */
+	template <typename RealType>
+	inline bool EpsilonEqual(RealType A, RealType B, RealType Epsilon)
+	{
+		return TMathUtil<RealType>::Abs(A - B) < Epsilon;
+	}
+
+	/** @return true if all coordinates of V0 and V1 are within Epsilon of eachother */
+	template <typename RealType>
+	inline bool EpsilonEqual(const FVector2<RealType>& V0, const FVector2<RealType>& V1, RealType Epsilon)
+	{
+		return EpsilonEqual(V0.X, V1.X, Epsilon) && EpsilonEqual(V0.Y, V1.Y, Epsilon);
+	}
+
+	/** @return true if all coordinates of V0 and V1 are within Epsilon of eachother */
+	template <typename RealType>
+	inline bool EpsilonEqual(const UE::Math::TVector<RealType>& V0, const UE::Math::TVector<RealType>& V1, RealType Epsilon)
+	{
+		return EpsilonEqual(V0.X, V1.X, Epsilon) && EpsilonEqual(V0.Y, V1.Y, Epsilon) && EpsilonEqual(V0.Z, V1.Z, Epsilon);
+	}
+
+
+	/** @return 0/1/2 index of smallest value in Vector3 */
+	template <typename ValueVecType>
+	inline int Min3Index(const ValueVecType& Vector3)
+	{
+		if (Vector3[0] <= Vector3[1])
+		{
+			return Vector3[0] <= Vector3[2] ? 0 : 2;
+		}
+		else
+		{
+			return (Vector3[1] <= Vector3[2]) ? 1 : 2;
+		}
+	}
+
+
+	/**
+	 * Calculates two vectors perpendicular to input Normal, as efficiently as possible.
+	 */
+	template <typename RealType>
+	inline void MakePerpVectors(const UE::Math::TVector<RealType>& Normal, UE::Math::TVector<RealType>& OutPerp1, UE::Math::TVector<RealType>& OutPerp2)
+	{
+		// Duff et al method, from https://graphics.pixar.com/library/OrthonormalB/paper.pdf
+		if (Normal.Z < (RealType)0)
+		{
+			RealType A = (RealType)1 / ((RealType)1 - Normal.Z);
+			RealType B = Normal.X * Normal.Y * A;
+			OutPerp1.X = (RealType)1 - Normal.X * Normal.X * A;
+			OutPerp1.Y = -B;
+			OutPerp1.Z = Normal.X;
+			OutPerp2.X = B;
+			OutPerp2.Y = Normal.Y * Normal.Y * A - (RealType)1;
+			OutPerp2.Z = -Normal.Y;
+		}
+		else
+		{
+			RealType A = (RealType)1 / ((RealType)1 + Normal.Z);
+			RealType B = -Normal.X * Normal.Y * A;
+			OutPerp1.X = (RealType)1 - Normal.X * Normal.X * A;
+			OutPerp1.Y = B;
+			OutPerp1.Z = -Normal.X;
+			OutPerp2.X = B;
+			OutPerp2.Y = (RealType)1 - Normal.Y * Normal.Y * A;
+			OutPerp2.Z = -Normal.Y;
+		}
+	}
+
+	/**
+	 * Calculates one vector perpendicular to input Normal, as efficiently as possible.
+	 */
+	template <typename RealType>
+	inline void MakePerpVector(const FVector3<RealType>& Normal, FVector3<RealType>& OutPerp1)
+	{
+		// Duff et al method, from https://graphics.pixar.com/library/OrthonormalB/paper.pdf
+		if (Normal.Z < (RealType)0)
+		{
+			RealType A = (RealType)1 / ((RealType)1 - Normal.Z);
+			RealType B = Normal.X * Normal.Y * A;
+			OutPerp1.X = (RealType)1 - Normal.X * Normal.X * A;
+			OutPerp1.Y = -B;
+			OutPerp1.Z = Normal.X;
+		}
+		else
+		{
+			RealType A = (RealType)1 / ((RealType)1 + Normal.Z);
+			RealType B = -Normal.X * Normal.Y * A;
+			OutPerp1.X = (RealType)1 - Normal.X * Normal.X * A;
+			OutPerp1.Y = B;
+			OutPerp1.Z = -Normal.X;
+		}
+	}
+
+
+	/**
+	 * Calculates angle between VFrom and VTo after projection onto plane with normal defined by PlaneN
+	 * @return angle in degrees
+	 */
+	template <typename RealType>
+	inline RealType PlaneAngleSignedD(const UE::Math::TVector<RealType>& VFrom, const UE::Math::TVector<RealType>& VTo, const UE::Math::TVector<RealType>& PlaneN)
+	{
+		FVector3<RealType> vFrom = VFrom - VFrom.Dot(PlaneN) * PlaneN;
+		FVector3<RealType> vTo = VTo - VTo.Dot(PlaneN) * PlaneN;
+		Normalize(vFrom);
+		Normalize(vTo);
+		FVector3<RealType> C = vFrom.Cross(vTo);
+		if (C.SquaredLength() < TMathUtil<RealType>::ZeroTolerance)
+		{ // vectors are parallel
+			return vFrom.Dot(vTo) < 0 ? (RealType)180 : (RealType)0;
+		}
+		RealType fSign = C.Dot(PlaneN) < 0 ? (RealType)-1 : (RealType)1;
+		return (RealType)(fSign * AngleD(vFrom, vTo));
+	}
+
+	/**
+	 * tan(theta/2) = +/- sqrt( (1-cos(theta)) / (1+cos(theta)) )
+	 * @return positive value of tan(theta/2) where theta is angle between normalized vectors A and B
+	 */
+	template <typename RealType>
+	RealType VectorTanHalfAngle(const UE::Math::TVector<RealType>& A, const UE::Math::TVector<RealType>& B)
+	{
+		RealType cosAngle = A.Dot(B);
+		RealType sqr = ((RealType)1 - cosAngle) / ((RealType)1 + cosAngle);
+		sqr = Clamp(sqr, (RealType)0, TMathUtil<RealType>::MaxReal);
+		return TMathUtil<RealType>::Sqrt(sqr);
+	}
+
+	/**
+	 * Fast cotangent of angle between two vectors (*do not have to be normalized unit vectors*).
+	 * cot = cos/sin, both of which can be computed from vector identities
+	 * @return cotangent of angle between V1 and V2, or zero if result would be unstable (eg infinity)
+	 */ 
+	template <typename RealType>
+	RealType VectorCot(const UE::Math::TVector<RealType>& V1, const UE::Math::TVector<RealType>& V2)
+	{
+		// formula from http://www.geometry.caltech.edu/pubs/DMSB_III.pdf
+		RealType fDot = V1.Dot(V2);
+		RealType lensqr1 = V1.SquaredLength();
+		RealType lensqr2 = V2.SquaredLength();
+		RealType d = Clamp(lensqr1 * lensqr2 - fDot * fDot, (RealType)0.0, TMathUtil<RealType>::MaxReal);
+		if (d < TMathUtil<RealType>::ZeroTolerance)
+		{
+			return (RealType)0;
+		}
+		else
+		{
+			return fDot / TMathUtil<RealType>::Sqrt(d);
+		}
+	}
+
+	/**
+	 * Compute barycentric coordinates/weights of vPoint inside 3D triangle (V0,V1,V2). 
+	 * If point is in triangle plane and inside triangle, coords will be positive and sum to 1.
+	 * ie if result is a, then vPoint = a.x*V0 + a.y*V1 + a.z*V2.
+	 * TODO: make robust to degenerate triangles?
+	 */
+	template <typename RealType>
+	FVector3<RealType> BarycentricCoords(const UE::Math::TVector<RealType>& Point, const UE::Math::TVector<RealType>& V0, const UE::Math::TVector<RealType>& V1, const UE::Math::TVector<RealType>& V2)
+	{
+		FVector3<RealType> kV02 = V0 - V2;
+		FVector3<RealType> kV12 = V1 - V2;
+		FVector3<RealType> kPV2 = Point - V2;
+		RealType fM00 = kV02.Dot(kV02);
+		RealType fM01 = kV02.Dot(kV12);
+		RealType fM11 = kV12.Dot(kV12);
+		RealType fR0 = kV02.Dot(kPV2);
+		RealType fR1 = kV12.Dot(kPV2);
+		RealType fDet = fM00 * fM11 - fM01 * fM01;
+		RealType fInvDet = 1.0 / fDet;
+		RealType fBary1 = (fM11 * fR0 - fM01 * fR1) * fInvDet;
+		RealType fBary2 = (fM00 * fR1 - fM01 * fR0) * fInvDet;
+		RealType fBary3 = 1.0 - fBary1 - fBary2;
+		return FVector3<RealType>(fBary1, fBary2, fBary3);
+	}
+
+	/**
+	* Compute barycentric coordinates/weights of vPoint inside 2D triangle (V0,V1,V2). 
+	* If point is inside triangle, coords will be positive and sum to 1.
+	* ie if result is a, then vPoint = a.x*V0 + a.y*V1 + a.z*V2.
+	* TODO: make robust to degenerate triangles?
+	*/
+	template <typename RealType>
+	FVector3<RealType> BarycentricCoords(const FVector2<RealType>& Point, const FVector2<RealType>& V0, const FVector2<RealType>& V1, const FVector2<RealType>& V2)
+	{
+		FVector2<RealType> kV02 = V0 - V2;
+		FVector2<RealType> kV12 = V1 - V2;
+		FVector2<RealType> kPV2 = Point - V2;
+		RealType fM00 = kV02.Dot(kV02);
+		RealType fM01 = kV02.Dot(kV12);
+		RealType fM11 = kV12.Dot(kV12);
+		RealType fR0 = kV02.Dot(kPV2);
+		RealType fR1 = kV12.Dot(kPV2);
+		RealType fDet = fM00 * fM11 - fM01 * fM01;
+		RealType fInvDet = 1.0 / fDet;
+		RealType fBary1 = (fM11 * fR0 - fM01 * fR1) * fInvDet;
+		RealType fBary2 = (fM00 * fR1 - fM01 * fR0) * fInvDet;
+		RealType fBary3 = 1.0 - fBary1 - fBary2;
+		return FVector3<RealType>(fBary1, fBary2, fBary3);
+	}
+
+	/**
+	 * @return solid angle at point P for triangle A,B,C
+	 */
+	template <typename RealType>
+	inline RealType TriSolidAngle(UE::Math::TVector<RealType> A, UE::Math::TVector<RealType> B, UE::Math::TVector<RealType> C, const UE::Math::TVector<RealType>& P)
+	{
+		// Formula from https://igl.ethz.ch/projects/winding-number/
+		A -= P;
+		B -= P;
+		C -= P;
+		RealType la = A.Length(), lb = B.Length(), lc = C.Length();
+		RealType top = (la * lb * lc) + A.Dot(B) * lc + B.Dot(C) * la + C.Dot(A) * lb;
+		RealType bottom = A.X * (B.Y * C.Z - C.Y * B.Z) - A.Y * (B.X * C.Z - C.X * B.Z) + A.Z * (B.X * C.Y - C.X * B.Y);
+		// -2 instead of 2 to account for UE winding
+		return RealType(-2.0) * atan2(bottom, top);
+	}
+
+
+	/**
+	 * Calculate gradient of scalar field values fi,fj,fk defined at corners of triangle Vi,Vj,Vk and interpolated across triangle using linear basis functions.
+	 * This gradient is a 3D vector lying in the plane of the triangle (or zero if field is constant).
+	 * @return gradient (3D vector) lying in plane of triangle.
+	 */
+	template <typename RealType>
+	inline FVector3<RealType> TriGradient(UE::Math::TVector<RealType> Vi, UE::Math::TVector<RealType> Vj, UE::Math::TVector<RealType> Vk, RealType fi, RealType fj, RealType fk)
+	{
+		// recenter (better for precision)
+		FVector3<RealType> Centroid = (Vi + Vj + Vk) / (RealType)3;
+		Vi -= Centroid; Vj -= Centroid; Vk -= Centroid;
+		// calculate tangent-normal frame
+		FVector3<RealType> Normal = VectorUtil::Normal<RealType>(Vi, Vj, Vk);
+		FVector3<RealType> Perp0, Perp1;
+		VectorUtil::MakePerpVectors<RealType>(Normal, Perp0, Perp1);
+		// project points to triangle plane coordinates
+		FVector2<RealType> vi(Vi.Dot(Perp0), Vi.Dot(Perp1));
+		FVector2<RealType> vj(Vj.Dot(Perp0), Vj.Dot(Perp1));
+		FVector2<RealType> vk(Vk.Dot(Perp0), Vk.Dot(Perp1));
+		// calculate gradient
+		FVector2<RealType> GradX = (fj-fi)*PerpCW(vi-vk) + (fk-fi)*PerpCW(vj-vi);
+		// map back to 3D vector in triangle plane
+		RealType AreaScale = (RealType)1 / ((RealType)2 * VectorUtil::Area<RealType>(Vi, Vj, Vk));
+		return AreaScale * (GradX.X * Perp0 + GradX.Y * Perp1);
+	}
+
+
+
+	/**
+	 * @return angle between vectors (A-CornerPt) and (B-CornerPt)
+	 */
+	template<typename RealType>
+	inline RealType OpeningAngleD(UE::Math::TVector<RealType> A, UE::Math::TVector<RealType> B, const UE::Math::TVector<RealType>& P)
+	{
+		A -= P; 
+		Normalize(A);
+		B -= P;
+		Normalize(B);
+		return AngleD(A, B);
+	}
+
+
+
+	/**
+	 * @return sign of Bitangent relative to Normal and Tangent
+	 */
+	template<typename RealType>
+	inline RealType BitangentSign(const UE::Math::TVector<RealType>& NormalIn, const UE::Math::TVector<RealType>& TangentIn, const UE::Math::TVector<RealType>& BitangentIn)
+	{
+		// following math from RenderUtils.h::GetBasisDeterminantSign()
+		RealType Cross00 = BitangentIn.Y*NormalIn.Z - BitangentIn.Z*NormalIn.Y;
+		RealType Cross10 = BitangentIn.Z*NormalIn.X - BitangentIn.X*NormalIn.Z;
+		RealType Cross20 = BitangentIn.X*NormalIn.Y - BitangentIn.Y*NormalIn.X;
+		RealType Determinant = TangentIn.X*Cross00 + TangentIn.Y*Cross10 + TangentIn.Z*Cross20;
+		return (Determinant < 0) ? (RealType)-1 : (RealType)1;
+	}
+
+	/**
+	 * @return Bitangent vector based on given Normal, Tangent, and Sign value (+1/-1)
+	 */
+	template<typename RealType>
+	inline FVector3<RealType> Bitangent(const UE::Math::TVector<RealType>& NormalIn, const UE::Math::TVector<RealType>& TangentIn, RealType BitangentSign)
+	{
+		return BitangentSign * FVector3<RealType>(
+			NormalIn.Y*TangentIn.Z - NormalIn.Z*TangentIn.Y,
+			NormalIn.Z*TangentIn.X - NormalIn.X*TangentIn.Z,
+			NormalIn.X*TangentIn.Y - NormalIn.Y*TangentIn.X);
+	}
+
+	/**
+	 * @return Tangent-Space vector based on given Normal and Bitangent
+	 */
+	template<typename RealType>
+	inline FVector3<RealType> TangentFromBitangent(const UE::Math::TVector<RealType>& NormalIn, const UE::Math::TVector<RealType>& BitangentIn)
+	{
+		return BitangentIn.Cross(NormalIn);
+	}
+
+	/**
+	 * @return Bitangent vector based on given Normal and Tangent
+	 */
+	template<typename RealType>
+	inline FVector3<RealType> BitangentFromTangent(const UE::Math::TVector<RealType>& NormalIn, const UE::Math::TVector<RealType>& TangentIn)
+	{
+		return NormalIn.Cross(TangentIn);
+	}
+
+	/// @return Aspect ratio of triangle 
+	inline double AspectRatio(const FVector3d& v1, const FVector3d& v2, const FVector3d& v3)
+	{
+		double a = Distance(v1, v2), b = Distance(v2, v3), c = Distance(v3, v1);
+		double s = (a + b + c) / 2.0;
+		return (a * b * c) / (8.0 * (s - a) * (s - b) * (s - c));
+	}
+
+
+}; // namespace VectorUtil
