Adding parallel transform-reduce

#rnx

diff --git a/Engine/Plugins/Experimental/GeometryProcessing/Source/DynamicMesh/Private/NormalFlowRemesher.cpp b/Engine/Plugins/Experimental/GeometryProcessing/Source/DynamicMesh/Private/NormalFlowRemesher.cpp
--- a/Engine/Plugins/Experimental/GeometryProcessing/Source/DynamicMesh/Private/NormalFlowRemesher.cpp
+++ b/Engine/Plugins/Experimental/GeometryProcessing/Source/DynamicMesh/Private/NormalFlowRemesher.cpp
@@ -3,6 +3,7 @@
 #include "NormalFlowRemesher.h"
 #include "Async/ParallelFor.h"
 #include "InfoTypes.h"
+#include "Async/ParallelTransformReduce.h"
 
 #include "ExplicitUseGeometryMathTypes.h"		// using UE::Geometry::(math types)
 using namespace UE::Geometry;
@@ -361,16 +362,20 @@ void FNormalFlowRemesher::TrackedFaceProjectionPass(double& MaxDistanceMoved, bo
 
 
 	// Return the maximum distance moved by a vertex
-	MaxDistanceMoved = 0.0;
-	for (int VertexID = 0; VertexID < Mesh->MaxVertexID(); ++VertexID)
+	MaxDistanceMoved = ParallelTransformReduce(Mesh->MaxVertexID(), 0.0, [&](int32 VID) -> double
 	{
-		if (TempFlagBuffer[VertexID] && Mesh->IsVertex(VertexID))
+		if (TempFlagBuffer[VID] && Mesh->IsVertex(VID))
 		{
-			const FVector3d& CurrentPosition = Mesh->GetVertex(VertexID);
-			const FVector3d& ProjectedPosition = TempPosBuffer[VertexID];
-			MaxDistanceMoved = FMath::Max(MaxDistanceMoved, Distance(CurrentPosition, ProjectedPosition));
+			const FVector3d& CurrentPosition = Mesh->GetVertex(VID);
+			return Distance(CurrentPosition, TempPosBuffer[VID]);
 		}
-	}
+		return 0.0;
+	},
+													   [](double A, double B) -> double
+	{
+		return FMath::Max(A, B);
+	},
+		32);
 
 	// Update vertices
 	constexpr bool bUpdateParallel = true;
diff --git a/Engine/Source/Runtime/GeometryCore/Public/Async/ParallelTransformReduce.h b/Engine/Source/Runtime/GeometryCore/Public/Async/ParallelTransformReduce.h
new file mode 100644
--- /dev/null
+++ b/Engine/Source/Runtime/GeometryCore/Public/Async/ParallelTransformReduce.h
@@ -0,0 +1,230 @@
+// Copyright Epic Games, Inc. All Rights Reserved.
+
+#pragma once
+
+#include "Async/ParallelFor.h"
+#include "MathUtil.h"
+
+namespace UE {
+namespace Geometry {
+
+	//
+	// Index interface (similar to ParallelFor)
+	// 	   TransformFuncT should be a function pointer-like object with signature: T(int)
+	// 	   ReduceFuncT should be a function pointer-like object with signature: T(T,T)
+
+	template<typename T, typename TransformFuncT, typename ReduceFuncT>
+	T ParallelTransformReduce(int64 Num,
+							  const T& Init,
+							  TransformFuncT Transform,
+							  ReduceFuncT Reduce,
+							  int64 NumTasks)
+	{
+		check(NumTasks > 0);
+		int64 NumPerTask = (int64)FMathf::Ceil((float)Num / NumTasks);
+
+		TArray<T> PerTaskResults;
+		PerTaskResults.SetNum(NumTasks);
+
+		ParallelFor(NumTasks, [&](int64 TaskIndex)
+		{
+			T LocalResult{ Init };
+
+			int64 End = FMath::Min((TaskIndex + 1) * NumPerTask, Num);
+			for (int64 Index = TaskIndex * NumPerTask; Index < End; ++Index)
+			{
+				T Transformed{ Transform(Index) };
+				LocalResult = Reduce(Transformed, LocalResult);
+			}
+
+			PerTaskResults[TaskIndex] = LocalResult;
+		});
+
+		T FinalResult{ Init };
+		for (const T& PartialResult : PerTaskResults)
+		{
+			FinalResult = Reduce(PartialResult, FinalResult);
+		}
+		return FinalResult;
+	}
+
+
+	//
+	// Index interface for non-copyable types
+	// 	   InitFuncT should be a function-like object with signature void(T&)
+	// 	   TransformFuncT should be a function-like object with signature void(int,T&)
+	// 	   ReduceFuncT should be a function-like object with signature void(T,T&)
+
+	template<typename T, typename InitFuncT, typename TransformFuncT, typename ReduceFuncT>
+	void ParallelTransformReduce(int64 Num,
+								 InitFuncT InitFunc,
+								 TransformFuncT Transform,
+								 ReduceFuncT Reduce,
+								 T& Out,
+								 int64 NumTasks)
+	{
+		check(NumTasks > 0);
+		int64 NumPerTask = (int64)FMathf::Ceil((float)Num / NumTasks);
+
+		TArray<T> PerTaskResults;
+		PerTaskResults.SetNum(NumTasks);
+
+		ParallelFor(NumTasks, [&](int64 TaskIndex)
+		{
+			InitFunc(PerTaskResults[TaskIndex]);
+
+			int64 End = FMath::Min((TaskIndex + 1) * NumPerTask, Num);
+			for (int64 Index = TaskIndex * NumPerTask; Index < End; ++Index)
+			{
+				T Transformed;
+				Transform(Index, Transformed);
+				Reduce(Transformed, PerTaskResults[TaskIndex]);
+			}
+		});
+
+		InitFunc(Out);
+		for (const T& PartialResult : PerTaskResults)
+		{
+			Reduce(PartialResult, Out);
+		}
+	}
+
+
+	//
+	// Iterator interface helpers
+	//
+
+	// TODO: Specialize this for for random-access iterators
+	template<typename IterT>
+	int64 IteratorDistance(IterT Begin, IterT End)
+	{
+		int64 Count = 0;
+		for (IterT Iter = Begin; Iter != End; ++Iter)
+		{
+			++Count;
+		}
+		return Count;
+	}
+
+
+	// TODO: Specialize this for random-access iterators
+	template<typename IterT>
+	IterT AdvanceIterator(IterT Start, int64 N)
+	{
+		IterT I = Start;
+		for (int64 Inc = 0; Inc < N; ++Inc)
+		{
+			++I;
+		}
+		return I;
+	}
+
+	template<typename ContainerType, typename ElementType, typename SizeType>
+	int64 IteratorDistance(TIndexedContainerIterator<ContainerType, ElementType, SizeType> Begin,
+						   TIndexedContainerIterator<ContainerType, ElementType, SizeType> End)
+	{
+		return End.GetIndex() - Begin.GetIndex();
+	}
+
+	template<typename ContainerType, typename ElementType, typename SizeType>
+	TIndexedContainerIterator<ContainerType, ElementType, SizeType>
+		AdvanceIterator(TIndexedContainerIterator<ContainerType, ElementType, SizeType> Begin,
+						int64 N)
+	{
+		return Begin + N;
+	}
+
+	// Iterator interface
+	// 	   TransformFuncT should be a function-like object with signature: T(U), where U is the type referred to by IterT
+	// 	   ReduceFuncT should be a function-like object with signature: T(T,T)
+
+	template<typename T, typename IterT, typename TransformFuncT, typename ReduceFuncT>
+	T ParallelTransformReduce(IterT BeginIterator,
+							  IterT EndIterator,
+							  const T& Init,
+							  TransformFuncT Transform,
+							  ReduceFuncT Reduce,
+							  int64 NumTasks)
+	{
+		int64 Num = IteratorDistance(BeginIterator, EndIterator);
+
+		check(NumTasks > 0);
+		int64 NumPerTask = (int64)FMathf::Ceil((float)Num / NumTasks);
+
+		TArray<T> PerTaskResults;
+		PerTaskResults.SetNum(NumTasks);
+
+		ParallelFor(NumTasks, [&](int64 TaskIndex)
+		{
+			T LocalResult{ Init };
+
+			IterT LocalIter = AdvanceIterator(BeginIterator, FMath::Min(Num, TaskIndex * NumPerTask));
+			IterT LocalEndIter = AdvanceIterator(BeginIterator, FMath::Min(Num, (TaskIndex + 1) * NumPerTask));
+
+			while (LocalIter != LocalEndIter)
+			{
+				T Transformed{ Transform(*LocalIter) };
+				LocalResult = Reduce(Transformed, LocalResult);
+
+				++LocalIter;
+			}
+
+			PerTaskResults[TaskIndex] = LocalResult;
+		});
+
+		T FinalResult{ Init };
+		for (const T& PartialResult : PerTaskResults)
+		{
+			FinalResult = Reduce(PartialResult, FinalResult);
+		}
+		return FinalResult;
+	}
+
+	//
+	// Iterator interface for non-copyable types
+	// 	   InitFuncT should be a function-like object with signature: void(T&)
+	// 	   TransformFuncT should be a function-like object with signature: void(U,T&), where U is the type referred to by IterT
+	// 	   ReduceFuncT should be a function-like object with signature: void(T,T&)
+
+	template<typename T, typename IterT, typename InitFuncT, typename TransformFuncT, typename ReduceFuncT>
+	void ParallelTransformReduce(IterT BeginIterator,
+								 IterT EndIterator,
+								 InitFuncT InitFunc,
+								 TransformFuncT Transform,
+								 ReduceFuncT Reduce,
+								 T& Out,
+								 int64 NumTasks)
+	{
+		int64 Num = IteratorDistance(BeginIterator, EndIterator);
+
+		check(NumTasks > 0);
+		int64 NumPerTask = (int64)FMathf::Ceil((float)Num / NumTasks);
+
+		TArray<T> PerTaskResults;
+		PerTaskResults.SetNum(NumTasks);
+
+		ParallelFor(NumTasks, [&](int64 TaskIndex)
+		{
+			InitFunc(PerTaskResults[TaskIndex]);
+
+			IterT LocalIter = AdvanceIterator(BeginIterator, FMath::Min(Num, TaskIndex * NumPerTask));
+			IterT LocalEndIter = AdvanceIterator(BeginIterator, FMath::Min(Num, (TaskIndex + 1) * NumPerTask));
+
+			while (LocalIter != LocalEndIter)
+			{
+				T Transformed;
+				Transform(*LocalIter, Transformed);
+				Reduce(Transformed, PerTaskResults[TaskIndex]);
+				++LocalIter;
+			}
+		});
+
+		InitFunc(Out);
+		for (const T& PartialResult : PerTaskResults)
+		{
+			Reduce(PartialResult, Out);
+		}
+	}
+
+} // end namespace UE::Geometry
+} // end namespace UE
