make_level_set3_parallel_floodfill()
Array access overflow - Overflowing NumSections (256) when accessing GridSections/SectionQueues

diff --git a/Engine/Source/Runtime/GeometryCore/Public/Implicit/SweepingMeshSDF.h b/Engine/Source/Runtime/GeometryCore/Public/Implicit/SweepingMeshSDF.h
--- a/Engine/Source/Runtime/GeometryCore/Public/Implicit/SweepingMeshSDF.h
+++ b/Engine/Source/Runtime/GeometryCore/Public/Implicit/SweepingMeshSDF.h
@@ -741,9 +741,10 @@ private:
 		// chunks for locking
 		int32 NumSections = 256;		// a bit arbitrary, stopped seeing improvement on a 64-core machine at this point
 		TArray<FCriticalSection> GridSections;
-		GridSections.SetNum(NumSections + 1);
+		GridSections.SetNum(NumSections);
 		int64 TotalGridCellCount = NI * NJ * NK;
-		int64 SectionSize = TotalGridCellCount / (int)NumSections;
+		int64 SectionSize = FMath::CeilToInt(TotalGridCellCount / (float)NumSections);
+		
 		// this returns the FCriticalSection to use for the given span of values
 		auto GetGridSectionLock = [this, SectionSize, &Distances, &GridSections](FVector3i CellGridIndex) -> FCriticalSection*
 		{
@@ -754,7 +755,7 @@ private:
 
 		// per-grid-chunk queue, each one of these cooresponds to a GridSections lock
 		TArray<TArray<int>> SectionQueues;
-		SectionQueues.SetNum(NumSections + 1);
+		SectionQueues.SetNum(NumSections);
 		auto AddToSectionQueue = [SectionSize, &SectionQueues](int64 CellLinearIndex)
 		{
 			int64 SectionIndex = CellLinearIndex / SectionSize;
