GeometryProcessing: add FDynamicMeshTangents helper class to simplify access to FDynamicMesh3 tangent attribute overlays. Add functions to FMeshTangents and TDynamicMeshVectorOverlay to simplify getting a single triangle-vertex tangent/bitangent. Fix paste-o in FDynamicMeshAttributeSet::SetNumNormalLayers()
#rnx

diff --git a/Engine/Source/Runtime/GeometryCore/Private/DynamicMesh/DynamicMeshAttributeSet.cpp b/Engine/Source/Runtime/GeometryCore/Private/DynamicMesh/DynamicMeshAttributeSet.cpp
--- a/Engine/Source/Runtime/GeometryCore/Private/DynamicMesh/DynamicMeshAttributeSet.cpp
+++ b/Engine/Source/Runtime/GeometryCore/Private/DynamicMesh/DynamicMeshAttributeSet.cpp
@@ -400,7 +400,7 @@ void FDynamicMeshAttributeSet::SetNumNormalLayers(int Num)
 	}
 	else
 	{
-		NormalLayers.RemoveAt(Num, UVLayers.Num() - Num);
+		NormalLayers.RemoveAt(Num, NormalLayers.Num() - Num);
 	}
 	ensure(NormalLayers.Num() == Num);
 }
diff --git a/Engine/Source/Runtime/GeometryCore/Private/DynamicMesh/MeshTangents.cpp b/Engine/Source/Runtime/GeometryCore/Private/DynamicMesh/MeshTangents.cpp
--- a/Engine/Source/Runtime/GeometryCore/Private/DynamicMesh/MeshTangents.cpp
+++ b/Engine/Source/Runtime/GeometryCore/Private/DynamicMesh/MeshTangents.cpp
@@ -1,11 +1,85 @@
 // Copyright Epic Games, Inc. All Rights Reserved.
 
 #include "DynamicMesh/MeshTangents.h"
+#include "VectorUtil.h"
 #include "Async/ParallelFor.h"
 
 #include "ExplicitUseGeometryMathTypes.h"		// using UE::Geometry::(math types)
 using namespace UE::Geometry;
 
+
+
+FDynamicMeshTangents::FDynamicMeshTangents(const FDynamicMesh3* MeshIn)
+{
+	Mesh = MeshIn;
+	if (ensure(Mesh) && ensure(Mesh->HasAttributes()))
+	{
+		Normals = Mesh->Attributes()->PrimaryNormals();
+		if (ensure(Normals) && Mesh->Attributes()->HasTangentSpace() )
+		{
+			Tangents = Mesh->Attributes()->PrimaryTangents();
+			Bitangents = Mesh->Attributes()->PrimaryBiTangents();
+		}
+	}
+}
+
+
+void FDynamicMeshTangents::GetTangentFrame(int32 TriangleID, int32 TriVertexIndex, FVector3f& Normal, FVector3f& Tangent, FVector3f& Bitangent) const
+{
+	if (Normals)
+	{
+		Normals->GetTriElement(TriangleID, TriVertexIndex, Normal);
+		if (Tangents)
+		{
+			Tangents->GetTriElement(TriangleID, TriVertexIndex, Tangent);
+			Bitangents->GetTriElement(TriangleID, TriVertexIndex, Bitangent);
+		}
+		else
+		{
+			VectorUtil::MakePerpVectors(Normal, Tangent, Bitangent);
+		}
+	}
+	else
+	{
+		Normal = FVector3f::UnitZ();
+		Tangent = FVector3f::UnitX();
+		Bitangent = FVector3f::UnitY();
+	}
+}
+
+
+void FDynamicMeshTangents::GetTangentVectors(int32 TriangleID, int32 TriVertexIndex, const FVector3f& Normal, FVector3f& Tangent, FVector3f& Bitangent) const
+{
+	if (Tangents)
+	{
+		Tangents->GetTriElement(TriangleID, TriVertexIndex, Tangent);
+		Bitangents->GetTriElement(TriangleID, TriVertexIndex, Bitangent);
+	}
+	else
+	{
+		VectorUtil::MakePerpVectors(Normal, Tangent, Bitangent);
+	}
+}
+
+
+
+
+void FDynamicMeshTangents::GetTangentVectors(int32 TriangleID, int32 TriVertexIndex, FVector3f& Tangent, FVector3f& Bitangent) const
+{
+	if (Tangents)
+	{
+		Tangents->GetTriElement(TriangleID, TriVertexIndex, Tangent);
+		Bitangents->GetTriElement(TriangleID, TriVertexIndex, Bitangent);
+	}
+	else
+	{
+		Tangent = FVector3f::UnitX();
+		Bitangent = FVector3f::UnitY();
+	}
+}
+
+
+
 template<typename RealType>
 void TMeshTangents<RealType>::SetTangentCount(int Count, bool bClearToZero)
 {
diff --git a/Engine/Source/Runtime/GeometryCore/Public/DynamicMesh/DynamicMeshOverlay.h b/Engine/Source/Runtime/GeometryCore/Public/DynamicMesh/DynamicMeshOverlay.h
--- a/Engine/Source/Runtime/GeometryCore/Public/DynamicMesh/DynamicMeshOverlay.h
+++ b/Engine/Source/Runtime/GeometryCore/Public/DynamicMesh/DynamicMeshOverlay.h
@@ -573,6 +573,16 @@ public:
 		BaseType::GetElement(ElementID, V);
 	}
 
+	/**
+	 * Get the Element associated with a vertex of a triangle
+	 * @param TriVertexIndex index of vertex in triangle, valid values are 0,1,2
+	 */
+	inline void GetTriElement(int TriangleID, int32 TriVertexIndex, VectorType& Value) const
+	{
+		checkSlow(TriVertexIndex >= 0 && TriVertexIndex <= 2);
+		GetElement(BaseType::ElementTriangles[(3 * TriangleID) + TriVertexIndex], Value);
+	}
+
 	/**
 	 * Get the three Elements associated with a triangle
 	 */
diff --git a/Engine/Source/Runtime/GeometryCore/Public/DynamicMesh/MeshTangents.h b/Engine/Source/Runtime/GeometryCore/Public/DynamicMesh/MeshTangents.h
--- a/Engine/Source/Runtime/GeometryCore/Public/DynamicMesh/MeshTangents.h
+++ b/Engine/Source/Runtime/GeometryCore/Public/DynamicMesh/MeshTangents.h
@@ -11,6 +11,36 @@ namespace UE
 namespace Geometry
 {
 
+/**
+ * FDynamicMeshTangents is a helper object for accessing tangents stored in the AttributeSet of a FDynamicMesh3.
+ */
+class GEOMETRYCORE_API FDynamicMeshTangents
+{
+public:
+	const FDynamicMesh3* Mesh = nullptr;
+	const FDynamicMeshNormalOverlay* Normals = nullptr;
+	const FDynamicMeshNormalOverlay* Tangents = nullptr;
+	const FDynamicMeshNormalOverlay* Bitangents = nullptr;
+
+	FDynamicMeshTangents(const FDynamicMesh3* MeshIn);
+
+	/** @return true if the mesh has valid tangents */
+	bool HasValidTangents() const { return Tangents != nullptr && Bitangents != nullptr; }
+
+	/**
+	 * If tangents are available in the overlays, returns them. If only Normal is available, computes orthogonal basis. Falls back to unit axes if no overlays are available.
+	 */
+	void GetTangentFrame(int32 TriangleID, int32 TriVertexIndex, FVector3f& NormalOut, FVector3f& TangentOut, FVector3f& BitangentOut) const;
+	/**
+	 * If tangents are available in the overlays, returns them. Otherwise computes orthogonal basis to Normal argument.
+	 */
+	void GetTangentVectors(int32 TriangleID, int32 TriVertexIndex, const FVector3f& Normal, FVector3f& TangentOut, FVector3f& BitangentOut) const;
+	/**
+	 * If tangents are available in the overlays, returns them, otherwise falls back to unit X/Y axes.
+	 */
+	void GetTangentVectors(int32 TriangleID, int32 TriVertexIndex, FVector3f& TangentOut, FVector3f& BitangentOut) const;
+};
+
 
 /**
  * Options used by TMeshTangents for tangents computation
@@ -169,6 +199,19 @@ public:
 		Bitangents[k] = Bitangent;
 	}
 
+	/**
+	 * Get tangent and bitangent at a vertex of a triangle for per-triangle computed tangents
+	 * @param TriangleID triangle index in mesh
+	 * @param TriVertIdx vertex index in range 0,1,2
+	 */
+	template<typename OtherVectorType>
+	inline void GetTriangleVertexTangentVectors(int32 TriangleID, int32 TriVertexIndex, OtherVectorType& TangentOut, OtherVectorType& BitangentOut) const
+	{
+		int k = TriangleID * 3 + TriVertexIndex;
+		TangentOut = (OtherVectorType)Tangents[k];
+		BitangentOut = (OtherVectorType)Bitangents[k];
+	}
+
 
 	//
 	// Per-Triangle Tangents
