Optimize occupancy map construction for large textures by removing lock

#rnx
#pf 60ca206d4ce02c0001885e5f

diff --git a/Engine/Source/Runtime/GeometryCore/Public/Image/ImageOccupancyMap.h b/Engine/Source/Runtime/GeometryCore/Public/Image/ImageOccupancyMap.h
--- a/Engine/Source/Runtime/GeometryCore/Public/Image/ImageOccupancyMap.h
+++ b/Engine/Source/Runtime/GeometryCore/Public/Image/ImageOccupancyMap.h
@@ -109,16 +109,14 @@ public:
 		QueryOptions.MaxDistance = (double)GutterSize * TexelDiag;
 
 		// find interior texels
-		TAtomic<int64> InteriorCounter;
-		FCriticalSection GutterLock;
-		ParallelFor(Dimensions.GetHeight(), 
-			[this, &UVSpaceMesh, &GetTriangleIDFunc, &FlatSpatial, TexelDiag, &QueryOptions, &InteriorCounter, &GutterLock, &TexelSize]
+		TAtomic<int64> TotalGutterCounter = 0;
+		TArray< TArray64<TTuple<int64, int64>> > GutterTexelsPerScanline;
+		GutterTexelsPerScanline.SetNum(Dimensions.GetHeight());
+
+		ParallelFor(Dimensions.GetHeight(),
+					[this, &UVSpaceMesh, &GetTriangleIDFunc, &FlatSpatial, TexelDiag, &QueryOptions, &TexelSize, &TotalGutterCounter, &GutterTexelsPerScanline]
 		(int32 ImgY)
 		{
-			// To minimize locks, accumulate gutter texel data locally and append
-			// to our gutter texel array at the end.
-			TArray64<TTuple<int64, int64>> LocalGutterTexels;
-			int64 LocalInteriorCounter = 0;
 			for (int32 ImgX = 0; ImgX < Dimensions.GetWidth(); ++ImgX)
 			{
 				// With multiple samples, a texel may contain multiple gutter samples.
@@ -151,7 +149,6 @@ public:
 							TexelQueryUV[SampleLinearIdx] = (FVector2f)UVPoint;
 							TexelQueryTriangle[SampleLinearIdx] = GetTriangleIDFunc(NearestTriID);
 							++TexelInteriorSamples[TexelLinearIdx];
-							++LocalInteriorCounter;
 							bIsGutterTexel = false;
 						}
 						else if (NearDistSqr < TexelDiag * TexelDiag)
@@ -165,7 +162,6 @@ public:
 							TexelQueryUV[SampleLinearIdx] = (FVector2f)NearestUV;
 							TexelQueryTriangle[SampleLinearIdx] = GetTriangleIDFunc(NearestTriID);
 							++TexelInteriorSamples[TexelLinearIdx];
-							++LocalInteriorCounter;
 							bIsGutterTexel = false;
 						}
 						else
@@ -193,19 +189,22 @@ public:
 					}
 					if (bIsGutterTexel)
 					{
-						LocalGutterTexels.Add(GutterNearestTexel);
+						GutterTexelsPerScanline[ImgY].Add(GutterNearestTexel);
 					}
 				}
 			}
-			if (LocalGutterTexels.Num() > 0)
-			{
-				GutterLock.Lock();
-				GutterTexels += LocalGutterTexels;
-				GutterLock.Unlock();
-			}
-			InteriorCounter += LocalInteriorCounter;
+
+			TotalGutterCounter += GutterTexelsPerScanline[ImgY].Num();
+
 		}, !bParallel ? EParallelForFlags::ForceSingleThread : EParallelForFlags::None);
 
+		// Now build GutterTexels
+		GutterTexels.Reserve(TotalGutterCounter);
+		for (const TArray64<TTuple<int64, int64>>& LineGutterTexels : GutterTexelsPerScanline)
+		{
+			GutterTexels.Append(LineGutterTexels);
+		}
+
 		return true;
 	}
 
