Parallelize expansion coefficient computation in fast winding number tree construction

#rnx

diff --git a/Engine/Source/Runtime/GeometryCore/Public/Spatial/FastWinding.h b/Engine/Source/Runtime/GeometryCore/Public/Spatial/FastWinding.h
--- a/Engine/Source/Runtime/GeometryCore/Public/Spatial/FastWinding.h
+++ b/Engine/Source/Runtime/GeometryCore/Public/Spatial/FastWinding.h
@@ -7,6 +7,7 @@
 #include "Spatial/MeshAABBTree3.h"
 #include "Util/MeshCaches.h"
 #include "MatrixTypes.h"
+#include "Async/ParallelTransformReduce.h"
 
 /**
  *  Formulas for triangle winding number approximation
@@ -18,7 +19,7 @@ namespace FastTriWinding
 	using FMatrix3d = UE::Geometry::FMatrix3d;
 
 	/**
-	 *  precompute constant coefficients of triangle winding number approximation
+	 *  precompute constant coefficients of triangle winding number approximation (serial implementation)
 	 *  P: 'Center' of expansion for Triangles (area-weighted centroid avg)
 	 *  R: max distance from P to Triangles
 	 *  Order1: first-order vector coeff
@@ -26,11 +27,12 @@ namespace FastTriWinding
 	 *  TriCache: precomputed triangle centroid/normal/area (todo @jimmy either support passing in as possibly-nullptr or remove commented-out null branches!)
 	 */
 	template <class TriangleMeshType, class IterableTriangleIndices>
-	void ComputeCoeffs(const TriangleMeshType& Mesh,
-					   const IterableTriangleIndices& Triangles,
-					   const FMeshTriInfoCache& TriCache,
-					   FVector3d& P, double& R,
-					   FVector3d& Order1, FMatrix3d& Order2)
+	void ComputeCoeffsSerial(const TriangleMeshType& Mesh,
+							 const IterableTriangleIndices& Triangles,
+							 const FMeshTriInfoCache& TriCache,
+							 FVector3d& P, double& R,
+							 FVector3d& Order1,
+							 FMatrix3d& Order2)
 	{
 		P = FVector3d::Zero();
 		Order1 = FVector3d::Zero();
@@ -87,6 +89,124 @@ namespace FastTriWinding
 		}
 	}
 
+	/**
+	 *  precompute constant coefficients of triangle winding number approximation (evaluated in parallel for large sets of triangles)
+	 *  P: 'Center' of expansion for Triangles (area-weighted centroid avg)
+	 *  R: max distance from P to Triangles
+	 *  Order1: first-order vector coeff
+	 *  Order2: second-order matrix coeff
+	 *  TriCache: precomputed triangle centroid/normal/area (todo @jimmy either support passing in as possibly-nullptr or remove commented-out null branches!)
+	 */
+	template <class TriangleMeshType>
+	void ComputeCoeffs(const TriangleMeshType& Mesh,
+					   const TSet<int>& TriangleSet,
+					   const FMeshTriInfoCache& TriCache,
+					   FVector3d& P,
+					   double& R,
+					   FVector3d& Order1,
+					   FMatrix3d& Order2,
+					   int NumTasks = 16)
+	{
+		// If the data is small enough, don't bother trying to parallelize
+		constexpr int ParallelThreshold = 3000;			// Benchmarking shows parallel starts to outperform serial at around 2500 triangles or so
+		if (TriangleSet.Num() < ParallelThreshold)
+		{
+			ComputeCoeffsSerial(Mesh, TriangleSet, TriCache, P, R, Order1, Order2);
+			return;
+		}
+
+		TArray<int> TriangleArray = TriangleSet.Array();
+
+		// First compute the area-weighted centroid
+
+		struct PData
+		{
+			FVector3d P;
+			double Area;
+		};
+
+		auto CentroidTransform = [&](int64 TriangleSubsetIndex) -> PData
+		{
+			int tid = TriangleArray[TriangleSubsetIndex];
+			FVector3d P0, P1, P2;
+			Mesh.GetTriVertices(tid, P0, P1, P2);
+
+			PData DataOut;
+			DataOut.Area = VectorUtil::Area(P0, P1, P2);
+			DataOut.P = DataOut.Area * ((P0 + P1 + P2) / 3.0);
+			return DataOut;
+		};
+
+		auto CentroidReduce = [](const PData& A, const PData& B) -> PData
+		{
+			PData Out;
+			Out.P = A.P + B.P;
+			Out.Area = A.Area + B.Area;
+			return Out;
+		};
+
+		PData CentroidData = ParallelTransformReduce(TriangleArray.Num(), PData{ FVector3d::Zero(), 0.0}, CentroidTransform, CentroidReduce, NumTasks);
+
+		CentroidData.P /= CentroidData.Area;
+
+		// Now compute the expansions
+
+		struct OrderData
+		{
+			FVector3d Order1;
+			FMatrix3d Order2;
+			double R;
+		};
+
+		auto ExpansionTransform = [&, P = CentroidData.P](int TriangleSubsetIndex) -> OrderData
+		{
+			int tid = TriangleArray[TriangleSubsetIndex];
+			FVector3d P0, P1, P2;
+			Mesh.GetTriVertices(tid, P0, P1, P2);
+
+			FVector3d n, c;
+			double a;
+			TriCache.GetTriInfo(tid, n, a, c);
+
+			OrderData DataOut;
+			DataOut.Order1 = a * n;
+
+			FVector3d dcp = c - P;
+			DataOut.Order2 = a * FMatrix3d(dcp, n);
+
+			// this is just for return value...
+			double maxdist = FMath::Max3(DistanceSquared(P0, P), DistanceSquared(P1, P), DistanceSquared(P2, P));
+			DataOut.R = sqrt(maxdist);
+
+			return DataOut;
+		};
+
+		auto ExpansionReduce = [](const OrderData& A, const OrderData& B) -> OrderData
+		{
+			OrderData Out;
+			Out.Order1 = A.Order1 + B.Order1;
+			Out.Order2 = A.Order2 + B.Order2;
+			Out.R = FMath::Max(A.R, B.R);
+			return Out;
+		};
+
+		OrderData Orders = ParallelTransformReduce(TriangleArray.Num(), 
+												   OrderData{ FVector3d::Zero(), FMatrix3d::Zero(), 0.0 }, 
+												   ExpansionTransform, 
+												   ExpansionReduce, 
+												   NumTasks);
+
+		// Set out values
+
+		P = CentroidData.P;
+		R = Orders.R;
+		Order1 = Orders.Order1;
+		Order2 = Orders.Order2;
+	}
+
+
+
+
 	/**
 	 *  Evaluate first-order FWN approximation at point Q, relative to Center c
 	 */
