GeometryProcessing: Reset existing UVs before Unwrap

#rnx

diff --git a/Engine/Plugins/Experimental/GeometryProcessing/Source/DynamicMesh/Private/Parameterization/MeshUVTransforms.cpp b/Engine/Plugins/Experimental/GeometryProcessing/Source/DynamicMesh/Private/Parameterization/MeshUVTransforms.cpp
--- a/Engine/Plugins/Experimental/GeometryProcessing/Source/DynamicMesh/Private/Parameterization/MeshUVTransforms.cpp
+++ b/Engine/Plugins/Experimental/GeometryProcessing/Source/DynamicMesh/Private/Parameterization/MeshUVTransforms.cpp
@@ -142,10 +142,7 @@ void UE::MeshUVTransforms::MakeSeamsDisjoint(FDynamicMeshUVOverlay* UVOverlay)
 	for (int32 elemid : UVOverlay->ElementIndicesItr())
 	{
 		int32 ParentVtxID = UVOverlay->GetParentVertex(elemid);
-		if (ParentVtxID == -1)
-		{
-			continue;
-		}
+		checkSlow(ParentVtxID != FDynamicMesh3::InvalidID);
 		int32 SeamListIndex = SeamVerticesMap[ParentVtxID];
 		if (SeamListIndex == -1)
 		{
diff --git a/Engine/Plugins/Experimental/MeshModelingToolset/Source/ModelingOperatorsEditorOnly/Private/ParameterizationOps/ParameterizeMeshOp.cpp b/Engine/Plugins/Experimental/MeshModelingToolset/Source/ModelingOperatorsEditorOnly/Private/ParameterizationOps/ParameterizeMeshOp.cpp
--- a/Engine/Plugins/Experimental/MeshModelingToolset/Source/ModelingOperatorsEditorOnly/Private/ParameterizationOps/ParameterizeMeshOp.cpp
+++ b/Engine/Plugins/Experimental/MeshModelingToolset/Source/ModelingOperatorsEditorOnly/Private/ParameterizationOps/ParameterizeMeshOp.cpp
@@ -306,6 +306,9 @@ bool FParameterizeMeshOp::ComputeUVs_ExpMap(FDynamicMesh3& Mesh, TFunction<bool(
 {
 	FDynamicMeshUVEditor UVEditor(&Mesh, UVLayer, true);
 
+	// Delete existing UVs, if there are any.
+	UVEditor.ResetUVs();
+	
 	TArray<int32> AllTriangles;
 	for (int32 tid : Mesh.TriangleIndicesItr())
 	{
@@ -324,6 +327,9 @@ bool FParameterizeMeshOp::ComputeUVs_ConformalFreeBoundary(FDynamicMesh3& InOutM
 {
 	FDynamicMeshUVEditor UVEditor(&InOutMesh, UVLayer, true);
 
+	// Delete existing UVs, if there are any.
+	UVEditor.ResetUVs();
+	
 	TArray<int32> AllTriangles;
 	for (int32 tid : InOutMesh.TriangleIndicesItr())
 	{
diff --git a/Engine/Source/Runtime/GeometryCore/Private/DynamicMesh/DynamicMeshOverlay.cpp b/Engine/Source/Runtime/GeometryCore/Private/DynamicMesh/DynamicMeshOverlay.cpp
--- a/Engine/Source/Runtime/GeometryCore/Private/DynamicMesh/DynamicMeshOverlay.cpp
+++ b/Engine/Source/Runtime/GeometryCore/Private/DynamicMesh/DynamicMeshOverlay.cpp
@@ -240,7 +240,7 @@ int TDynamicMeshOverlay<RealType, ElementSize>::SplitElementWithNewParent(int El
 	if (ElementsRefCounts.GetRefCount(ElementID) == 1) 
 	{
 		ElementsRefCounts.Decrement(ElementID);
-		ParentVertices[ElementID] = -1;
+		ParentVertices[ElementID] = FDynamicMesh3::InvalidID;
 	}
 
 	return NewElID;
@@ -386,7 +386,7 @@ template<typename RealType, int ElementSize>
 void TDynamicMeshOverlay<RealType, ElementSize>::UnsetTriangle(int TriangleID)
 {
 	int i = 3 * TriangleID;
-	if (ElementTriangles[i] == -1)
+	if (ElementTriangles[i] == FDynamicMesh3::InvalidID)
 	{
 		return;
 	}
@@ -397,9 +397,9 @@ void TDynamicMeshOverlay<RealType, ElementSize>::UnsetTriangle(int TriangleID)
 		if (ElementsRefCounts.GetRefCount(ElementTriangles[i + SubIdx]) == 1)
 		{
 			ElementsRefCounts.Decrement(ElementTriangles[i + SubIdx]);
-			ParentVertices[ElementTriangles[i + SubIdx]] = -1;
+			ParentVertices[ElementTriangles[i + SubIdx]] = FDynamicMesh3::InvalidID;
 		}
-		ElementTriangles[i + SubIdx] = -1;
+		ElementTriangles[i + SubIdx] = FDynamicMesh3::InvalidID;
 	}
 }
 
