GeometryProcessing: Move FDynamicMesh3 serialization to a separate file and add local versioning so we can iterate without having to bump stream version
#rnx

diff --git a/Engine/Source/Runtime/GeometryCore/Private/DynamicMesh/DynamicMesh3.cpp b/Engine/Source/Runtime/GeometryCore/Private/DynamicMesh/DynamicMesh3.cpp
--- a/Engine/Source/Runtime/GeometryCore/Private/DynamicMesh/DynamicMesh3.cpp
+++ b/Engine/Source/Runtime/GeometryCore/Private/DynamicMesh/DynamicMesh3.cpp
@@ -4,7 +4,6 @@
 #include "DynamicMesh/DynamicMeshAttributeSet.h"
 #include "Generators/MeshShapeGenerator.h"
 #include "Templates/UniquePtr.h"
-#include "UObject/UE5MainStreamObjectVersion.h"
 
 #include "ExplicitUseGeometryMathTypes.h"		// using UE::Geometry::(math types)
 using namespace UE::Geometry;
@@ -370,35 +369,7 @@ void FDynamicMesh3::Clear()
 	CachedIsClosedTimestamp = -1;
 }
 
-void FDynamicMesh3::Serialize(FArchive& Ar)
-{
-	Ar.UsingCustomVersion(FUE5MainStreamObjectVersion::GUID);
-
-	Ar << Vertices;
-	Ar << VertexRefCounts;
-	Ar << VertexNormals;
-	Ar << VertexColors;
-	Ar << VertexUVs;
-	Ar << VertexEdgeLists;
-	Ar << Triangles;
-	Ar << TriangleRefCounts;
-	Ar << TriangleEdges;
-	Ar << TriangleGroups;
-	Ar << GroupIDCounter;
-	Ar << Edges;
-	Ar << EdgeRefCounts;
-
-	bool bHasAttributes = HasAttributes();
-	Ar << bHasAttributes;
-	if (bHasAttributes)
-	{
-		if (Ar.IsLoading())
-		{
-			EnableAttributes();
-		}
-		Ar << *AttributeSet;
-	}
-}
+
 
 int FDynamicMesh3::GetComponentsFlags() const
 {
@@ -1007,7 +978,8 @@ bool FDynamicMesh3::CheckValidity(FValidityOptions ValidityOptions, EValidityChe
 		{
 			CheckOrFailF(vTris.Num() == GetVtxEdgeCount(vID) || vTris.Num() == GetVtxEdgeCount(vID) - 1);
 		}
-		CheckOrFailF(VertexRefCounts.GetRefCount(vID) == vTris.Num() + 1);
+		int32 VertexRefCount = VertexRefCounts.GetRefCount(vID);
+		CheckOrFailF(VertexRefCount == (vTris.Num() + 1));
 		CheckOrFailF(triToVtxRefs[vID] == vTris.Num());
 		for (int tID : vTris)
 		{
diff --git a/Engine/Source/Runtime/GeometryCore/Private/DynamicMesh/DynamicMesh3_Serialization.cpp b/Engine/Source/Runtime/GeometryCore/Private/DynamicMesh/DynamicMesh3_Serialization.cpp
new file mode 100644
--- /dev/null
+++ b/Engine/Source/Runtime/GeometryCore/Private/DynamicMesh/DynamicMesh3_Serialization.cpp
@@ -0,0 +1,49 @@
+// Copyright Epic Games, Inc. All Rights Reserved.
+
+#include "DynamicMesh/DynamicMesh3.h"
+#include "DynamicMesh/DynamicMeshAttributeSet.h"
+#include "UObject/UE5MainStreamObjectVersion.h"
+
+#include "ExplicitUseGeometryMathTypes.h"		// using UE::Geometry::(math types)
+using namespace UE::Geometry;
+
+enum class EDynamicMeshSerializationVersion
+{
+	InitialVersion = 1
+};
+
+
+void FDynamicMesh3::Serialize(FArchive& Ar)
+{
+	checkSlow(CheckValidity(FValidityOptions(), EValidityCheckFailMode::Ensure));
+
+	Ar.UsingCustomVersion(FUE5MainStreamObjectVersion::GUID);
+
+	int32 SerializationVersion = (int32)EDynamicMeshSerializationVersion::InitialVersion;
+	Ar << SerializationVersion;
+
+	Ar << Vertices;
+	Ar << VertexRefCounts;
+	Ar << VertexNormals;
+	Ar << VertexColors;
+	Ar << VertexUVs;
+	Ar << VertexEdgeLists;
+	Ar << Triangles;
+	Ar << TriangleRefCounts;
+	Ar << TriangleEdges;
+	Ar << TriangleGroups;
+	Ar << GroupIDCounter;
+	Ar << Edges;
+	Ar << EdgeRefCounts;
+
+	bool bHasAttributes = HasAttributes();
+	Ar << bHasAttributes;
+	if (bHasAttributes)
+	{
+		if (Ar.IsLoading())
+		{
+			EnableAttributes();
+		}
+		Ar << *AttributeSet;
+	}
+}
\ No newline at end of file
