ModelingTools: Fix various issues related to Tangents. FDynamicMeshEditor was using hardcoded 0 index instead of loop-variable NormalLayerIndex in several places. FDynamicMeshTangents did not handle invalid overlay triangles. UEditorModelingObjectsCreationAPI::CreateDynamicMeshActor() was not enabling Tangents on new DynamicMeshComponents. UDynamicMeshComponentToolTarget was not transferring tangents between FDynamicMesh3 and FMeshDescription when converting.
#rnx
preflight 60d25e9678c3b000012438cc

diff --git a/Engine/Plugins/Experimental/GeometryProcessing/Source/DynamicMesh/Private/DynamicMeshEditor.cpp b/Engine/Plugins/Experimental/GeometryProcessing/Source/DynamicMesh/Private/DynamicMeshEditor.cpp
--- a/Engine/Plugins/Experimental/GeometryProcessing/Source/DynamicMesh/Private/DynamicMeshEditor.cpp
+++ b/Engine/Plugins/Experimental/GeometryProcessing/Source/DynamicMesh/Private/DynamicMeshEditor.cpp
@@ -1512,7 +1512,7 @@ void FDynamicMeshEditor::AppendMesh(const FDynamicMesh3* AppendMesh,
 			FDynamicMeshNormalOverlay* ToNormals = Mesh->Attributes()->GetNormalLayer(NormalLayerIndex);
 			if (FromNormals != nullptr && ToNormals != nullptr)
 			{
-				FIndexMapi& NormalMap = IndexMapsOut.GetNormalMap(0);
+				FIndexMapi& NormalMap = IndexMapsOut.GetNormalMap(NormalLayerIndex);
 				NormalMap.Reserve(FromNormals->ElementCount());
 				AppendNormals(AppendMesh, FromNormals, ToNormals,
 					VertexMap, TriangleMap, NormalTransform, NormalMap);
@@ -1845,7 +1845,7 @@ static void AppendTriangleAttributes(const FDynamicMesh3* FromMesh, int FromTria
 			for (int j = 0; j < 3; ++j)
 			{
 				check(FromElemTri[j] != FDynamicMesh3::InvalidID);
-				int NewElemID = AppendTriangleNormalAttribute(FromMesh, FromElemTri[j], ToMesh, 0, IndexMaps);
+				int NewElemID = AppendTriangleNormalAttribute(FromMesh, FromElemTri[j], ToMesh, NormalLayerIndex, IndexMaps);
 				ToElemTri[j] = NewElemID;
 			}
 			ToNormalOverlay->SetTriangle(ToTriangleID, ToElemTri);
diff --git a/Engine/Plugins/Experimental/MeshModelingToolset/Source/ModelingComponentsEditorOnly/Private/EditorModelingObjectsCreationAPI.cpp b/Engine/Plugins/Experimental/MeshModelingToolset/Source/ModelingComponentsEditorOnly/Private/EditorModelingObjectsCreationAPI.cpp
--- a/Engine/Plugins/Experimental/MeshModelingToolset/Source/ModelingComponentsEditorOnly/Private/EditorModelingObjectsCreationAPI.cpp
+++ b/Engine/Plugins/Experimental/MeshModelingToolset/Source/ModelingComponentsEditorOnly/Private/EditorModelingObjectsCreationAPI.cpp
@@ -196,6 +196,9 @@ FCreateMeshObjectResult UEditorModelingObjectsCreationAPI::CreateDynamicMeshActo
 
 	UDynamicMeshComponent* NewComponent = NewActor->GetDynamicMeshComponent();
 
+	// assume that DynamicMeshComponent always has tangents on it's internal UDynamicMesh
+	NewComponent->SetTangentsType(EDynamicMeshComponentTangentsMode::ExternallyProvided);
+
 	if (CreateMeshParams.MeshType == ECreateMeshObjectSourceMeshType::DynamicMesh)
 	{
 		NewComponent->SetMesh(MoveTemp(CreateMeshParams.DynamicMesh.GetValue()));
diff --git a/Engine/Plugins/Experimental/MeshModelingToolset/Source/ModelingComponentsEditorOnly/Private/ToolTargets/DynamicMeshComponentToolTarget.cpp b/Engine/Plugins/Experimental/MeshModelingToolset/Source/ModelingComponentsEditorOnly/Private/ToolTargets/DynamicMeshComponentToolTarget.cpp
--- a/Engine/Plugins/Experimental/MeshModelingToolset/Source/ModelingComponentsEditorOnly/Private/ToolTargets/DynamicMeshComponentToolTarget.cpp
+++ b/Engine/Plugins/Experimental/MeshModelingToolset/Source/ModelingComponentsEditorOnly/Private/ToolTargets/DynamicMeshComponentToolTarget.cpp
@@ -106,7 +106,7 @@ const FMeshDescription* UDynamicMeshComponentToolTarget::GetMeshDescription()
 		Attributes.Register();
 
 		FDynamicMeshToMeshDescription Converter;
-		Converter.Convert(DynamicMeshComponent->GetMesh(), *ConvertedMeshDescription);
+		Converter.Convert(DynamicMeshComponent->GetMesh(), *ConvertedMeshDescription, true);
 
 		bHaveMeshDescription = true;
 		return ConvertedMeshDescription.Get();
@@ -130,7 +130,7 @@ void UDynamicMeshComponentToolTarget::CommitMeshDescription(const FCommitter& Co
 	FMeshDescriptionToDynamicMesh Converter;
 	TSharedPtr<FDynamicMesh3> NewMeshShared = MakeShared<FDynamicMesh3>();
 	NewMeshShared->EnableAttributes();
-	Converter.Convert(CommitterParams.MeshDescriptionOut, *NewMeshShared);
+	Converter.Convert(CommitterParams.MeshDescriptionOut, *NewMeshShared, true);
 
 	DynamicMesh->EditMesh([&](FDynamicMesh3& EditMesh) { EditMesh = *NewMeshShared; });
 
diff --git a/Engine/Source/Runtime/GeometryCore/Private/DynamicMesh/MeshTangents.cpp b/Engine/Source/Runtime/GeometryCore/Private/DynamicMesh/MeshTangents.cpp
--- a/Engine/Source/Runtime/GeometryCore/Private/DynamicMesh/MeshTangents.cpp
+++ b/Engine/Source/Runtime/GeometryCore/Private/DynamicMesh/MeshTangents.cpp
@@ -26,10 +26,10 @@ FDynamicMeshTangents::FDynamicMeshTangents(const FDynamicMesh3* MeshIn)
 
 void FDynamicMeshTangents::GetTangentFrame(int32 TriangleID, int32 TriVertexIndex, FVector3f& Normal, FVector3f& Tangent, FVector3f& Bitangent) const
 {
-	if (Normals)
+	if (Normals && Normals->IsSetTriangle(TriangleID))
 	{
 		Normals->GetTriElement(TriangleID, TriVertexIndex, Normal);
-		if (Tangents)
+		if (Tangents && Tangents->IsSetTriangle(TriangleID))
 		{
 			Tangents->GetTriElement(TriangleID, TriVertexIndex, Tangent);
 			Bitangents->GetTriElement(TriangleID, TriVertexIndex, Bitangent);
@@ -50,7 +50,7 @@ void FDynamicMeshTangents::GetTangentFrame(int32 TriangleID, int32 TriVertexInde
 
 void FDynamicMeshTangents::GetTangentVectors(int32 TriangleID, int32 TriVertexIndex, const FVector3f& Normal, FVector3f& Tangent, FVector3f& Bitangent) const
 {
-	if (Tangents)
+	if (Tangents && Tangents->IsSetTriangle(TriangleID) )
 	{
 		Tangents->GetTriElement(TriangleID, TriVertexIndex, Tangent);
 		Bitangents->GetTriElement(TriangleID, TriVertexIndex, Bitangent);
@@ -66,7 +66,7 @@ void FDynamicMeshTangents::GetTangentVectors(int32 TriangleID, int32 TriVertexIn
 
 void FDynamicMeshTangents::GetTangentVectors(int32 TriangleID, int32 TriVertexIndex, FVector3f& Tangent, FVector3f& Bitangent) const
 {
-	if (Tangents)
+	if (Tangents && Tangents->IsSetTriangle(TriangleID))
 	{
 		Tangents->GetTriElement(TriangleID, TriVertexIndex, Tangent);
 		Bitangents->GetTriElement(TriangleID, TriVertexIndex, Bitangent);
