workaround internal compiler error on MSVC2019 + ASan

#rnx

diff --git a/Engine/Source/Runtime/GeometryCore/Public/Util/DynamicVector.h b/Engine/Source/Runtime/GeometryCore/Public/Util/DynamicVector.h
--- a/Engine/Source/Runtime/GeometryCore/Public/Util/DynamicVector.h
+++ b/Engine/Source/Runtime/GeometryCore/Public/Util/DynamicVector.h
@@ -92,11 +92,17 @@ public:
 
 	inline Type& operator[](unsigned int Index)
 	{
-		return Blocks[Index >> nShiftBits][Index & BlockIndexBitmask];
+		int ArrayIndex, BlockIndex;
+		GetIndices( Index, ArrayIndex, BlockIndex );
+	
+		return Blocks[ArrayIndex][BlockIndex];
 	}
 	inline const Type& operator[](unsigned int Index) const
 	{
-		return Blocks[Index >> nShiftBits][Index & BlockIndexBitmask];
+		int ArrayIndex, BlockIndex;
+		GetIndices( Index, ArrayIndex, BlockIndex );
+	
+		return Blocks[ArrayIndex][BlockIndex];
 	}
 
 	// apply f() to each member sequentially
@@ -392,6 +398,13 @@ private:
 		}
 	};
 
+	// helper function - avoids MSVC+ASan internal compiler error in [] operators
+	inline void GetIndices( unsigned int Index, int& ArrayIndex, int& BlockIndex ) const
+	{
+		ArrayIndex = (Index >> nShiftBits);
+		BlockIndex = (Index & BlockIndexBitmask);
+	}
+
 	// memory chunks for dynamic vector
 	TBlockVector<BlockType> Blocks;
 };
