FOptionallySparseIndexMap: fix crash when adding a new key/value pair and we're using sparse map type

#rnx

diff --git a/Engine/Source/Runtime/GeometryCore/Public/Util/SparseIndexCollectionTypes.h b/Engine/Source/Runtime/GeometryCore/Public/Util/SparseIndexCollectionTypes.h
--- a/Engine/Source/Runtime/GeometryCore/Public/Util/SparseIndexCollectionTypes.h
+++ b/Engine/Source/Runtime/GeometryCore/Public/Util/SparseIndexCollectionTypes.h
@@ -361,7 +361,7 @@ struct FOptionallySparseIndexMap
 		}
 		else
 		{
-			Sparse[Index] = Value;
+			Sparse.FindOrAdd(Index) = Value;
 		}
 	}
 
