fix issues w/ overlay that can leave ref counts and indices in invalid state

#rnx

diff --git a/Engine/Source/Runtime/GeometryCore/Private/DynamicMesh/DynamicMeshOverlay.cpp b/Engine/Source/Runtime/GeometryCore/Private/DynamicMesh/DynamicMeshOverlay.cpp
--- a/Engine/Source/Runtime/GeometryCore/Private/DynamicMesh/DynamicMeshOverlay.cpp
+++ b/Engine/Source/Runtime/GeometryCore/Private/DynamicMesh/DynamicMeshOverlay.cpp
@@ -74,7 +74,7 @@ EMeshResult TDynamicMeshOverlay<RealType, ElementSize>::InsertElement(int Elemen
 		Elements.InsertAt(Value[k], i + k);
 	}
 
-	ParentVertices.InsertAt(FDynamicMesh3::InvalidID, ElementID);
+	ParentVertices.InsertAt(FDynamicMesh3::InvalidID, ElementID, FDynamicMesh3::InvalidID);
 
 	//UpdateTimeStamp(true, true);
 	return EMeshResult::Ok;
@@ -152,9 +152,9 @@ void TDynamicMeshOverlay<RealType, ElementSize>::CreateFromPredicate(TFunctionRe
 				int VertSubIdx = IndexUtil::FindTriIndex(VertexID, TriVertIDs);
 				int i = 3 * TriID;
 				int ElementIndex = AppendedElements[TrisActiveSubGroup[TriSubIdx]];
-				ElementTriangles.InsertAt(ElementIndex, i + VertSubIdx);
+				ElementTriangles.InsertAt(ElementIndex, i + VertSubIdx, FDynamicMesh3::InvalidID);
 				ElementsRefCounts.Increment(ElementIndex);
-				ParentVertices.InsertAt(VertexID, ElementIndex);
+				ParentVertices.InsertAt(VertexID, ElementIndex); // elements were appended one-by-one above, so default initialization not needed here
 			}
 			GroupStart += GroupNum;
 		}
@@ -232,7 +232,7 @@ int TDynamicMeshOverlay<RealType, ElementSize>::SplitElementWithNewParent(int El
 			}
 		}
 	}
-	ParentVertices.InsertAt(NewParentID, NewElID);
+	ParentVertices.InsertAt(NewParentID, NewElID, FDynamicMesh3::InvalidID);
 
 	checkSlow(ElementsRefCounts.IsValid(ElementID));
 
@@ -405,7 +405,7 @@ void TDynamicMeshOverlay<RealType, ElementSize>::UnsetTriangle(int TriangleID)
 
 
 template<typename RealType, int ElementSize>
-void TDynamicMeshOverlay<RealType, ElementSize>::InternalSetTriangle(int tid, const FIndex3i& tv, bool bIncrementRefCounts)
+void TDynamicMeshOverlay<RealType, ElementSize>::InternalSetTriangle(int tid, const FIndex3i& tv, bool bUpdateRefCounts)
 {
 	if (ensure(ParentMesh) == false)
 	{
@@ -413,11 +413,17 @@ void TDynamicMeshOverlay<RealType, ElementSize>::InternalSetTriangle(int tid, co
 	}
 
 	int i = 3 * tid;
-	ElementTriangles.InsertAt(tv[2], i + 2);
-	ElementTriangles.InsertAt(tv[1], i + 1);
-	ElementTriangles.InsertAt(tv[0], i);
 
-	if (bIncrementRefCounts)
+	if (!ElementTriangles.SetMinimumSize(i + 3, FDynamicMesh3::InvalidID) && bUpdateRefCounts)
+	{
+		UnsetTriangle(tid);
+	}
+
+	ElementTriangles[i + 2] = tv[2];
+	ElementTriangles[i + 1] = tv[1];
+	ElementTriangles[i] = tv[0];
+
+	if (bUpdateRefCounts)
 	{
 		ElementsRefCounts.Increment(tv[0]);
 		ElementsRefCounts.Increment(tv[1]);
@@ -437,7 +443,7 @@ void TDynamicMeshOverlay<RealType, ElementSize>::InternalSetTriangle(int tid, co
 			// allowed to be used for multiple vertices.
 			checkSlow(ParentVertices[tv[VInd]] == ParentTriangle[VInd] || ParentVertices[tv[VInd]] == FDynamicMesh3::InvalidID);
 
-			ParentVertices.InsertAt(ParentTriangle[VInd], tv[VInd]);
+			ParentVertices.InsertAt(ParentTriangle[VInd], tv[VInd], FDynamicMesh3::InvalidID);
 		}
 	}
 }
@@ -448,9 +454,10 @@ template<typename RealType, int ElementSize>
 void TDynamicMeshOverlay<RealType, ElementSize>::InitializeNewTriangle(int tid)
 {
 	int i = 3 * tid;
-	ElementTriangles.InsertAt(FDynamicMesh3::InvalidID, i + 2);
-	ElementTriangles.InsertAt(FDynamicMesh3::InvalidID, i + 1);
-	ElementTriangles.InsertAt(FDynamicMesh3::InvalidID, i);
+	ElementTriangles.SetMinimumSize(i + 3, FDynamicMesh3::InvalidID);
+	ElementTriangles[i + 2] = FDynamicMesh3::InvalidID;
+	ElementTriangles[i + 1] = FDynamicMesh3::InvalidID;
+	ElementTriangles[i] = FDynamicMesh3::InvalidID;
 
 	//updateTimeStamp(true);
 }
diff --git a/Engine/Source/Runtime/GeometryCore/Public/DynamicMesh/DynamicMeshOverlay.h b/Engine/Source/Runtime/GeometryCore/Public/DynamicMesh/DynamicMeshOverlay.h
--- a/Engine/Source/Runtime/GeometryCore/Public/DynamicMesh/DynamicMeshOverlay.h
+++ b/Engine/Source/Runtime/GeometryCore/Public/DynamicMesh/DynamicMeshOverlay.h
@@ -506,7 +506,7 @@ protected:
 	void SetElementFromBary(int SetElement, int ElementA, int ElementB, int ElementC, const FVector3<double>& BaryCoords);
 
 	/** updates the triangles array and optionally the element reference counts */
-	void InternalSetTriangle(int TriangleID, const FIndex3i& TriElements, bool bIncrementRefCounts);
+	void InternalSetTriangle(int TriangleID, const FIndex3i& TriElements, bool bUpdateRefCounts);
 
 };
 
diff --git a/Engine/Source/Runtime/GeometryCore/Public/Util/DynamicVector.h b/Engine/Source/Runtime/GeometryCore/Public/Util/DynamicVector.h
--- a/Engine/Source/Runtime/GeometryCore/Public/Util/DynamicVector.h
+++ b/Engine/Source/Runtime/GeometryCore/Public/Util/DynamicVector.h
@@ -72,6 +72,8 @@ public:
 	inline void Fill(const Type& Value);
 	inline void Resize(size_t Count);
 	inline void Resize(size_t Count, const Type& InitValue);
+	/// Resize if Num() is less than Count; returns true if resize occurred
+	inline bool SetMinimumSize(size_t Count, const Type& InitValue);
 	inline void SetNum(size_t Count) { Resize(Count); }
 
 	inline bool IsEmpty() const { return CurBlock == 0 && CurBlockUsed == 0; }
@@ -85,6 +87,7 @@ public:
 	inline void PopBack();
 
 	inline void InsertAt(const Type& Data, unsigned int Index);
+	inline void InsertAt(const Type& Data, unsigned int Index, const Type& InitValue);
 	inline Type& ElementAt(unsigned int Index, Type InitialValue = Type{});
 
 	inline const Type& Front() const { return Blocks[0][0]; }
@@ -630,6 +633,22 @@ void TDynamicVector<Type>::Resize(size_t Count, const Type& InitValue)
 	}
 }
 
+template <class Type>
+bool TDynamicVector<Type>::SetMinimumSize(size_t Count, const Type& InitValue)
+{
+	size_t nCurSize = GetLength();
+	if (Count <= nCurSize)
+	{
+		return false;
+	}
+	Resize(Count);
+	for (size_t Index = nCurSize; Index < Count; ++Index)
+	{
+		this->operator[](Index) = InitValue;
+	}
+	return true;
+}
+
 template <class Type>
 void TDynamicVector<Type>::Add(const Type& Value)
 {
@@ -708,6 +727,18 @@ void TDynamicVector<Type>::InsertAt(const Type& AddData, unsigned int Index)
 	}
 }
 
+template <class Type>
+void TDynamicVector<Type>::InsertAt(const Type& AddData, unsigned int Index, const Type& InitValue)
+{
+	size_t nCurSize = GetLength();
+	InsertAt(AddData, Index);
+	// initialize all new values up to (but not including) the inserted index
+	for (size_t i = nCurSize; i < (size_t)Index; ++i)
+	{
+		this->operator[](Index) = InitValue;
+	}
+}
+
 template <typename Type>
 template <typename Func>
 void TDynamicVector<Type>::Apply(const Func& applyF)
