fix wrong index used when filling default values w/ InsertAt

#rnx

diff --git a/Engine/Source/Runtime/GeometryCore/Public/Util/DynamicVector.h b/Engine/Source/Runtime/GeometryCore/Public/Util/DynamicVector.h
--- a/Engine/Source/Runtime/GeometryCore/Public/Util/DynamicVector.h
+++ b/Engine/Source/Runtime/GeometryCore/Public/Util/DynamicVector.h
@@ -735,7 +735,7 @@ void TDynamicVector<Type>::InsertAt(const Type& AddData, unsigned int Index, con
 	// initialize all new values up to (but not including) the inserted index
 	for (size_t i = nCurSize; i < (size_t)Index; ++i)
 	{
-		this->operator[](Index) = InitValue;
+		this->operator[](i) = InitValue;
 	}
 }
 
