GeometryCore: Add LogGeometry category, add generic FGeometryResult w/ result code and FGeometryError/FGeometryWarning, add as optional return from FDynamicMeshOperator.
#rnx

diff --git a/Engine/Plugins/Experimental/MeshModelingToolset/Source/ModelingOperators/Public/ModelingOperators.h b/Engine/Plugins/Experimental/MeshModelingToolset/Source/ModelingOperators/Public/ModelingOperators.h
--- a/Engine/Plugins/Experimental/MeshModelingToolset/Source/ModelingOperators/Public/ModelingOperators.h
+++ b/Engine/Plugins/Experimental/MeshModelingToolset/Source/ModelingOperators/Public/ModelingOperators.h
@@ -22,6 +22,7 @@ class FDynamicMeshOperator
 protected:
 	TUniquePtr<FDynamicMesh3> ResultMesh;
 	FTransform3d ResultTransform;
+	FGeometryResult ResultInfo;
 
 public:
 	FDynamicMeshOperator()
@@ -41,6 +42,13 @@ public:
 		ResultTransform = Transform;
 	}
 
+	/**
+	 * Set the output information
+	 */
+	virtual void SetResultInfo(const FGeometryResult& Info)
+	{
+		ResultInfo = Info;
+	}
 
 	/**
 	 * @return ownership of the internal mesh that CalculateResult() produced
@@ -58,6 +66,14 @@ public:
 		return ResultTransform;
 	}
 
+	/**
+	 * @return the result information returned by CalculateResult()
+	 */
+	const FGeometryResult& GetResultInfo() const
+	{
+		return ResultInfo;
+	}
+
 	/**
 	 * Calculate the result of the operator. This will populate the internal Mesh and Transform.
 	 * @param Progress implementors can use this object to report progress and determine if they should halt and terminate expensive computations
diff --git a/Engine/Source/Runtime/GeometryCore/Private/GeometryCoreModule.cpp b/Engine/Source/Runtime/GeometryCore/Private/GeometryCoreModule.cpp
--- a/Engine/Source/Runtime/GeometryCore/Private/GeometryCoreModule.cpp
+++ b/Engine/Source/Runtime/GeometryCore/Private/GeometryCoreModule.cpp
@@ -6,6 +6,8 @@
 
 #define LOCTEXT_NAMESPACE "FGeometryCoreModule"
 
+DEFINE_LOG_CATEGORY(LogGeometry);
+
 void FGeometryCoreModule::StartupModule()
 {
 	// This code will execute after your module is loaded into memory; the exact timing is specified in the .uplugin file per-module
diff --git a/Engine/Source/Runtime/GeometryCore/Public/GeometryBase.h b/Engine/Source/Runtime/GeometryCore/Public/GeometryBase.h
--- a/Engine/Source/Runtime/GeometryCore/Public/GeometryBase.h
+++ b/Engine/Source/Runtime/GeometryCore/Public/GeometryBase.h
@@ -3,6 +3,7 @@
 #pragma once
 
 #include "CoreTypes.h"   // required for GEOMETRYCORE_API macro
+#include "Logging/LogMacros.h"
 
 //
 // The following are convenience macros that can be used to predeclare types in the UE::Geometry::
@@ -30,4 +31,6 @@ namespace UE
 			int a;
 		};
 	}
-}
\ No newline at end of file
+}
+
+GEOMETRYCORE_API DECLARE_LOG_CATEGORY_EXTERN(LogGeometry, Log, All);
\ No newline at end of file
diff --git a/Engine/Source/Runtime/GeometryCore/Public/Util/ProgressCancel.h b/Engine/Source/Runtime/GeometryCore/Public/Util/ProgressCancel.h
--- a/Engine/Source/Runtime/GeometryCore/Public/Util/ProgressCancel.h
+++ b/Engine/Source/Runtime/GeometryCore/Public/Util/ProgressCancel.h
@@ -9,6 +9,157 @@
 #include "Internationalization/Text.h"
 #include "Misc/DateTime.h"
 
+
+class FProgressCancel;
+
+
+namespace UE
+{
+namespace Geometry
+{
+
+
+/**
+ * FGeometryError represents an error code/message emitted by a geometry operation.
+ * The intention of Errors is that they are fatal, IE if an operation emits Errors then
+ * it did not complete successfully. If that is not the case, use a FGeometryWarning instead.
+ */
+struct FGeometryError
+{
+	int32 ErrorCode = 0;
+	FText Message;
+	FDateTime Timestamp;
+	TArray<unsigned char> CustomData;
+
+	FGeometryError() 
+	{ 
+		Timestamp = FDateTime::Now();
+	}
+
+	FGeometryError(int32 Code, const FText& MessageIn) : ErrorCode(Code), Message(MessageIn) 
+	{
+		Timestamp = FDateTime::Now();
+	}
+};
+
+/**
+ * FGeometryWarning represents a warning code/message emitted by a geometry operation.
+ * The intention of Warnings is that they are non-fatal, IE an operation might successfully
+ * complete but still emit Warnings
+ */
+struct FGeometryWarning
+{
+	int32 WarningCode = 0;
+	FText Message;
+	FDateTime Timestamp;
+	TArray<unsigned char> CustomData;
+
+
+	FGeometryWarning()
+	{
+		Timestamp = FDateTime::Now();
+	}
+
+	FGeometryWarning(int32 Code, const FText& MessageIn) : WarningCode(Code), Message(MessageIn)
+	{
+		Timestamp = FDateTime::Now();
+	}
+};
+
+
+/**
+ * EGeometryResultType is a generic result-code for use by geometry operations.
+ */
+enum class EGeometryResultType
+{
+	/** The Geometry Operation successfully completed */
+	Success = 0,
+	/** The Geometry Operation is in progress (this is intended for use in cases where there is a background computation that can be queried for incremental status/etc) */
+	InProgress = 1,
+	/** The Geometry Operation was Cancelled and did not complete */
+	Cancelled = 2,
+	/** The Geometry Operation completed but was not fully successful (eg some sub-computation failed and was gracefully handled, etc) */
+	PartialResult = 3,
+	/** The Geometry Operation failed and no valid result was produced */
+	Failure = 4
+};
+
+
+/**
+ * FGeometryResult represents a combined "success/failure/etc" state for a geometry operation
+ * along with a set of Error and Warning codes/messages. 
+ */
+struct FGeometryResult
+{
+	EGeometryResultType Result;
+	TArray<FGeometryError> Errors;
+	TArray<FGeometryWarning> Warnings;
+
+
+	FGeometryResult()
+	{
+		Result = EGeometryResultType::Success;
+	}
+
+	FGeometryResult(EGeometryResultType ResultType) 
+	{
+		Result = ResultType;
+	}
+
+	void UpdateResultType(EGeometryResultType NewType)
+	{
+		Result = NewType;
+	}
+
+	void SetFailed() { Result = EGeometryResultType::Failure; }
+	void SetCancelled() { Result = EGeometryResultType::Cancelled; }
+
+	void SetSuccess() { Result = EGeometryResultType::Success; }
+
+	/**
+	 * Set to Success/Failure based on bSuccess, or Cancelled if the (optional) FProgressCancel indicates that it was Cancelled
+	 */
+	inline void SetSuccess(bool bSuccess, FProgressCancel* Progress = nullptr );
+
+	/**
+	 * Append an Error to the result
+	 */
+	void AddError(FGeometryError Error) 
+	{ 
+		Errors.Add(Error); 
+	}
+	
+	/**
+	 * Append a Warning to the result
+	 */
+	void AddWarning(FGeometryWarning Warning)
+	{ 
+		Warnings.Add(Warning); 
+	}
+
+	/**
+	 * @return true if the geometry operation failed
+	 */
+	bool HasFailed() const { return Result == EGeometryResultType::Failure; }
+
+	/**
+	 * @return true if the geometry operation has some result (ie was a Success or returned a PartialResult)
+	 */
+	bool HasResult() const { return Result == EGeometryResultType::Success || Result == EGeometryResultType::PartialResult; }
+
+	static FGeometryResult Failed() { return FGeometryResult(EGeometryResultType::Failure); }
+	static FGeometryResult Cancelled() { return FGeometryResult(EGeometryResultType::Cancelled); }
+};
+
+
+
+} // end namespace UE::Geometry
+} // end namespace UE
+
+
+
+
+
 /**
  * FProgressCancel is an obejct that is intended to be passed to long-running
  * computes to do two things:
@@ -69,3 +220,18 @@ public:
 
 	TArray<FMessageInfo> Warnings;
 };
+
+
+
+
+void UE::Geometry::FGeometryResult::SetSuccess(bool bSuccess, FProgressCancel* Progress)
+{
+	if (Progress && Progress->Cancelled())
+	{
+		Result = EGeometryResultType::Cancelled;
+	}
+	else
+	{
+		Result = (bSuccess) ? EGeometryResultType::Success : EGeometryResultType::Failure;
+	}
+}
\ No newline at end of file
