ModelingTools: Refactored ImageTile out of ImageDimensions
- Inlined MeshMapBaker::BakePixel().
- Fixed gutter texels when multisampling.

#rnx

diff --git a/Engine/Plugins/Experimental/GeometryProcessing/Source/DynamicMesh/Private/Sampling/MeshMapBaker.cpp b/Engine/Plugins/Experimental/GeometryProcessing/Source/DynamicMesh/Private/Sampling/MeshMapBaker.cpp
--- a/Engine/Plugins/Experimental/GeometryProcessing/Source/DynamicMesh/Private/Sampling/MeshMapBaker.cpp
+++ b/Engine/Plugins/Experimental/GeometryProcessing/Source/DynamicMesh/Private/Sampling/MeshMapBaker.cpp
@@ -3,6 +3,7 @@
 #include "Sampling/MeshMapBaker.h"
 #include "Sampling/MeshBakerCommon.h"
 #include "Image/ImageOccupancyMap.h"
+#include "Image/ImageTile.h"
 #include "Spatial/DenseGrid2.h"
 
 #include "ExplicitUseGeometryMathTypes.h"		// using UE::Geometry::(math types)
@@ -171,30 +172,74 @@ void FMeshMapBaker::Bake()
 	});
 
 	// Setup image tiling
-	FImageTiling Tiles(Dimensions, TileSize, TileSize, 0);
+	FImageTiling Tiles(Dimensions, TileSize, TileSize);
 	const int32 NumBakers = Bakers.Num();
 	const int32 NumResults = BakeResults.Num();
 	ParallelFor(Tiles.Num(), [this, &Tiles, &NumBakers, &NumResults](int32 TileIdx)
 	{
-		FImageDimensions Tile = Tiles.GetTile(TileIdx);
+		const FImageTile Tile = Tiles.GetTile(TileIdx);
 
 		FImageOccupancyMap OccupancyMap;
 		OccupancyMap.GutterSize = GutterSize;
-		OccupancyMap.Initialize(Tile, Multisampling);
+		OccupancyMap.Initialize(Dimensions, Tile, Multisampling);
 		OccupancyMap.ComputeFromUVSpaceMesh(FlatMesh, [this](int32 TriangleID) { return FlatMesh.GetTriangleGroup(TriangleID); });
 
 		FMeshMapTileBuffer TileBuffer(Tile, BakeSampleBufferSize);
 
 		// Calculate interior texels
-		for (FVector2i Texel(0,0); Texel.Y < Tile.GetHeight(); ++Texel.Y)
+		for (FVector2i TileCoords(0,0); TileCoords.Y < Tile.GetHeight(); ++TileCoords.Y)
 		{
-			for (Texel.X = 0; Texel.X < Tile.GetWidth(); ++Texel.X)
+			for (TileCoords.X = 0; TileCoords.X < Tile.GetWidth(); ++TileCoords.X)
 			{
 				if (CancelF())
 				{
 					return;
 				}
-				BakePixel(TileBuffer, OccupancyMap, Tile, Texel);
+
+				const int64 TilePixelLinearIdx = Tile.GetIndex(TileCoords);
+				if (OccupancyMap.TexelNumSamples(TilePixelLinearIdx) == 0)
+				{
+					continue;
+				}
+
+				float* PixelBuffer = TileBuffer.GetPixel(TilePixelLinearIdx);
+				float& PixelWeight = TileBuffer.GetPixelWeight(TilePixelLinearIdx);
+
+				const FVector2i ImageCoords = Tile.GetSourceCoords(TileCoords);
+				const int32 NumSamples = OccupancyMap.Multisampler.Num();
+				const int32 NumEvaluatorAdd = BakeAccumulateLists[(int32)FMeshMapEvaluator::EAccumulateMode::Add].Num();
+				for (int32 SampleIdx = 0; SampleIdx < NumSamples; ++SampleIdx)
+				{
+					const int64 LinearIdx = TilePixelLinearIdx * OccupancyMap.Multisampler.Num() + SampleIdx;
+					if (!OccupancyMap.IsInterior(LinearIdx))
+					{
+						PixelWeight += NumEvaluatorAdd > 0;
+						for (int32 Idx : BakeAccumulateLists[(int32) FMeshMapEvaluator::EAccumulateMode::Add] )
+						{
+							const int32 NumData = BakeContexts[Idx].DataLayout.Num();
+							for (int32 DataIdx = 0; DataIdx < NumData; ++DataIdx)
+							{
+								const int32 ResultIdx = BakeOffsets[Idx] + DataIdx;
+								const int32 Offset = BakeSampleOffsets[ResultIdx];
+								const int32 Stride = (int32) BakeContexts[Idx].DataLayout[DataIdx];
+								for (int32 BufIdx = 0; BufIdx < Stride; ++BufIdx)
+								{
+									PixelBuffer[Offset + BufIdx] += BakeDefaults[Offset + BufIdx];
+								}
+							}
+						}
+					}
+					else
+					{
+						const FVector2d UVPosition = (FVector2d)OccupancyMap.TexelQueryUV[LinearIdx];
+						const int32 UVTriangleID = OccupancyMap.TexelQueryTriangle[LinearIdx];
+
+						FMeshMapEvaluator::FCorrespondenceSample Sample;
+						DetailMeshSampler.SampleUV(UVTriangleID, UVPosition, Sample);
+
+						BakeSample(TileBuffer, Sample, Tile, TileCoords, ImageCoords, 1.0f);
+					}
+				}
 			}
 		}
 
@@ -212,7 +257,8 @@ void FMeshMapBaker::Bake()
 				const float& PixelWeight = TileBuffer.GetPixelWeight(PixelLinearIdx);
 				if (PixelWeight > 0.0)
 				{
-					const int64 ImageLinearIdx = Tile.GetSourceIndex(Texel);
+					const FVector2i ImageCoords = Tile.GetSourceCoords(Texel);
+					const int64 ImageLinearIdx = Dimensions.GetIndex(ImageCoords);
 					float* PixelBuffer = TileBuffer.GetPixel(PixelLinearIdx);
 					const float OneOverPixelWeight = 1.0f / PixelWeight;
 
@@ -262,55 +308,7 @@ void FMeshMapBaker::Bake()
 	}, !bParallel ? EParallelForFlags::ForceSingleThread : EParallelForFlags::None);
 }
 
-void FMeshMapBaker::BakePixel(FMeshMapTileBuffer& TileBuffer, FImageOccupancyMap& OccupancyMap, const FImageDimensions& Tile, const FVector2i& TileCoords)
-{
-	if (OccupancyMap.TexelNumSamples(Tile.GetIndex(TileCoords)) == 0)
-	{
-		return;
-	}
-
-	const int64 TilePixelLinearIdx = Tile.GetIndex(TileCoords);
-	float* PixelBuffer = TileBuffer.GetPixel(TilePixelLinearIdx);
-	float& PixelWeight = TileBuffer.GetPixelWeight(TilePixelLinearIdx);
-
-	const FVector2i ImageCoords = Tile.GetSourceCoords(TileCoords.X, TileCoords.Y);
-	const int32 NumSamples = OccupancyMap.Multisampler.Num();
-	const int32 NumEvaluatorAdd = BakeAccumulateLists[(int32)FMeshMapEvaluator::EAccumulateMode::Add].Num();
-	for (int32 SampleIdx = 0; SampleIdx < NumSamples; ++SampleIdx)
-	{
-		const int64 LinearIdx = TilePixelLinearIdx * OccupancyMap.Multisampler.Num() + SampleIdx;
-		if (!OccupancyMap.IsInterior(LinearIdx))
-		{
-			PixelWeight += NumEvaluatorAdd > 0;
-			for (int32 Idx : BakeAccumulateLists[(int32) FMeshMapEvaluator::EAccumulateMode::Add] )
-			{
-				const int32 NumData = BakeContexts[Idx].DataLayout.Num();
-				for (int32 DataIdx = 0; DataIdx < NumData; ++DataIdx)
-				{
-					const int32 ResultIdx = BakeOffsets[Idx] + DataIdx;
-					const int32 Offset = BakeSampleOffsets[ResultIdx];
-					const int32 Stride = (int32) BakeContexts[Idx].DataLayout[DataIdx];
-					for (int32 BufIdx = 0; BufIdx < Stride; ++BufIdx)
-					{
-						PixelBuffer[Offset + BufIdx] += BakeDefaults[Offset + BufIdx];
-					}
-				}
-			}
-		}
-		else
-		{
-			const FVector2d UVPosition = (FVector2d)OccupancyMap.TexelQueryUV[LinearIdx];
-			const int32 UVTriangleID = OccupancyMap.TexelQueryTriangle[LinearIdx];
-
-			FMeshMapEvaluator::FCorrespondenceSample Sample;
-			DetailMeshSampler.SampleUV(UVTriangleID, UVPosition, Sample);
-
-			BakeSample(TileBuffer, Sample, Tile, TileCoords, ImageCoords, 1.0f);
-		}
-	}
-}
-
-void FMeshMapBaker::BakeSample(FMeshMapTileBuffer& TileBuffer, const FMeshMapEvaluator::FCorrespondenceSample& Sample, const FImageDimensions& Tile,
+void FMeshMapBaker::BakeSample(FMeshMapTileBuffer& TileBuffer, const FMeshMapEvaluator::FCorrespondenceSample& Sample, const FImageTile& Tile,
 	const FVector2i& TileCoords, const FVector2i& ImageCoords, const float& SampleWeight)
 {
 	const int64 TilePixelLinearIdx = Tile.GetIndex(TileCoords);
diff --git a/Engine/Plugins/Experimental/GeometryProcessing/Source/DynamicMesh/Private/Sampling/MeshVertexBaker.cpp b/Engine/Plugins/Experimental/GeometryProcessing/Source/DynamicMesh/Private/Sampling/MeshVertexBaker.cpp
--- a/Engine/Plugins/Experimental/GeometryProcessing/Source/DynamicMesh/Private/Sampling/MeshVertexBaker.cpp
+++ b/Engine/Plugins/Experimental/GeometryProcessing/Source/DynamicMesh/Private/Sampling/MeshVertexBaker.cpp
@@ -170,7 +170,7 @@ void FMeshVertexBaker::BakeImpl(void* Data)
 	const int NumBakers = Baker->Bakers.Num();
 	ParallelFor(NumTiles, [&Baker, &Tiles, &TileWidth, &NumBakers, &SampleSurface](const int32 TileIdx)
 	{
-		const FImageDimensions Tile = Tiles.GetTile(TileIdx);
+		const FImageTile Tile = Tiles.GetTile(TileIdx);
 		const int Width = Tile.GetWidth();
 		for (int32 Idx = 0; Idx < Width; ++Idx)
 		{
diff --git a/Engine/Plugins/Experimental/GeometryProcessing/Source/DynamicMesh/Public/Sampling/MeshBakerCommon.h b/Engine/Plugins/Experimental/GeometryProcessing/Source/DynamicMesh/Public/Sampling/MeshBakerCommon.h
--- a/Engine/Plugins/Experimental/GeometryProcessing/Source/DynamicMesh/Public/Sampling/MeshBakerCommon.h
+++ b/Engine/Plugins/Experimental/GeometryProcessing/Source/DynamicMesh/Public/Sampling/MeshBakerCommon.h
@@ -4,7 +4,7 @@
 
 #include "DynamicMesh/DynamicMesh3.h"
 #include "DynamicMesh/DynamicMeshAABBTree3.h"
-#include "Image/ImageDimensions.h"
+#include "Image/ImageTile.h"
 
 namespace UE
 {
@@ -35,7 +35,7 @@ bool GetDetailMeshTrianglePoint_Nearest(
 class FMeshMapTileBuffer
 {
 public:
-	FMeshMapTileBuffer(const FImageDimensions& TileIn, const int32 PixelSizeIn)
+	FMeshMapTileBuffer(const FImageTile& TileIn, const int32 PixelSizeIn)
 		: Tile(TileIn)
 		, PixelSize(PixelSizeIn + 1) // + 1 for accumulated pixel weight.
 	{
@@ -60,7 +60,7 @@ public:
 	}
 
 private:
-	const FImageDimensions Tile;
+	const FImageTile Tile;
 	const int32 PixelSize;
 	float* Buffer;
 };
diff --git a/Engine/Plugins/Experimental/GeometryProcessing/Source/DynamicMesh/Public/Sampling/MeshMapBaker.h b/Engine/Plugins/Experimental/GeometryProcessing/Source/DynamicMesh/Public/Sampling/MeshMapBaker.h
--- a/Engine/Plugins/Experimental/GeometryProcessing/Source/DynamicMesh/Public/Sampling/MeshMapBaker.h
+++ b/Engine/Plugins/Experimental/GeometryProcessing/Source/DynamicMesh/Public/Sampling/MeshMapBaker.h
@@ -15,6 +15,7 @@ namespace Geometry
 {
 
 class FImageOccupancyMap;
+class FImageTile;
 class FMeshMapTileBuffer;
 
 class DYNAMICMESH_API FMeshMapBaker : public FMeshBaseBaker
@@ -60,12 +61,14 @@ public:
 	int32 GetMultisampling() const { return Multisampling; }
 
 protected:
-	/** Evaluate samples for this tile pixel. */
-	void BakePixel(FMeshMapTileBuffer& TileBuffer, FImageOccupancyMap& OccupancyMap, const FImageDimensions& Tile, const FVector2i& TileCoords);
-
 	/** Evaluate this sample. */
-	void BakeSample(FMeshMapTileBuffer& TileBuffer, const FMeshMapEvaluator::FCorrespondenceSample& Sample, const FImageDimensions& Tile,
-		const FVector2i& TileCoords, const FVector2i& ImageCoords, const float& SampleWeight);
+	void BakeSample(
+		FMeshMapTileBuffer& TileBuffer,
+		const FMeshMapEvaluator::FCorrespondenceSample& Sample,
+		const FImageTile& Tile,
+		const FVector2i& TileCoords,
+		const FVector2i& ImageCoords,
+		const float& SampleWeight);
 
 	/** Initialize evaluation contexts and precompute data for bake evaluation. */
 	void InitBake();
diff --git a/Engine/Source/Runtime/GeometryCore/Public/Image/ImageDimensions.h b/Engine/Source/Runtime/GeometryCore/Public/Image/ImageDimensions.h
--- a/Engine/Source/Runtime/GeometryCore/Public/Image/ImageDimensions.h
+++ b/Engine/Source/Runtime/GeometryCore/Public/Image/ImageDimensions.h
@@ -20,25 +20,10 @@ protected:
 	int32 Width = 0;
 	int32 Height = 0;
 
-	bool bIsTile = false;
-	int32 SourceWidth = 0;
-	int32 SourceHeight = 0;
-	FVector2i SourceOffset = FVector2i::Zero();
-
 public:
 	FImageDimensions(int32 WidthIn = 0, int32 HeightIn = 0)
-		: SourceWidth(WidthIn), SourceHeight(HeightIn)
-	{
-		SetDimensions(WidthIn, HeightIn);
-	}
-
-	/** Tile constructor. */
-	FImageDimensions(int32 WidthIn, int32 HeightIn, int32 SourceWidthIn, int32 SourceHeightIn, FVector2i SourceOffsetIn)
-		: bIsTile(true), SourceWidth(SourceWidthIn), SourceHeight(SourceHeightIn), SourceOffset(SourceOffsetIn)
 	{
 		SetDimensions(WidthIn, HeightIn);
-		checkSlow(Width >= 0 && Width <= SourceWidth && SourceOffset.X >= 0 && Width + SourceOffset.X <= SourceWidth);
-		checkSlow(Height >= 0 && Height <= SourceHeight && SourceOffset.Y >= 0 && Height + SourceOffset.Y <= SourceHeight);
 	}
 
 	/** Set the dimensions of this image. */
@@ -57,7 +42,7 @@ public:
 	int64 Num() const { return (int64)Width * (int64)Height; }
 
 	/** @return true if image is square */
-	bool IsSquare() const { return bIsTile ? SourceWidth == SourceHeight : Width == Height; }
+	bool IsSquare() const { return Width == Height; }
 
 	/** @return true if coordinates are valid, ie in-bounds of image dimensions */
 	bool IsValidCoords(const FVector2i& Coords) const
@@ -111,24 +96,15 @@ public:
 	/** @return Real-valued dimensions of a pixel/texel in the image, relative to default UV space [0..1]^2 */
 	FVector2d GetTexelSize() const
 	{
-		return FVector2d(1.0 / (double)SourceWidth, 1.0 / (double)SourceHeight);
+		return FVector2d(1.0 / (double)Width, 1.0 / (double)Height);
 	}
 
 	/** @return Real-valued position of given texel center in default UV-space [0..1]^2 */
 	FVector2d GetTexelUV(const FVector2i& Coords) const
 	{
-		if (bIsTile)
-		{
-			return FVector2d(
-				((double)SourceOffset.X + (double)Coords.X + 0.5) / (double)SourceWidth,
-				((double)SourceOffset.Y + (double)Coords.Y + 0.5) / (double)SourceWidth);
-		}
-		else
-		{
-			return FVector2d(
-				((double)Coords.X + 0.5) / (double)Width,
-				((double)Coords.Y + 0.5) / (double)Height);
-		}
+		return FVector2d(
+			((double)Coords.X + 0.5) / (double)Width,
+			((double)Coords.Y + 0.5) / (double)Height);
 	}
 
 	/** @return Real-valued position of given texel center in default UV-space [0..1]^2 */
@@ -140,129 +116,30 @@ public:
 	/** @return integer XY coordinates for real-valued XY coordinates (ie texel that contains value, if texel origin is in bottom-left */
 	FVector2i PixelToCoords(const FVector2d& PixelPosition) const
 	{
-		if (bIsTile)
-		{
-			int32 X = FMath::Clamp((int32)PixelPosition.X, 0, SourceWidth - 1);
-			int32 Y = FMath::Clamp((int32)PixelPosition.Y, 0, SourceHeight - 1);
-			return FVector2i(X, Y);
-		}
-		else
-		{
-			int32 X = FMath::Clamp((int32)PixelPosition.X, 0, Width - 1);
-			int32 Y = FMath::Clamp((int32)PixelPosition.Y, 0, Height - 1);
-			return FVector2i(X, Y);
-		}
+		int32 X = FMath::Clamp((int32)PixelPosition.X, 0, Width - 1);
+		int32 Y = FMath::Clamp((int32)PixelPosition.Y, 0, Height - 1);
+		return FVector2i(X, Y);
 	}
 
 	/** @return integer XY coordinates for UV coordinates, assuming default UV space [0..1]^2 */
 	FVector2i UVToCoords(const FVector2d& UVPosition) const
 	{
-		if (bIsTile)
-		{
-			double PixelX = (UVPosition.X * (double)SourceWidth /*- 0.5*/);
-			double PixelY = (UVPosition.Y * (double)SourceHeight /*- 0.5*/);
-			return PixelToCoords(FVector2d(PixelX, PixelY));
-		}
-		else
-		{
-			double PixelX = (UVPosition.X * (double)Width /*- 0.5*/);
-			double PixelY = (UVPosition.Y * (double)Height /*- 0.5*/);
-			return PixelToCoords(FVector2d(PixelX, PixelY));
-		}
+		double PixelX = (UVPosition.X * (double)Width /*- 0.5*/);
+		double PixelY = (UVPosition.Y * (double)Height /*- 0.5*/);
+		return PixelToCoords(FVector2d(PixelX, PixelY));
 	}
 
 	bool operator==(const FImageDimensions& Other) const
 	{
-		return Width == Other.Width && Height == Other.Height &&
-			bIsTile == Other.bIsTile &&
-			SourceWidth == Other.SourceWidth &&
-			SourceHeight == Other.SourceHeight &&
-			SourceOffset == Other.SourceOffset;
+		return Width == Other.Width && Height == Other.Height;
 	}
 
 	bool operator!=(const FImageDimensions& Other) const
 	{
 		return !(*this == Other);
 	}
-
-	// Tiling
-
-	/** @return true if this image is a tile. */
-	bool IsTile() const
-	{
-		return bIsTile;
-	}
-
-	/** @return source image dimensions for this tile. */
-	FImageDimensions GetSourceDimensions() const
-	{
-		return FImageDimensions(SourceWidth, SourceHeight);
-	}
-
-	/** @return source image linear index from tile coordinates. */
-	int64 GetSourceIndex(const FVector2i& Coords) const
-	{
-		FVector2i SourceCoords = GetSourceCoords(Coords.X, Coords.Y);
-		return GetSourceDimensions().GetIndex(SourceCoords);
-	}
-
-	/** @return source image coordinates from tile coordinates. */
-	FVector2i GetSourceCoords(int32 TileX, int32 TileY) const
-	{
-		checkSlow(IsValidCoords(FVector2i(TileX, TileY)));
-		return FVector2i(SourceOffset.X + TileX, SourceOffset.Y + TileY);
-	}
 };
 
-class FImageTiling
-{
-private:
-	FImageDimensions Dimensions;
-	int32 TileWidth = 32;
-	int32 TileHeight = 32;
-	int32 TilePadding = 0;
-
-public:
-	FImageTiling(FImageDimensions DimensionsIn, int32 TileWidthIn, int32 TileHeightIn, int32 TilePaddingIn = 0)
-	{
-		SetTiling(DimensionsIn, TileWidthIn, TileHeightIn, TilePaddingIn);
-	}
-
-	void SetTiling(FImageDimensions DimensionsIn, int32 TileWidthIn, int32 TileHeightIn, int32 TilePaddingIn)
-	{
-		checkSlow(TileWidthIn >= 1 && TileHeightIn >= 1 && TilePaddingIn >= 0);
-		Dimensions = DimensionsIn;
-		TileWidth = FMath::Clamp(TileWidthIn, 1, Dimensions.GetWidth());
-		TileHeight = FMath::Clamp(TileHeightIn, 1, Dimensions.GetHeight());
-		TilePadding = TilePaddingIn < 0 ? 0 : TilePaddingIn;
-	}
-
-	int32 NumTilesX() const
-	{
-		return Dimensions.GetWidth() / TileWidth + (Dimensions.GetWidth() % TileWidth == 0 ? 0 : 1);
-	}
-
-	int32 NumTilesY() const
-	{
-		return Dimensions.GetHeight() / TileHeight + (Dimensions.GetHeight() % TileHeight == 0 ? 0 : 1);
-	}
-
-	int32 Num() const
-	{
-		return NumTilesX() * NumTilesY();
-	}
-
-	/** @return the tile corresponding to the linear index [0, NumTiles()-1]. */
-	FImageDimensions GetTile(int32 LinearTileIndex) const
-	{
-		FVector2i TileIndex(LinearTileIndex % NumTilesX(), LinearTileIndex / NumTilesX());
-		FVector2i TileOffset(TileIndex.X * TileWidth, TileIndex.Y * TileHeight);
-		int32 ClampedWidth = FMath::Clamp(TileWidth, 0, Dimensions.GetWidth() - TileIndex.X * TileWidth);
-		int32 ClampedHeight = FMath::Clamp(TileHeight, 0, Dimensions.GetHeight() - TileIndex.Y * TileHeight);
-		// TODO: Padding support
-		return FImageDimensions(ClampedWidth, ClampedHeight, Dimensions.GetWidth(), Dimensions.GetHeight(), TileOffset);
-	}
-};
 
 } // end namespace UE::Geometry
 } // end namespace UE
diff --git a/Engine/Source/Runtime/GeometryCore/Public/Image/ImageOccupancyMap.h b/Engine/Source/Runtime/GeometryCore/Public/Image/ImageOccupancyMap.h
--- a/Engine/Source/Runtime/GeometryCore/Public/Image/ImageOccupancyMap.h
+++ b/Engine/Source/Runtime/GeometryCore/Public/Image/ImageOccupancyMap.h
@@ -5,6 +5,7 @@
 #include "Templates/Tuple.h"
 #include "VectorTypes.h"
 #include "Image/ImageDimensions.h"
+#include "Image/ImageTile.h"
 #include "Spatial/MeshAABBTree3.h"
 #include "Sampling/GridSampler.h"
 
@@ -30,6 +31,7 @@ class FImageOccupancyMap
 public:
 	/** Image Dimensions */
 	FImageDimensions Dimensions;
+	FImageTile Tile;
 
 	/** Width of the gutter. This is actually multiplied by the diagonal length of a texel, so
 	    the gutter is generally larger than this number of pixels */
@@ -64,6 +66,15 @@ public:
 	{
 		check(DimensionsIn.IsSquare()); // are we sure it works otherwise?
 		Dimensions = DimensionsIn;
+		Tile = FImageTile(FVector2i(0,0), FVector2i(Dimensions.GetWidth(), Dimensions.GetHeight()));
+		Multisampler = FGridSampler(Multisampling);
+	}
+
+	void Initialize(FImageDimensions DimensionsIn, const FImageTile& TileIn, int32 Multisampling = 1)
+	{
+		check(DimensionsIn.IsSquare());
+		Dimensions = DimensionsIn;
+		Tile = TileIn;
 		Multisampler = FGridSampler(Multisampling);
 	}
 
@@ -95,8 +106,8 @@ public:
 		// make flat mesh
 		TMeshAABBTree3<MeshType> FlatSpatial(&UVSpaceMesh, true);
 
-		const int32 LinearImageSize = Dimensions.Num();
-		const int32 LinearSampleSize = Dimensions.Num() * Multisampler.Num();
+		const int32 LinearImageSize = Tile.Num();
+		const int32 LinearSampleSize = Tile.Num() * Multisampler.Num();
 		TexelType.Init(EmptyTexel, LinearSampleSize);
 		TexelInteriorSamples.Init(0, LinearImageSize);
 		TexelQueryUV.Init(FVector2f::Zero(), LinearSampleSize);
@@ -111,13 +122,13 @@ public:
 		// find interior texels
 		TAtomic<int64> TotalGutterCounter = 0;
 		TArray< TArray64<TTuple<int64, int64>> > GutterTexelsPerScanline;
-		GutterTexelsPerScanline.SetNum(Dimensions.GetHeight());
+		GutterTexelsPerScanline.SetNum(Tile.GetHeight());
 
-		ParallelFor(Dimensions.GetHeight(),
+		ParallelFor(Tile.GetHeight(),
 					[this, &UVSpaceMesh, &GetTriangleIDFunc, &FlatSpatial, TexelDiag, &QueryOptions, &TexelSize, &TotalGutterCounter, &GutterTexelsPerScanline]
 		(int32 ImgY)
 		{
-			for (int32 ImgX = 0; ImgX < Dimensions.GetWidth(); ++ImgX)
+			for (int32 ImgX = 0; ImgX < Tile.GetWidth(); ++ImgX)
 			{
 				// With multiple samples, a texel may contain multiple gutter samples.
 				// Since we copy nearest interior texel over top our gutter texels, we
@@ -128,10 +139,12 @@ public:
 				TTuple<int64, int64> GutterNearestTexel;
 				for (int32 Sample = 0; Sample < Multisampler.Num(); ++Sample)
 				{
-					const int64 TexelLinearIdx = Dimensions.GetIndex(ImgX, ImgY);
+					const FVector2i SourceCoords = Tile.GetSourceCoords(ImgX, ImgY);
+					const int64 SourceTexelLinearIdx = Dimensions.GetIndex(SourceCoords.X, SourceCoords.Y);
+					const int64 TexelLinearIdx = Tile.GetIndex(ImgX, ImgY);
 					const int64 SampleLinearIdx = TexelLinearIdx * Multisampler.Num() + Sample;
 					const FVector2d SampleUV = Multisampler.Sample(Sample);
-					FVector2d UVPoint = Dimensions.GetTexelUV(TexelLinearIdx);
+					FVector2d UVPoint = Dimensions.GetTexelUV(SourceTexelLinearIdx);
 					UVPoint = UVPoint - 0.5 * TexelSize + SampleUV * TexelSize;
 					const FVector3d UVPoint3d(UVPoint.X, UVPoint.Y, 0);
 
@@ -174,23 +187,14 @@ public:
 							// To avoid artifacts with mipmapping the gutter texels store the nearest
 							// interior texel and a post pass copies those nearest interior texel values
 							// to each gutter pixel.
-							//
-							// The gutter pixel structure stores pixel indices to copy in the source
-							// image's linear index.
-							//
-							// For tiled images:
-							//    NearestCoords are in source space --> GetSource().GetIndex() will return source linear index.
-							//    ImgX, ImgY are in local tile space --> GetSourceIndex() will return source linear index.
-							const int64 NearestLinearIdx = Dimensions.IsTile() ? Dimensions.GetSourceDimensions().GetIndex(NearestCoords) : Dimensions.GetIndex(NearestCoords);
-							const int32 LinearSourceIdx = Dimensions.IsTile() ? Dimensions.GetSourceIndex(FVector2i(ImgX, ImgY)) : TexelLinearIdx;
-
-							GutterNearestTexel = TTuple<int64, int64>(LinearSourceIdx, NearestLinearIdx);
+							const int64 NearestLinearIdx = Dimensions.GetIndex(NearestCoords);
+							GutterNearestTexel = TTuple<int64, int64>(SourceTexelLinearIdx, NearestLinearIdx);
 						}
 					}
-					if (bIsGutterTexel)
-					{
-						GutterTexelsPerScanline[ImgY].Add(GutterNearestTexel);
-					}
+				}
+				if (bIsGutterTexel)
+				{
+					GutterTexelsPerScanline[ImgY].Add(GutterNearestTexel);
 				}
 			}
 
diff --git a/Engine/Source/Runtime/GeometryCore/Public/Image/ImageTile.h b/Engine/Source/Runtime/GeometryCore/Public/Image/ImageTile.h
new file mode 100644
--- /dev/null
+++ b/Engine/Source/Runtime/GeometryCore/Public/Image/ImageTile.h
@@ -0,0 +1,124 @@
+// Copyright Epic Games, Inc. All Rights Reserved.
+
+#pragma once
+
+#include "IntVectorTypes.h"
+#include "Image/ImageDimensions.h"
+
+namespace UE
+{
+namespace Geometry
+{
+
+class FImageTile
+{
+private:
+	FVector2i Start = FVector2i::Zero();
+	FVector2i End = FVector2i::Zero(); 
+
+public:
+	FImageTile() = default;
+	FImageTile(const FVector2i& InStart, const FVector2i& InEnd)
+		: Start(InStart), End(InEnd)
+	{
+	}
+
+	int32 GetWidth() const
+	{
+		return End.X - Start.X;
+	}
+
+	int32 GetHeight() const
+	{
+		return End.Y - Start.Y;
+	}
+
+	int64 Num() const
+	{
+		return GetWidth() * GetHeight();
+	}
+
+	/** @return the linear index for the given local XY coordinates into this tile. */
+	int64 GetIndex(const int32 X, const int32 Y) const
+	{
+		return Y * GetWidth() + X;
+	}
+
+	/** @return the linear index for the given local coordinates into this tile. */
+	int64 GetIndex(const FVector2i& LocalCoords) const
+	{
+		return GetIndex(LocalCoords.X, LocalCoords.Y);
+	}
+
+	/** @return the source image's coordinates given linear index into this tile. */
+	FVector2i GetSourceCoords(const int64 LinearIdx) const
+	{
+		return FVector2i(Start.X + LinearIdx % GetWidth(), Start.Y + LinearIdx / GetWidth());
+	}
+
+	/** @return the source image's coordinates given local XY coordinates into this tile. */
+	FVector2i GetSourceCoords(const int32 X, const int32 Y) const
+	{
+		return FVector2i(Start.X + X, Start.Y + Y);
+	}
+
+	/** @return the source image's coordinates given local coordinates into this tile. */
+	FVector2i GetSourceCoords(const FVector2i& LocalCoords) const
+	{
+		return GetSourceCoords(LocalCoords.X, LocalCoords.Y);
+	}
+};
+
+class FImageTiling
+{
+private:
+	FImageDimensions Dimensions;
+	int32 TileWidth = 32;
+	int32 TileHeight = 32;
+
+public:
+	FImageTiling(const FImageDimensions DimensionsIn, const int32 TileWidthIn, const int32 TileHeightIn)
+	{
+		SetTiling(DimensionsIn, TileWidthIn, TileHeightIn);
+	}
+
+	void SetTiling(const FImageDimensions DimensionsIn, const int32 TileWidthIn, const int32 TileHeightIn)
+	{
+		checkSlow(TileWidthIn >= 1 && TileHeightIn >= 1 && TilePaddingIn >= 0);
+		Dimensions = DimensionsIn;
+		TileWidth = FMath::Clamp(TileWidthIn, 1, Dimensions.GetWidth());
+		TileHeight = FMath::Clamp(TileHeightIn, 1, Dimensions.GetHeight());
+	}
+
+	int32 NumTilesX() const
+	{
+		return Dimensions.GetWidth() / TileWidth + (Dimensions.GetWidth() % TileWidth == 0 ? 0 : 1);
+	}
+
+	int32 NumTilesY() const
+	{
+		return Dimensions.GetHeight() / TileHeight + (Dimensions.GetHeight() % TileHeight == 0 ? 0 : 1);
+	}
+
+	int32 Num() const
+	{
+		return NumTilesX() * NumTilesY();
+	}
+
+	/** @return the tile data corresponding to the linear index [0, NumTiles()-1] w/ optional padding. */
+	FImageTile GetTile(const int32 LinearTileIndex, const int32 Padding = 0) const
+	{
+		const int32 TilePadding = Padding < 0 ? 0 : Padding;
+		const FVector2i TileIndex(LinearTileIndex % NumTilesX(), LinearTileIndex / NumTilesX());
+		
+		FVector2i TileStart, TileEnd;
+		TileStart.X = FMath::Clamp(TileIndex.X * TileWidth - TilePadding, 0, Dimensions.GetWidth());
+		TileStart.Y = FMath::Clamp(TileIndex.Y * TileHeight - TilePadding, 0, Dimensions.GetHeight());
+		TileEnd.X = FMath::Clamp((TileIndex.X + 1) * TileWidth + TilePadding, 0, Dimensions.GetWidth());
+		TileEnd.Y = FMath::Clamp((TileIndex.Y + 1) * TileHeight + TilePadding, 0, Dimensions.GetHeight());
+		return FImageTile(TileStart, TileEnd);
+	}
+};
+	
+} // end namespace UE::Geometry
+} // end namespace UE
\ No newline at end of file
