Fixed a build break in debug builds.

#rnx

diff --git a/Engine/Source/Runtime/GeometryCore/Public/Image/ImageTile.h b/Engine/Source/Runtime/GeometryCore/Public/Image/ImageTile.h
--- a/Engine/Source/Runtime/GeometryCore/Public/Image/ImageTile.h
+++ b/Engine/Source/Runtime/GeometryCore/Public/Image/ImageTile.h
@@ -84,7 +84,7 @@ public:
 
 	void SetTiling(const FImageDimensions DimensionsIn, const int32 TileWidthIn, const int32 TileHeightIn)
 	{
-		checkSlow(TileWidthIn >= 1 && TileHeightIn >= 1 && TilePaddingIn >= 0);
+		checkSlow(TileWidthIn >= 1 && TileHeightIn >= 1);
 		Dimensions = DimensionsIn;
 		TileWidth = FMath::Clamp(TileWidthIn, 1, Dimensions.GetWidth());
 		TileHeight = FMath::Clamp(TileHeightIn, 1, Dimensions.GetHeight());
