ModelingTools: add new FindPolygons Algorithm based on furthest-point sampling on surface.
- add optional support for distance-weighting/modulation in TMeshDijkstra. Replace previous FVector2d seed points w/ new FSeedPoint struct that provides an external ID. Track seed point source through graph propagation, add GetSeedExternalIDForPointSetID() to look up seedpoint source external ID for any computed point.
- add new TIncrementalMeshDijkstra, variant that supports incremental update (ie adding new seed points) of a dijkstra graph-distance solution. This variant assumes solution will be computed across the entire graph, so dynamic allocation of nodes is not supported, full set is allocated/initialized on graph initialization.
- Add new method to FFindPolygonsAlgorithm that generates face groups derived from approximate surface-voronoi diagram, computed based on a point set found by (approximate) incremental furthest-point sampling.
- Add optional min-group-size to FFindPolygonsAlgorithm, groups smaller than the given criteria will be appended to largest neighbouring group
- Expose above additions in ConvertToPolygonsTool

#rnx

diff --git a/Engine/Plugins/Experimental/MeshModelingToolset/Source/MeshModelingTools/Private/ConvertToPolygonsTool.cpp b/Engine/Plugins/Experimental/MeshModelingToolset/Source/MeshModelingTools/Private/ConvertToPolygonsTool.cpp
--- a/Engine/Plugins/Experimental/MeshModelingToolset/Source/MeshModelingTools/Private/ConvertToPolygonsTool.cpp
+++ b/Engine/Plugins/Experimental/MeshModelingToolset/Source/MeshModelingTools/Private/ConvertToPolygonsTool.cpp
@@ -38,7 +38,14 @@ USingleSelectionMeshEditingTool* UConvertToPolygonsToolBuilder::CreateNewTool(co
 class FConvertToPolygonsOp : public  FDynamicMeshOperator
 {
 public:
-	FConvertToPolygonsOp() : FDynamicMeshOperator() {}
+	// parameters set by the tool
+	EConvertToPolygonsMode ConversionMode = EConvertToPolygonsMode::FaceNormalDeviation;
+	double AngleTolerance = 0.1f;
+	int32 NumPoints = 10;
+	FFindPolygonsAlgorithm::EWeightingType WeightingType = FFindPolygonsAlgorithm::EWeightingType::None;
+	FVector3d WeightingCoeffs = FVector3d::One();
+	int32 MinGroupSize = 2;
+	bool bCalculateNormals = false;
 
 	virtual void CalculateResult(FProgressCancel* Progress) override
 	{
@@ -56,6 +63,7 @@ public:
 		}
 
 		Polygons = FFindPolygonsAlgorithm(ResultMesh.Get());
+		Polygons.MinGroupSize = MinGroupSize;
 
 		switch (ConversionMode)
 		{
@@ -75,6 +83,11 @@ public:
 			Polygons.FindPolygonsFromFaceNormals(DotTolerance);
 			break;
 		}
+		case EConvertToPolygonsMode::FromFurthestPointSampling:
+		{
+			Polygons.FindPolygonsFromFurthestPointSampling(NumPoints, WeightingType, WeightingCoeffs);
+			break;
+		}
 		default:
 			check(0);
 		}
@@ -113,11 +126,6 @@ public:
 	
 	FFindPolygonsAlgorithm Polygons;
 	TSharedPtr<FDynamicMesh3, ESPMode::ThreadSafe> OriginalMesh;
-
-	// parameters set by the tool
-	EConvertToPolygonsMode ConversionMode;
-	double AngleTolerance;
-	bool bCalculateNormals;
 };
 
 TUniquePtr<FDynamicMeshOperator> UConvertToPolygonsOperatorFactory::MakeNewOperator()
@@ -201,15 +209,13 @@ void UConvertToPolygonsTool::Setup()
 	// updates the triangle color visualization
 	UpdateVisualization();
 
-	Settings->WatchProperty(Settings->ConversionMode,
-							[this](EConvertToPolygonsMode NewMode)
-							{ OnSettingsModified(); });
-	Settings->WatchProperty(Settings->bShowGroupColors,
-							[this](bool bNewValue) 
-							{ UpdateVisualization(); });
-	Settings->WatchProperty(Settings->AngleTolerance,
-							[this](float value)
-							{ OnSettingsModified(); });
+	Settings->WatchProperty(Settings->ConversionMode, [this](EConvertToPolygonsMode) { OnSettingsModified(); });
+	Settings->WatchProperty(Settings->bShowGroupColors, [this](bool) { UpdateVisualization(); });
+	Settings->WatchProperty(Settings->AngleTolerance, [this](float) { OnSettingsModified(); });
+	Settings->WatchProperty(Settings->NumPoints, [this](int32) { OnSettingsModified(); });
+	Settings->WatchProperty(Settings->bNormalWeighted, [this](bool) { OnSettingsModified(); });
+	Settings->WatchProperty(Settings->NormalWeighting, [this](float) { OnSettingsModified(); });
+	Settings->WatchProperty(Settings->MinGroupSize, [this](int32) { OnSettingsModified(); });
 	
 
 	GetToolManager()->DisplayMessage(
@@ -222,6 +228,10 @@ void UConvertToPolygonsTool::UpdateOpParameters(FConvertToPolygonsOp& ConvertToP
 	ConvertToPolygonsOp.bCalculateNormals = Settings->bCalculateNormals;
 	ConvertToPolygonsOp.ConversionMode    = Settings->ConversionMode;
 	ConvertToPolygonsOp.AngleTolerance    = Settings->AngleTolerance;
+	ConvertToPolygonsOp.NumPoints = Settings->NumPoints;
+	ConvertToPolygonsOp.WeightingType = (Settings->bNormalWeighted) ? FFindPolygonsAlgorithm::EWeightingType::NormalDeviation : FFindPolygonsAlgorithm::EWeightingType::None;
+	ConvertToPolygonsOp.WeightingCoeffs = FVector3d(Settings->NormalWeighting, 1.0, 1.0);
+	ConvertToPolygonsOp.MinGroupSize = Settings->MinGroupSize;
 	ConvertToPolygonsOp.OriginalMesh = OriginalDynamicMesh;
 	
 	FTransform LocalToWorld = (FTransform)UE::ToolTarget::GetLocalToWorldTransform(Target);
diff --git a/Engine/Plugins/Experimental/MeshModelingToolset/Source/MeshModelingTools/Private/FindPolygonsAlgorithm.cpp b/Engine/Plugins/Experimental/MeshModelingToolset/Source/MeshModelingTools/Private/FindPolygonsAlgorithm.cpp
--- a/Engine/Plugins/Experimental/MeshModelingToolset/Source/MeshModelingTools/Private/FindPolygonsAlgorithm.cpp
+++ b/Engine/Plugins/Experimental/MeshModelingToolset/Source/MeshModelingTools/Private/FindPolygonsAlgorithm.cpp
@@ -4,6 +4,13 @@
 #include "DynamicMesh/MeshNormals.h"
 #include "Selections/MeshConnectedComponents.h"
 #include "Util/IndexUtil.h"
+#include "Parameterization/MeshDijkstra.h"
+#include "Parameterization/IncrementalMeshDijkstra.h"
+#include "Parameterization/MeshRegionGraph.h"
+
+#include "Curve/DynamicGraph3.h"
+
+
 
 #include "ExplicitUseGeometryMathTypes.h"		// using UE::Geometry::(math types)
 using namespace UE::Geometry;
@@ -31,6 +38,7 @@ bool FFindPolygonsAlgorithm::FindPolygonsFromUVIslands()
 		FoundPolygons.Add(Components.GetComponent(ci).Indices);
 	}
 
+	PostProcessPolygons(false);
 	SetGroupsFromPolygons();
 
 	return (FoundPolygons.Num() > 0);
@@ -55,6 +63,7 @@ bool FFindPolygonsAlgorithm::FindPolygonsFromConnectedTris()
 		FoundPolygons.Add(Components.GetComponent(ci).Indices);
 	}
 
+	PostProcessPolygons(false);
 	SetGroupsFromPolygons();
 
 	return (FoundPolygons.Num() > 0);
@@ -112,6 +121,192 @@ bool FFindPolygonsAlgorithm::FindPolygonsFromFaceNormals(double DotTolerance)
 		FoundPolygons.Add(Polygon);
 	}
 
+	PostProcessPolygons(true);
+	SetGroupsFromPolygons();
+
+	return (FoundPolygons.Num() > 0);
+}
+
+
+
+
+/**
+ * Dual graph of mesh faces, ie graph of edges across faces between face centers.
+ * Normals and Areas are tracked for each point.
+ */
+class FMeshFaceDualGraph : public FDynamicGraph3d
+{
+public:
+	TDynamicVectorN<double, 3> Normals;
+	TDynamicVector<double> Areas;
+
+	int AppendVertex(const FVector3d& Centroid, const FVector3d& Normal, double Area)
+	{
+		int vid = FDynamicGraph3d::AppendVertex(Centroid);
+		Normals.InsertAt({ {Normal.X, Normal.Y, Normal.Z} }, vid);
+		Areas.InsertAt(Area, vid);
+		return vid;
+	}
+
+	FVector3d GetNormal(int32 vid) const
+	{
+		return Normals.AsVector3(vid);
+	}
+
+	/** Build a Face Dual Graph for a triangle mesh */
+	static void MakeFaceDualGraphForMesh(FDynamicMesh3* Mesh, FMeshFaceDualGraph& FaceGraph)
+	{
+		for (int32 tid : Mesh->TriangleIndicesItr())
+		{
+			FVector3d Normal, Centroid;
+			double Area;
+			Mesh->GetTriInfo(tid, Normal, Area, Centroid);
+
+			int32 newvid = FaceGraph.AppendVertex(Centroid, Normal, Area);
+			check(newvid == tid);
+		}
+
+		for (int32 tid : Mesh->TriangleIndicesItr())
+		{
+			FIndex3i NbrT = Mesh->GetTriNeighbourTris(tid);
+			for (int32 j = 0; j < 3; ++j)
+			{
+				if (Mesh->IsTriangle(NbrT[j]))
+				{
+					FaceGraph.AppendEdge(tid, NbrT[j]);
+				}
+			}
+		}
+	}
+
+};
+
+
+
+
+
+
+bool FFindPolygonsAlgorithm::FindPolygonsFromFurthestPointSampling(int32 NumPoints, EWeightingType WeightingType, FVector3d WeightingCoeffs)
+{
+	NumPoints = FMath::Min(NumPoints, Mesh->VertexCount());
+
+	FMeshFaceDualGraph FaceGraph;
+	FMeshFaceDualGraph::MakeFaceDualGraphForMesh(Mesh, FaceGraph);
+
+	TIncrementalMeshDijkstra<FMeshFaceDualGraph> FurthestPoints(&FaceGraph);
+
+	TArray<int32> SeedIndices;
+
+	// need to add at least one seed point for each mesh connected component, so that all triangles are assigned a group
+	// TODO: two seed points for components that have no boundary?
+	FMeshConnectedComponents Components(Mesh);
+	Components.FindConnectedTriangles();
+	int32 NumConnected = Components.Num();
+	for (int32 k = 0; k < NumConnected; ++k)
+	{
+		FIndex3i TriVerts = Mesh->GetTriangle(Components[k].Indices[0]);
+		SeedIndices.Add(TriVerts.A);
+	}
+
+	// initial incremental update from per-component points
+	TArray<TIncrementalMeshDijkstra<FMeshFaceDualGraph>::FSeedPoint> ComponentSeeds;
+	for (int32 vid : SeedIndices)
+	{
+		ComponentSeeds.Add(TIncrementalMeshDijkstra<FMeshFaceDualGraph>::FSeedPoint{ vid, vid, 0.0 });
+	}
+	FurthestPoints.AddSeedPoints(ComponentSeeds);
+
+	// TODO: can (approximately) bound the size of a region based on mesh area. Then can pass that in
+	// with seed point as upper distance bound. This will change the result as it will initially grow on
+	// the 'front' (and possibly no guarantee that mesh is covered?)
+	//    (initial furthest-points would be any with an invalid value...then we can ensure coverage)
+	// FindMaxGraphDistancePointID() seems like it might be somewhat expensive...
+
+	// incrementally add furthest-points
+	while ( SeedIndices.Num() < NumPoints)
+	{
+		int32 NextPointID = FurthestPoints.FindMaxGraphDistancePointID();
+		if (NextPointID >= 0)
+		{
+			TIncrementalMeshDijkstra<FMeshFaceDualGraph>::FSeedPoint NextSeedPoint{ NextPointID, NextPointID, 0.0 };
+			FurthestPoints.AddSeedPoints({ NextSeedPoint });
+			SeedIndices.Add(NextPointID);
+		}
+		else
+		{
+			break;
+		}
+	}
+
+
+	// Now that we have furthest point set, recompute a Dijkstra propagation with optional weighting
+	// (unweighted version should be the same as the FurthestPoints dijkstra though...could re-use?)
+
+	TMeshDijkstra<FMeshFaceDualGraph> SuperPixels(&FaceGraph);
+	
+	if (WeightingType == EWeightingType::NormalDeviation)
+	{
+		SuperPixels.bEnableDistanceWeighting = true;
+		SuperPixels.GetWeightedDistanceFunc = [this, WeightingCoeffs, &FaceGraph](int32 FromVID, int32 ToVID, int32 SeedVID, double Distance)
+		{
+			FVector3d NA = FaceGraph.GetNormal(ToVID);
+			FVector3d NB = FaceGraph.GetNormal(SeedVID);
+			double Dot = NA.Dot(NB);
+			if (WeightingCoeffs.X > 0.001)
+			{
+				Dot = FMathd::Pow(Dot, WeightingCoeffs.X);
+			}
+			double W = FMathd::Clamp(1.0 - Dot * Dot, 0.0, 1.0);
+			W = W * W * W;
+			double Weight = FMathd::Clamp(W, 0.001, 1.0);
+			return Weight * Distance;
+		};
+	}
+
+	TArray<TMeshDijkstra<FMeshFaceDualGraph>::FSeedPoint> SuperPixelSeeds;
+	for (int32 k = 0; k < SeedIndices.Num(); ++k)
+	{
+		SuperPixelSeeds.Add(TMeshDijkstra<FMeshFaceDualGraph>::FSeedPoint{ k, SeedIndices[k], 0.0 });
+	}
+	SuperPixels.ComputeToMaxDistance(SuperPixelSeeds, TNumericLimits<double>::Max());
+
+	TArray<TArray<int32>> TriSets;
+	TriSets.SetNum(SeedIndices.Num());
+	TArray<int32> FailedSet;
+
+	for (int32 tid : Mesh->TriangleIndicesItr())
+	{
+		int32 SeedID = SuperPixels.GetSeedExternalIDForPointSetID(tid);
+		if (SeedID >= 0)
+		{
+			TriSets[SeedID].Add(tid);
+		}
+		else
+		{
+			FailedSet.Add(tid);
+		}
+	}
+
+
+	for (int32 k = 0; k < NumPoints; ++k)
+	{
+		if (TriSets[k].Num() > 0)
+		{
+			FoundPolygons.Add(MoveTemp(TriSets[k]));
+		}
+	}
+	
+	if (FailedSet.Num() > 0)
+	{
+		FMeshConnectedComponents FailedComponents(Mesh);
+		FailedComponents.FindConnectedTriangles(FailedSet);
+		for (FMeshConnectedComponents::FComponent& Component : FailedComponents)
+		{
+			FoundPolygons.Add(Component.Indices);
+		}
+	}
+
+	PostProcessPolygons(true);
 	SetGroupsFromPolygons();
 
 	return (FoundPolygons.Num() > 0);
@@ -119,6 +314,41 @@ bool FFindPolygonsAlgorithm::FindPolygonsFromFaceNormals(double DotTolerance)
 
 
 
+
+void FFindPolygonsAlgorithm::PostProcessPolygons(bool bApplyMerging)
+{
+	if (bApplyMerging && MinGroupSize > 1)
+	{
+		OptimizePolygons();
+	}
+}
+
+
+void FFindPolygonsAlgorithm::OptimizePolygons()
+{
+	FMeshRegionGraph RegionGraph;
+	RegionGraph.BuildFromTriangleSets(*Mesh, FoundPolygons, [&](int32 SetIdx) { return SetIdx; });
+	bool bMerged = RegionGraph.MergeSmallRegions(MinGroupSize-1, [&](int32 A, int32 B) { return RegionGraph.GetRegionTriCount(A) > RegionGraph.GetRegionTriCount(B); });
+	bool bSwapped = RegionGraph.OptimizeBorders();
+	if (bMerged || bSwapped)
+	{
+		FoundPolygons.Reset();
+
+		int32 N = RegionGraph.MaxRegionIndex();
+		for (int32 k = 0; k < N; ++k)
+		{
+			if (RegionGraph.IsRegion(k))
+			{
+				const TArray<int32>& Tris = RegionGraph.GetRegionTris(k);
+				FoundPolygons.Add(Tris);
+			}
+		}
+	}
+}
+
+
+
+
 void FFindPolygonsAlgorithm::SetGroupsFromPolygons()
 {
 	Mesh->EnableTriangleGroups(0);
@@ -157,6 +387,9 @@ void FFindPolygonsAlgorithm::SetGroupsFromPolygons()
 }
 
 
+
+
+
 bool FFindPolygonsAlgorithm::FindPolygonEdges()
 {
 	for (int eid : Mesh->EdgeIndicesItr())
diff --git a/Engine/Plugins/Experimental/MeshModelingToolset/Source/MeshModelingTools/Private/SeamSculptTool.cpp b/Engine/Plugins/Experimental/MeshModelingToolset/Source/MeshModelingTools/Private/SeamSculptTool.cpp
--- a/Engine/Plugins/Experimental/MeshModelingToolset/Source/MeshModelingTools/Private/SeamSculptTool.cpp
+++ b/Engine/Plugins/Experimental/MeshModelingToolset/Source/MeshModelingTools/Private/SeamSculptTool.cpp
@@ -218,8 +218,8 @@ void USeamSculptTool::UpdateCurrentDrawPath()
 
 	const FDynamicMesh3* Mesh = PreviewMesh->GetPreviewDynamicMesh();
 	TMeshDijkstra<FDynamicMesh3> PathFinder(Mesh);
-	TArray<FVector2d> SeedPoints;
-	SeedPoints.Add(FVector2d(DrawPathStartVertex, 0));
+	TArray<TMeshDijkstra<FDynamicMesh3>::FSeedPoint> SeedPoints;
+	SeedPoints.Add({ DrawPathStartVertex, DrawPathStartVertex, 0 });
 	if ( PathFinder.ComputeToTargetPoint(SeedPoints, CurrentSnapVertex) )
 	{
 		PathFinder.FindPathToNearestSeed(CurrentSnapVertex, CurDrawPath);
diff --git a/Engine/Plugins/Experimental/MeshModelingToolset/Source/MeshModelingTools/Private/UVProjectionTool.cpp b/Engine/Plugins/Experimental/MeshModelingToolset/Source/MeshModelingTools/Private/UVProjectionTool.cpp
--- a/Engine/Plugins/Experimental/MeshModelingToolset/Source/MeshModelingTools/Private/UVProjectionTool.cpp
+++ b/Engine/Plugins/Experimental/MeshModelingToolset/Source/MeshModelingTools/Private/UVProjectionTool.cpp
@@ -419,7 +419,7 @@ void UUVProjectionTool::OnMeshUpdated(UMeshOpPreviewWithBackgroundCompute* Previ
 	{
 		EdgeRenderer->CreateOrUpdateLineSet(TEXT("UVSeams"), UseMesh->MaxEdgeID(), [&](int32 eid, TArray<FRenderableLine>& LinesOut) {
 			FIndex2i EdgeT = UseMesh->GetEdgeT(eid);
-			if (TriangleROISet.Contains(EdgeT.A) && TriangleROISet.Contains(EdgeT.B) && UVOverlay->IsSeamEdge(eid))
+			if (UseMesh->IsEdge(eid) && TriangleROISet.Contains(EdgeT.A) && TriangleROISet.Contains(EdgeT.B) && UVOverlay->IsSeamEdge(eid))
 			{
 				AppendSeamEdge(eid, LinesOut);
 			}
@@ -428,7 +428,7 @@ void UUVProjectionTool::OnMeshUpdated(UMeshOpPreviewWithBackgroundCompute* Previ
 	else
 	{
 		EdgeRenderer->CreateOrUpdateLineSet(TEXT("UVSeams"), UseMesh->MaxEdgeID(), [&](int32 eid, TArray<FRenderableLine>& LinesOut) {
-			if (UVOverlay->IsSeamEdge(eid))
+			if (UseMesh->IsEdge(eid) && UVOverlay->IsSeamEdge(eid))
 			{
 				AppendSeamEdge(eid, LinesOut);
 			}
diff --git a/Engine/Plugins/Experimental/MeshModelingToolset/Source/MeshModelingTools/Public/ConvertToPolygonsTool.h b/Engine/Plugins/Experimental/MeshModelingToolset/Source/MeshModelingTools/Public/ConvertToPolygonsTool.h
--- a/Engine/Plugins/Experimental/MeshModelingToolset/Source/MeshModelingTools/Public/ConvertToPolygonsTool.h
+++ b/Engine/Plugins/Experimental/MeshModelingToolset/Source/MeshModelingTools/Public/ConvertToPolygonsTool.h
@@ -37,7 +37,9 @@ enum class EConvertToPolygonsMode
 	/** Create PolyGroups based on UV Islands */
 	FromUVIslands  UMETA(DisplayName = "From UV Islands"),
 	/** Create Polygroups based on Connected Triangles */
-	FromConnectedTris UMETA(DisplayName = "From Connected Tris")
+	FromConnectedTris UMETA(DisplayName = "From Connected Tris"),
+	/** Create Polygroups based on Connected Triangles */
+	FromFurthestPointSampling UMETA(DisplayName = "Furthest Point Sampling")
 };
 
 
@@ -50,15 +52,34 @@ class MESHMODELINGTOOLS_API UConvertToPolygonsToolProperties : public UInteracti
 public:
 	/** Strategy to use to group triangles */
 	UPROPERTY(EditAnywhere, Category = PolyGroups)
-	EConvertToPolygonsMode ConversionMode = EConvertToPolygonsMode::FaceNormalDeviation;
+		EConvertToPolygonsMode ConversionMode = EConvertToPolygonsMode::FromFurthestPointSampling;
+	//EConvertToPolygonsMode ConversionMode = EConvertToPolygonsMode::FaceNormalDeviation;
 
 	/** Tolerance for planarity */
 	UPROPERTY(EditAnywhere, Category = PolyGroups, meta = (UIMin = "0.001", UIMax = "20.0", ClampMin = "0.0", ClampMax = "90.0", EditCondition = "ConversionMode == EConvertToPolygonsMode::FaceNormalDeviation", EditConditionHides))
 	float AngleTolerance = 0.1f;
 
+	/** Furthest-Point Sample count, approximately this number of polygroups will be generated */
+	UPROPERTY(EditAnywhere, Category = PolyGroups, meta = (UIMin = "1", UIMax = "100", ClampMin = "1", ClampMax = "10000", EditCondition = "ConversionMode == EConvertToPolygonsMode::FromFurthestPointSampling", EditConditionHides))
+	int32 NumPoints = 100;
+
+	/** If true, region-growing in Sampling modes will be controlled by face normals, resulting in regions with borders that are more-aligned with curvature ridges */
+	UPROPERTY(EditAnywhere, Category = PolyGroups, meta = (EditCondition = "ConversionMode == EConvertToPolygonsMode::FromFurthestPointSampling", EditConditionHides))
+	bool bNormalWeighted = true;
+
+	/** This parameter modulates the effect of normal weighting during region-growing */
+	UPROPERTY(EditAnywhere, Category = PolyGroups, meta = (UIMin = "0.1", UIMax = "2.0", ClampMin = "0.01", ClampMax = "100.0", EditCondition = "ConversionMode == EConvertToPolygonsMode::FromFurthestPointSampling", EditConditionHides))
+	float NormalWeighting = 1.0f;
+
+
+	/** group filtering */
+	UPROPERTY(EditAnywhere, Category = Filtering, meta = (UIMin = "1", UIMax = "100", ClampMin = "1", ClampMax = "10000"))
+	int32 MinGroupSize = 2;
+
+
 	/** If true, normals are recomputed per-group, with hard edges at group boundaries */
-	UPROPERTY(EditAnywhere, Category = PolyGroups, meta = (EditCondition = "ConversionMode == EConvertToPolygonsMode::FaceNormalDeviation", EditConditionHides))
-	bool bCalculateNormals = true;
+	UPROPERTY(EditAnywhere, Category = Output)
+	bool bCalculateNormals = false;
 	
 	/** Display each group with a different auto-generated color */
 	UPROPERTY(EditAnywhere, Category = Display)
diff --git a/Engine/Plugins/Experimental/MeshModelingToolset/Source/MeshModelingTools/Public/FindPolygonsAlgorithm.h b/Engine/Plugins/Experimental/MeshModelingToolset/Source/MeshModelingTools/Public/FindPolygonsAlgorithm.h
--- a/Engine/Plugins/Experimental/MeshModelingToolset/Source/MeshModelingTools/Public/FindPolygonsAlgorithm.h
+++ b/Engine/Plugins/Experimental/MeshModelingToolset/Source/MeshModelingTools/Public/FindPolygonsAlgorithm.h
@@ -19,16 +19,36 @@ public:
 	TArray<int> PolygonTags;
 	TArray<FVector3d> PolygonNormals;
 
+	// if > 1, groups with a triangle count smaller than this will be merged with a neighbouring group
+	int32 MinGroupSize = 0;
+
 	TArray<int> PolygonEdges;
 
 	bool FindPolygonsFromFaceNormals(double DotTolerance = 0.0001);
 	bool FindPolygonsFromUVIslands();
 	bool FindPolygonsFromConnectedTris();
 
+	enum class EWeightingType
+	{
+		None,
+		NormalDeviation
+	};
+
+	/**
+	 * Incrementally compute approximate-geodesic-based furthest-point sampling of the mesh until NumPoints
+	 * samples have been found, then compute local geodesic patches (eg approximate surface voronoi diagaram).
+	 * Optionally weight geodesic-distance computation, which will produce different patch shapes.
+	 */
+	bool FindPolygonsFromFurthestPointSampling(int32 NumPoints, EWeightingType WeightingType, FVector3d WeightingCoeffs = FVector3d::One());
+
+
 	bool FindPolygonEdges();
 
 
 protected:
 
+	void PostProcessPolygons(bool bApplyMerging);
+	void OptimizePolygons();
+
 	void SetGroupsFromPolygons();
 };
\ No newline at end of file
diff --git a/Engine/Source/Runtime/GeometryCore/Public/Parameterization/IncrementalMeshDijkstra.h b/Engine/Source/Runtime/GeometryCore/Public/Parameterization/IncrementalMeshDijkstra.h
new file mode 100644
--- /dev/null
+++ b/Engine/Source/Runtime/GeometryCore/Public/Parameterization/IncrementalMeshDijkstra.h
@@ -0,0 +1,296 @@
+// Copyright Epic Games, Inc. All Rights Reserved.
+
+#pragma once
+
+#include "VectorTypes.h"
+#include "MatrixTypes.h"
+#include "BoxTypes.h"
+#include "FrameTypes.h"
+#include "Util/IndexPriorityQueue.h"
+#include "Util/DynamicVector.h"
+
+namespace UE
+{
+namespace Geometry
+{
+
+/**
+ * TIncrementalMeshDijkstra computes graph distances on a mesh from seed point(s) using Dijkstra's algorithm.
+ * Derived from TMeshDijkstra. This Incremental variant allows adding new Seed points to an already-computed
+ * solution.
+ */
+template<class PointSetType>
+class TIncrementalMeshDijkstra
+{
+public:
+	/** PointSet we are calculating on */
+	const PointSetType* PointSet;
+
+	/**
+	 * Return the 3D Position of a given Point in the PointSet. 
+	 * This function is set to PointSet->GetVertex() in the constructor below, but can be 
+	 * replaced with an external function if necessary (eg to provide deformed mesh positions, etc)
+	 */
+	TUniqueFunction<FVector3d(int32)> GetPositionFunc;
+
+
+	/**
+	 * Construct TMeshDijkstra for the given PointSet. We will hold a reference to this
+	 * PointSet for the lifetime of the class.
+	 */
+	TIncrementalMeshDijkstra(const PointSetType* PointSetIn)
+	{
+		PointSet = PointSetIn;
+
+		int32 MaxID = PointSet->MaxVertexID();
+		Queue.Initialize(MaxID);
+
+		GetPositionFunc = [this](int32 PointID) { return  PointSet->GetVertex(PointID); };
+
+		ForceInitializeAllNodes();
+	}
+
+	static double InvalidDistance() { return TNumericLimits<double>::Max(); };
+
+	/**
+	 *  FSeedPoint defines a seed point passed to the various compute methods below
+	 */
+	struct FSeedPoint
+	{
+		/** Client-defined integer ID for this seed point (not used internally) */
+		int32 ExternalID = -1;
+		/** Point ID for this seed point, must be a valid point ID in the input PointSet */
+		int32 PointID = 0;
+		/** Initial distance for this seed point */
+		double StartDistance = 0;
+	};
+
+
+	/**
+	 * Add new SeedPoints to the current solution, and then propagate updated
+	 * graph distances to any points closer to these new seed points than to
+	 * existing seed points.
+	 */
+	void AddSeedPoints(const TArray<FSeedPoint>& SeedPointsIn)
+	{
+		++CurrentSeedTimestamp;
+
+		for (const FSeedPoint& SeedPoint : SeedPointsIn)
+		{
+			check(SeedPointExternalIDMap.Contains(SeedPoint.ExternalID) == false);
+			int32 NewIndex = SeedPoints.Num();
+			SeedPointExternalIDMap.Add(SeedPoint.ExternalID, NewIndex);
+			SeedPoints.Add(SeedPoint);
+
+			int32 PointID = SeedPoint.PointID;
+			if (ensure(Queue.Contains(PointID) == false))
+			{
+				FGraphNode* Node = GetNodeForPointSetID(PointID);
+				Node->GraphDistance = SeedPoint.StartDistance;
+				Node->FrozenTimestamp = CurrentSeedTimestamp;
+				Node->SeedID = NewIndex;
+				Queue.Insert(PointID, Node->GraphDistance);
+			}
+		}
+
+
+		while (Queue.GetCount() > 0)
+		{
+			int32 NextID = Queue.Dequeue();
+			FGraphNode* Node = GetNodeForPointSetID(NextID);
+			check(Node != nullptr);
+
+			Node->FrozenTimestamp = CurrentSeedTimestamp;
+			UpdateNeighboursSparse(Node);
+		}
+	}
+
+	/**
+	 * @return the ExternalID of the SeedPoint that is closest to PointID, or -1 if PointID does not have a valid graph distance
+	 */
+	int32 GetSeedExternalIDForPointSetID(int32 PointID)
+	{
+		const FGraphNode* Node = GetNodeForPointSetID(PointID);
+		if (Node == nullptr || Node->FrozenTimestamp == InvalidFrozenTimestamp)
+		{
+			return -1;
+		}
+		int32 SeedIndex = Node->SeedID;
+		return SeedPoints[SeedIndex].ExternalID;
+	}
+
+
+	/**
+	 * @return vertex id associated with the maximum graph distance returned by GetMaxGraphDistance()
+	 */
+	int32 FindMaxGraphDistancePointID() const
+	{
+		int32 MaxID = PointSet->MaxVertexID();
+		double MaxDistance = 0;
+		int32 MaxDistancePointID = -1;
+		for (int32 PointID = 0; PointID < MaxID; PointID++)
+		{
+			if (PointSet->IsVertex(PointID))
+			{
+				double Distance = GetDistance(PointID);
+				if (Distance != InvalidDistance() && Distance > MaxDistance)
+				{
+					MaxDistance = Distance;
+					MaxDistancePointID = PointID;
+				}
+			}
+		}
+		return MaxDistancePointID;
+	}
+
+
+	/**
+	 * @return true if the distance for index PointID was calculated
+	 */
+	bool HasDistance(int32 PointID) const
+	{
+		const FGraphNode* Node = GetNodeForPointSetID(PointID);
+		return (Node != nullptr && Node->FrozenTimestamp != InvalidFrozenTimestamp);
+	}
+
+
+	/**
+	 * @return the distance calculated for index PointID
+	 */
+	double GetDistance(int32 PointID) const
+	{
+		const FGraphNode* Node = GetNodeForPointSetID(PointID);
+		return (Node != nullptr && Node->FrozenTimestamp != InvalidFrozenTimestamp) ? Node->GraphDistance : InvalidDistance();
+	}
+
+	/**
+	 * Find path from a point to the nearest seed point
+	 * @param PointID starting point, assumption is that we have computed dijkstra to this point
+	 * @param PathToSeedOut path is returned here, includes PointID and seed point as last element
+	 * @param MaxLength if PathToSeedOut grows beyond this length, we abort the search
+	 * @return true if valid path was found
+	 */
+	bool FindPathToNearestSeed(int32 PointID, TArray<int32>& PathToSeedOut, int32 MaxLength = 100000)
+	{
+		const FGraphNode* CurNode = GetNodeForPointSetID(PointID);
+		if (CurNode == nullptr || CurNode->FrozenTimestamp == InvalidFrozenTimestamp)
+		{
+			return false;
+		}
+
+		PathToSeedOut.Reset();
+		PathToSeedOut.Add(PointID);
+
+		int32 IterCount = 0;
+		while (IterCount++ < MaxLength)
+		{
+			if (CurNode->ParentPointID == -1)
+			{
+				return true;
+			}
+
+			PathToSeedOut.Add(CurNode->ParentPointID);
+
+			CurNode = GetNodeForPointSetID(CurNode->ParentPointID);
+			if (CurNode == nullptr || CurNode->FrozenTimestamp == InvalidFrozenTimestamp)
+			{
+				PathToSeedOut.Reset();
+				return false;
+			}
+		}
+		return true;
+	}
+
+private:
+
+	// information about each active/computed point
+	struct FGraphNode
+	{
+		int32 PointID;
+		int32 ParentPointID;
+		int32 SeedID;
+		double GraphDistance;
+		int32 FrozenTimestamp;
+	};
+
+	// currently all nodes are allocated on initialization
+	TArray<FGraphNode> AllocatedNodes;
+
+	// queue of nodes to process (for dijkstra front propagation)
+	FIndexPriorityQueue Queue;
+
+	TArray<FSeedPoint> SeedPoints;
+	TMap<int32, int32> SeedPointExternalIDMap;
+
+	const int32 InvalidFrozenTimestamp = -1;
+	int32 CurrentSeedTimestamp = 0;
+
+
+	void ForceInitializeAllNodes()
+	{
+		int32 MaxID = PointSet->MaxVertexID();
+		AllocatedNodes.SetNum(MaxID);
+		for (int32 k = 0; k < MaxID; ++k)
+		{
+			FGraphNode NewNode{ k, -1, -1, InvalidDistance(), InvalidFrozenTimestamp };
+			AllocatedNodes[k] = NewNode;
+		}
+	}
+
+
+	FGraphNode* GetNodeForPointSetID(int32 PointSetID)
+	{
+		return &AllocatedNodes[PointSetID];
+	}
+
+	const FGraphNode* GetNodeForPointSetID(int32 PointSetID) const
+	{
+		return &AllocatedNodes[PointSetID];
+	}
+
+
+	// given new Distance/UV at Parent, check if any of its neighbours are in the queue,
+	// and if they are, and the new graph distance is shorter, update their queue position
+	// (this is basically the update step of Disjktras algorithm)
+	void UpdateNeighboursSparse(FGraphNode* Parent)
+	{
+		FVector3d ParentPos(GetPositionFunc(Parent->PointID));
+		double ParentDist = Parent->GraphDistance;
+
+		for (int32 NbrPointID : PointSet->VtxVerticesItr(Parent->PointID))
+		{
+			FGraphNode* NbrNode = GetNodeForPointSetID(NbrPointID);
+			if (NbrNode->FrozenTimestamp == CurrentSeedTimestamp)
+			{
+				continue;
+			}
+
+			double NbrDist = ParentDist + Distance(ParentPos, GetPositionFunc(NbrPointID));
+			if (Queue.Contains(NbrPointID))
+			{
+				if (NbrDist < NbrNode->GraphDistance)
+				{
+					NbrNode->ParentPointID = Parent->PointID;
+					NbrNode->SeedID = Parent->SeedID;
+					NbrNode->GraphDistance = NbrDist;
+					Queue.Update(NbrPointID, NbrNode->GraphDistance);
+				}
+			}
+			else 
+			{
+				if (NbrDist < NbrNode->GraphDistance)
+				{
+					NbrNode->ParentPointID = Parent->PointID;
+					NbrNode->SeedID = Parent->SeedID;
+					NbrNode->GraphDistance = NbrDist;
+					Queue.Insert(NbrPointID, NbrNode->GraphDistance);
+				}
+			}
+		}
+	}
+
+};
+
+
+} // end namespace UE::Geometry
+} // end namespace UE
\ No newline at end of file
diff --git a/Engine/Source/Runtime/GeometryCore/Public/Parameterization/MeshDijkstra.h b/Engine/Source/Runtime/GeometryCore/Public/Parameterization/MeshDijkstra.h
--- a/Engine/Source/Runtime/GeometryCore/Public/Parameterization/MeshDijkstra.h
+++ b/Engine/Source/Runtime/GeometryCore/Public/Parameterization/MeshDijkstra.h
@@ -34,6 +34,19 @@ public:
 	 */
 	TUniqueFunction<FVector3d(int32)> GetPositionFunc;
 
+	/**
+	 * If enabled, when computing local point-pair distances, GetWeightedDistanceFunc() will be called with the 
+	 * Euclidean distance to allow for alternative metrics
+	 */
+	bool bEnableDistanceWeighting = false;
+
+	/**
+	 * Called when computing pairwise point distances between neighbours FromVID and ToVID, to allow for alternative distance metrics.
+	 * SeedVID is the seed point that FromVID's value was propagated from (ie point reached by gradient walk)
+	 * EuclideanDistance is the distance between the two input points
+	 */
+	TUniqueFunction<double(int32 FromVID, int32 ToVID, int32 SeedVID, double EuclideanDistance)> GetWeightedDistanceFunc;
+
 
 	/**
 	 * Construct TMeshDijkstra for the given PointSet. We will hold a reference to this
@@ -50,8 +63,13 @@ public:
 		MaxGraphDistancePointID = -1;
 
 		GetPositionFunc = [this](int32 PointID) { return  PointSet->GetVertex(PointID); };
+		bEnableDistanceWeighting = false;
+		GetWeightedDistanceFunc = [](int32, int32, int32, double Distance) { return Distance; };
 	}
 
+	/** @return distance value that indicates invalid or uncomputed distance */
+	static double InvalidDistance() { return TNumericLimits<double>::Max(); };
+
 	/**
 	 * Reset internal data structures but keep allocated memory
 	 */
@@ -65,31 +83,48 @@ public:
 	}
 
 
+	/**
+	 *  FSeedPoint defines a seed point passed to the various compute methods below
+	 */
+	struct FSeedPoint
+	{
+		/** Client-defined integer ID for this seed point (not used internally) */
+		int32 ExternalID = -1;
+		/** Point ID for this seed point, must be a valid point ID in the input PointSet */
+		int32 PointID = 0;
+		/** Initial distance for this seed point */
+		double StartDistance = 0;
+	};
+
 
 	/**
-	 * Computes outwards from seed points to all points that are less/equal to ComputeToMaxDistance from the seed.
-	 * @param SeedPointsIn seed points defined as 2D vector tuples, will be interpreted as (seed_point_vertex_id, seed_distance)
-	 * @param ComputeToMaxDistance target radius for parameterization, will not set UVs on points with graph-distance larger than this
+	 * Computes graph distances outwards from seed points to all points that are less/equal to ComputeToMaxDistance from the seed.
+	 * @param SeedPointsIn seed points used to initialize computation, ie geodesics propagate out from this point set
+	 * @param ComputeToMaxDistance target graph-distance radius, will not compute/set distances on points with graph-distance larger than this
 	 */
-	void ComputeToMaxDistance(const TArray<FVector2d>& SeedPointsIn, double ComputeToMaxDistanceIn)
+	void ComputeToMaxDistance(const TArray<FSeedPoint>& SeedPointsIn, double ComputeToMaxDistanceIn)
 	{
-		SeedPoints = SeedPointsIn;
 		MaxGraphDistance = 0.0f;
 		MaxGraphDistancePointID = -1;
 
-		for (const FVector2d& SeedPoint : SeedPoints )
+		SeedPoints.Reset();
+		for (const FSeedPoint& SeedPoint : SeedPointsIn)
 		{
-			int32 PointID = (int32)SeedPoint.X;
-			if (Queue.Contains(PointID) == false)
+			int32 NewIndex = SeedPoints.Num();
+			SeedPoints.Add(SeedPoint);
+
+			int32 PointID = SeedPoint.PointID;
+			if (ensure(Queue.Contains(PointID) == false))
 			{
 				FGraphNode* Node = GetNodeForPointSetID(PointID, true);
-				Node->GraphDistance = SeedPoint.Y;
+				Node->GraphDistance = SeedPoint.StartDistance;
 				Node->bFrozen = true;
+				Node->SeedPointID = NewIndex;
 				Queue.Insert(PointID, Node->GraphDistance);
 			}
 		}
 
-		while (Queue.GetCount() > 0) 
+		while (Queue.GetCount() > 0)
 		{
 			int32 NextID = Queue.Dequeue();
 			FGraphNode* Node = GetNodeForPointSetID(NextID, false);
@@ -110,26 +145,49 @@ public:
 
 
 	/**
-	 * Computes outwards from seed points to TargetPointID, or stop when all points are further than ComputeToMaxDistance from the seed
-	 * @param SeedPointsIn seed points defined as 2D vector tuples, will be interpreted as (seed_point_vertex_id, seed_distance)
-	 * @param TargetPointID
-	 * @param ComputeToMaxDistance target radius for parameterization, will not set UVs on points with graph-distance larger than this
-	 * @return true if TargetPointID was reached
+	 * Computes graph distances outwards from seed points to all points that are less/equal to ComputeToMaxDistance from the seed.
+	 * @param SeedPointsIn 2D tuples that define seed points as (PointID, InitialDistance) pairs
+	 * @param ComputeToMaxDistance target graph-distance radius, will not compute/set distances on points with graph-distance larger than this
 	 */
-	bool ComputeToTargetPoint(const TArray<FVector2d>& SeedPointsIn, int32 TargetPointID,  double ComputeToMaxDistanceIn = TNumericLimits<double>::Max())
+	void ComputeToMaxDistance(const TArray<FVector2d>& SeedPointsIn, double ComputeToMaxDistanceIn)
+	{
+		TArray<FSeedPoint> ConvertedSeedPoints;
+		int32 N = SeedPointsIn.Num();
+		ConvertedSeedPoints.SetNum(SeedPointsIn.Num());
+		for (int32 k = 0; k < N; ++k)
+		{
+			ConvertedSeedPoints[k] = FSeedPoint{ -1, (int)SeedPointsIn[k].X, SeedPointsIn[k].Y };
+		}
+		ComputeToMaxDistance(ConvertedSeedPoints, ComputeToMaxDistanceIn);
+	}
+
+
+
+	 /**
+	  * Computes graph distances outwards from seed points to all points that are less/equal to ComputeToMaxDistance from the seed, or until TargetPointID is reached.
+	  * This is useful for finding shortest paths between two points
+	  * @param SeedPointsIn seed points used to initialize computation, ie geodesics propagate out from this point set
+	  * @param TargetPointID the target point, computation stops when this point is reached
+	  * @param ComputeToMaxDistance target graph-distance radius, will not compute/set distances on points with graph-distance larger than this
+	  */
+	bool ComputeToTargetPoint(const TArray<FSeedPoint>& SeedPointsIn, int32 TargetPointID, double ComputeToMaxDistanceIn = TNumericLimits<double>::Max())
 	{
-		SeedPoints = SeedPointsIn;
 		MaxGraphDistance = 0.0f;
 		MaxGraphDistancePointID = -1;
 
-		for (const FVector2d& SeedPoint : SeedPoints)
+		SeedPoints.Reset();
+		for (const FSeedPoint& SeedPoint : SeedPointsIn)
 		{
-			int32 PointID = (int32)SeedPoint.X;
-			if (Queue.Contains(PointID) == false)
+			int32 NewIndex = SeedPoints.Num();
+			SeedPoints.Add(SeedPoint);
+
+			int32 PointID = SeedPoint.PointID;
+			if (ensure(Queue.Contains(PointID) == false))
 			{
 				FGraphNode* Node = GetNodeForPointSetID(PointID, true);
-				Node->GraphDistance = SeedPoint.Y;
+				Node->GraphDistance = SeedPoint.StartDistance;
 				Node->bFrozen = true;
+				Node->SeedPointID = NewIndex;
 				Queue.Insert(PointID, Node->GraphDistance);
 			}
 		}
@@ -181,6 +239,21 @@ public:
 	}
 
 
+	/**
+	 * @return the ExternalID of the SeedPoint that is closest to PointID, or -1 if PointID does not have a valid graph distance
+	 */
+	int32 GetSeedExternalIDForPointSetID(int32 PointID)
+	{
+		const FGraphNode* Node = GetNodeForPointSetID(PointID);
+		if (Node == nullptr || Node->bFrozen == false)
+		{
+			return -1;
+		}
+		int32 SeedIndex = Node->SeedPointID;
+		return SeedPoints[SeedIndex].ExternalID;
+	}
+
+
 	/**
 	 * @return true if the distance for index PointID was calculated
 	 */
@@ -197,25 +270,9 @@ public:
 	double GetDistance(int32 PointID) const
 	{
 		const FGraphNode* Node = GetNodeForPointSetID(PointID);
-		return (Node != nullptr && Node->bFrozen) ? Node->GraphDistance : InvalidValue();
+		return (Node != nullptr && Node->bFrozen) ? Node->GraphDistance : InvalidDistance();
 	}
 
-	
-	//bool GetInterpTriDistance(const FIndex3i Triangle, const FVector3d& BaryCoords, double& DistanceOut) const
-	//{
-	//	const FGraphNode* NodeA = GetNodeForPointSetID(Triangle.A);
-	//	const FGraphNode* NodeB = GetNodeForPointSetID(Triangle.B);
-	//	const FGraphNode* NodeC = GetNodeForPointSetID(Triangle.C);
-	//	if (NodeA == nullptr || NodeA.bFrozen == false || NodeB == nullptr || NodeB.bFrozen == false || NodeC == nullptr || NodeC.bFrozen == false)
-	//	{
-	//		DistanceOut = InvalidValue();
-	//		return false;
-	//	}
-
-	//	DistanceOut = BaryCoords.X*NodeA->GraphDistance + BaryCoords.Y*NodeB->GraphDistance + BaryCoords.Z*NodeC->GraphDistance;
-	//	return true;
-	//}
-
 	/**
 	 * Find path from a point to the nearest seed point
 	 * @param PointID starting point, assumption is that we have computed dijkstra to this point
@@ -261,6 +318,7 @@ private:
 	{
 		int32 PointID;
 		int32 ParentPointID;
+		int32 SeedPointID;
 		double GraphDistance;
 		bool bFrozen;
 	};
@@ -273,9 +331,7 @@ private:
 	// queue of nodes to process (for dijkstra front propagation)
 	FIndexPriorityQueue Queue;
 
-	static double InvalidValue() { return TNumericLimits<double>::Max(); };
-
-	TArray<FVector2d> SeedPoints;
+	TArray<FSeedPoint> SeedPoints;
 
 	// max distances encountered during last compute
 	double MaxGraphDistance;
@@ -330,13 +386,21 @@ private:
 				continue;
 			}
 
-			double NbrDist = ParentDist + Distance(ParentPos, GetPositionFunc(NbrPointID));
+			double LocalDist = Distance(ParentPos, GetPositionFunc(NbrPointID));
+			if (bEnableDistanceWeighting)
+			{
+				int SeedPointID = SeedPoints[Parent->SeedPointID].PointID;
+				LocalDist = GetWeightedDistanceFunc(Parent->PointID, NbrPointID, SeedPointID, LocalDist);
+			}
+			double NbrDist = ParentDist + LocalDist;
+
 			if (Queue.Contains(NbrPointID))
 			{
 				if (NbrDist < NbrNode->GraphDistance)
 				{
 					NbrNode->ParentPointID = Parent->PointID;
 					NbrNode->GraphDistance = NbrDist;
+					NbrNode->SeedPointID = Parent->SeedPointID;
 					Queue.Update(NbrPointID, NbrNode->GraphDistance);
 				}
 			}
@@ -344,6 +408,7 @@ private:
 			{
 				NbrNode->ParentPointID = Parent->PointID;
 				NbrNode->GraphDistance = NbrDist;
+				NbrNode->SeedPointID = Parent->SeedPointID;
 				Queue.Insert(NbrPointID, NbrNode->GraphDistance);
 			}
 		}
