GeometryProcessing: minor cleanups. Handle invalid UVs in UVPacking. Remove ensure in FDynamicMeshUVEditor::EstimateGeodesicCenterFrameVertex(), this is easily hit in UV generation and unavoidable from the users PoV. Fix possible infinite loop in FMeshRegionGraph::MergeSmallRegionsPass(). Add VectorUtil::SignedArea().

#rnx

diff --git a/Engine/Plugins/Runtime/GeometryProcessing/Source/DynamicMesh/Private/Parameterization/DynamicMeshUVEditor.cpp b/Engine/Plugins/Runtime/GeometryProcessing/Source/DynamicMesh/Private/Parameterization/DynamicMeshUVEditor.cpp
--- a/Engine/Plugins/Runtime/GeometryProcessing/Source/DynamicMesh/Private/Parameterization/DynamicMeshUVEditor.cpp
+++ b/Engine/Plugins/Runtime/GeometryProcessing/Source/DynamicMesh/Private/Parameterization/DynamicMeshUVEditor.cpp
@@ -197,7 +197,7 @@ bool FDynamicMeshUVEditor::EstimateGeodesicCenterFrameVertex(const FDynamicMesh3
 	FVector3d Normal = FMeshNormals::ComputeVertexNormal(Mesh, VertexIDOut);
 
 	FMeshBoundaryLoops LoopsCalc(&Mesh, true);
-	if (ensure(LoopsCalc.GetLoopCount() > 0) == false )
+	if (LoopsCalc.GetLoopCount() == 0)
 	{
 		FrameOut = Mesh.GetVertexFrame(VertexIDOut, false, &Normal);
 		return false;
diff --git a/Engine/Plugins/Runtime/GeometryProcessing/Source/DynamicMesh/Private/Parameterization/MeshRegionGraph.cpp b/Engine/Plugins/Runtime/GeometryProcessing/Source/DynamicMesh/Private/Parameterization/MeshRegionGraph.cpp
--- a/Engine/Plugins/Runtime/GeometryProcessing/Source/DynamicMesh/Private/Parameterization/MeshRegionGraph.cpp
+++ b/Engine/Plugins/Runtime/GeometryProcessing/Source/DynamicMesh/Private/Parameterization/MeshRegionGraph.cpp
@@ -157,7 +157,7 @@ bool FMeshRegionGraph::MergeSmallRegionsPass(int32 SmallThreshold,
 		}
 	}
 
-	return true;
+	return bMergedAny;
 }
 
 
diff --git a/Engine/Source/Runtime/GeometryCore/Private/Parameterization/UVPacking.cpp b/Engine/Source/Runtime/GeometryCore/Private/Parameterization/UVPacking.cpp
--- a/Engine/Source/Runtime/GeometryCore/Private/Parameterization/UVPacking.cpp
+++ b/Engine/Source/Runtime/GeometryCore/Private/Parameterization/UVPacking.cpp
@@ -783,6 +783,11 @@ void FStandardChartPacker::RasterizeChart(const FUVIsland& Chart, uint32 RectW,
 	for (int32 tid : Chart.Triangles)
 	{
 		FIndex3i UVTriangle = Mesh->GetUVTriangle(tid);
+		if (UVTriangle.A < 0 || UVTriangle.B < 0 || UVTriangle.C < 0)
+		{
+			continue;
+		}
+
 		FVector2f Points[3];
 		for (int k = 0; k <3; k++)
 		{
@@ -861,9 +866,12 @@ bool FUVPacker::StandardPack(IUVMeshView* Mesh, int NumIslands, TFunctionRef<voi
 
 		for (int32 elemid : IslandElements)
 		{
-			FVector2d UV = (FVector2d)Mesh->GetUV(elemid);
-			FVector2d TransformedUV = UV.X * Chart.PackingScaleU + UV.Y * Chart.PackingScaleV + Chart.PackingBias;
-			Mesh->SetUV(elemid, (FVector2f)TransformedUV);
+			if (elemid >= 0)
+			{
+				FVector2d UV = (FVector2d)Mesh->GetUV(elemid);
+				FVector2d TransformedUV = UV.X * Chart.PackingScaleU + UV.Y * Chart.PackingScaleV + Chart.PackingBias;
+				Mesh->SetUV(elemid, (FVector2f)TransformedUV);
+			}
 		}
 
 	}
@@ -918,17 +926,23 @@ bool FUVPacker::StackPack(IUVMeshView* Mesh, int NumIslands, TFunctionRef<void(i
 		for (int32 tid : Island)
 		{
 			FIndex3i UVTri = Mesh->GetUVTriangle(tid);
-			IslandElements.Add(UVTri[0]);
-			IslandElements.Add(UVTri[1]);
-			IslandElements.Add(UVTri[2]);
+			if (UVTri[0] >= 0 && UVTri[1] >= 0 && UVTri[2] >= 0)
+			{
+				IslandElements.Add(UVTri[0]);
+				IslandElements.Add(UVTri[1]);
+				IslandElements.Add(UVTri[2]);
+			}
 		}
 
 		double ScaleFactor = UseUniformScale * AllIslandScaleFactors[ci];
 		for (int32 elemid : IslandElements)
 		{
-			FVector2d CurUV = (FVector2d)Mesh->GetUV(elemid);
-			FVector2d NewUV = (CurUV - IslandBounds.Min) * ScaleFactor;
-			Mesh->SetUV(elemid, (FVector2f)NewUV);
+			if (elemid >= 0)
+			{
+				FVector2d CurUV = (FVector2d)Mesh->GetUV(elemid);
+				FVector2d NewUV = (CurUV - IslandBounds.Min) * ScaleFactor;
+				Mesh->SetUV(elemid, (FVector2f)NewUV);
+			}
 		}
 	}
 
diff --git a/Engine/Source/Runtime/GeometryCore/Public/VectorUtil.h b/Engine/Source/Runtime/GeometryCore/Public/VectorUtil.h
--- a/Engine/Source/Runtime/GeometryCore/Public/VectorUtil.h
+++ b/Engine/Source/Runtime/GeometryCore/Public/VectorUtil.h
@@ -109,6 +109,15 @@ namespace VectorUtil
 		return (RealType)0.5 * TMathUtil<RealType>::Abs(CrossZ);
 	}
 
+	/**
+	 * @return signed area of 2D triangle V0,V1,V2
+	 */
+	template <typename RealType>
+	inline RealType SignedArea(const FVector2<RealType>& V0, const FVector2<RealType>& V1, const FVector2<RealType>& V2)
+	{
+		return ((RealType)0.5) * ((V0.X * V1.Y - V0.Y * V1.X) + (V1.X * V2.Y - V1.Y * V2.X) + (V2.X * V0.Y - V2.Y * V0.X));
+	}
+
 	/**
 	 * @return true if triangle V1,V2,V3 is obtuse
 	 */
