Fix for crash when calling TPolygon2<>::Contains() on an empty polygon.

diff --git a/Engine/Source/Runtime/GeometryCore/Public/Polygon2.h b/Engine/Source/Runtime/GeometryCore/Public/Polygon2.h
--- a/Engine/Source/Runtime/GeometryCore/Public/Polygon2.h
+++ b/Engine/Source/Runtime/GeometryCore/Public/Polygon2.h
@@ -411,6 +411,10 @@ public:
 		int nWindingNumber = 0;
 
 		int N = Vertices.Num();
+		if (N == 0)
+		{
+			return false;
+		}
 		FVector2<T> a = Vertices[0], b = FVector2<T>::Zero();
 		for (int i = 0; i < N; ++i) 
 		{
