GeometryCore: Fix gutter list entries for samples outside search space in ImageOccupancyMap.

#rnx

diff --git a/Engine/Source/Runtime/GeometryCore/Public/Image/ImageOccupancyMap.h b/Engine/Source/Runtime/GeometryCore/Public/Image/ImageOccupancyMap.h
--- a/Engine/Source/Runtime/GeometryCore/Public/Image/ImageOccupancyMap.h
+++ b/Engine/Source/Runtime/GeometryCore/Public/Image/ImageOccupancyMap.h
@@ -136,6 +136,7 @@ public:
 				// texel are gutter texels. This bool will track that as we iterate
 				// over each sample of the texel.
 				bool bIsGutterTexel = true;
+				bool bHitTriID = false;
 				TTuple<int64, int64> GutterNearestTexel;
 				for (int32 Sample = 0; Sample < Multisampler.Num(); ++Sample)
 				{
@@ -190,9 +191,10 @@ public:
 							const int64 NearestLinearIdx = Dimensions.GetIndex(NearestCoords);
 							GutterNearestTexel = TTuple<int64, int64>(SourceTexelLinearIdx, NearestLinearIdx);
 						}
+						bHitTriID = true;
 					}
 				}
-				if (bIsGutterTexel)
+				if (bIsGutterTexel && bHitTriID)
 				{
 					GutterTexelsPerScanline[ImgY].Add(GutterNearestTexel);
 				}
