GeometryProcessingUnitTests: Added unit tests for FMeshMapBaker.

#rnx

diff --git a/Engine/Source/Runtime/GeometryCore/Public/Image/ImageTile.h b/Engine/Source/Runtime/GeometryCore/Public/Image/ImageTile.h
--- a/Engine/Source/Runtime/GeometryCore/Public/Image/ImageTile.h
+++ b/Engine/Source/Runtime/GeometryCore/Public/Image/ImageTile.h
@@ -28,6 +28,16 @@ public:
 	{
 	}
 
+	FVector2i GetStart() const
+	{
+		return Start;
+	}
+
+	FVector2i GetEnd() const
+	{
+		return End;
+	}
+
 	int32 GetWidth() const
 	{
 		return Width;
diff --git a/Engine/Source/Runtime/GeometryCore/Public/VectorUtil.h b/Engine/Source/Runtime/GeometryCore/Public/VectorUtil.h
--- a/Engine/Source/Runtime/GeometryCore/Public/VectorUtil.h
+++ b/Engine/Source/Runtime/GeometryCore/Public/VectorUtil.h
@@ -167,6 +167,13 @@ namespace VectorUtil
 		return EpsilonEqual(V0.X, V1.X, Epsilon) && EpsilonEqual(V0.Y, V1.Y, Epsilon) && EpsilonEqual(V0.Z, V1.Z, Epsilon);
 	}
 
+	/** @return true if all coordinates of V0 and V1 are within Epsilon of each other */
+	template <typename RealType>
+	inline bool EpsilonEqual(const TVector4<RealType>& V0, const TVector4<RealType>& V1, RealType Epsilon)
+	{
+		return EpsilonEqual(V0.X, V1.X, Epsilon) && EpsilonEqual(V0.Y, V1.Y, Epsilon) && EpsilonEqual(V0.Z, V1.Z, Epsilon) && EpsilonEqual(V0.W, V1.W, Epsilon);
+	}
+
 
 	/** @return 0/1/2 index of smallest value in Vector3 */
 	template <typename ValueVecType>
