MeshModelingTools: Fixing mesh overlay ref count issues when setting an island triangle, when retriangulating elements without wanting them to be freed, and when doing a check for overlay validity.

#rnx

diff --git a/Engine/Source/Runtime/GeometryCore/Private/DynamicMesh/DynamicMeshOverlay.cpp b/Engine/Source/Runtime/GeometryCore/Private/DynamicMesh/DynamicMeshOverlay.cpp
--- a/Engine/Source/Runtime/GeometryCore/Private/DynamicMesh/DynamicMeshOverlay.cpp
+++ b/Engine/Source/Runtime/GeometryCore/Private/DynamicMesh/DynamicMeshOverlay.cpp
@@ -363,7 +363,7 @@ void TDynamicMeshOverlay<RealType, ElementSize>::InitializeTriangles(int MaxTria
 
 
 template<typename RealType, int ElementSize>
-EMeshResult TDynamicMeshOverlay<RealType, ElementSize>::SetTriangle(int tid, const FIndex3i& tv)
+EMeshResult TDynamicMeshOverlay<RealType, ElementSize>::SetTriangle(int tid, const FIndex3i& tv, bool bAllowElementFreeing)
 {
 	if (IsElement(tv[0]) == false || IsElement(tv[1]) == false || IsElement(tv[2]) == false)
 	{
@@ -376,12 +376,40 @@ EMeshResult TDynamicMeshOverlay<RealType, ElementSize>::SetTriangle(int tid, con
 		return EMeshResult::Failed_InvalidNeighbourhood;
 	}
 
-	InternalSetTriangle(tid, tv, true);
+	InternalSetTriangle(tid, tv, true, bAllowElementFreeing);
 
 	//updateTimeStamp(true);
 	return EMeshResult::Ok;
 }
 
+template<typename RealType, int ElementSize>
+void UE::Geometry::TDynamicMeshOverlay<RealType, ElementSize>::FreeUnusedElements(const TSet<int>* ElementsToCheck)
+{
+	auto FreeIfUnused = [this](int ElementID)
+	{
+		if (ElementsRefCounts.IsValid(ElementID) && ElementsRefCounts.GetRefCount(ElementID) == 1)
+		{
+			ElementsRefCounts.Decrement(ElementID);
+			ParentVertices[ElementID] = FDynamicMesh3::InvalidID;
+		}
+	};
+
+	if (ElementsToCheck)
+	{
+		for (int ElementID : *ElementsToCheck)
+		{
+			FreeIfUnused(ElementID);
+		}
+	}
+	else
+	{
+		for (int ElementID : ElementsRefCounts.Indices())
+		{
+			FreeIfUnused(ElementID);
+		}
+	}
+}
+
 template<typename RealType, int ElementSize>
 void TDynamicMeshOverlay<RealType, ElementSize>::UnsetTriangle(int TriangleID)
 {
@@ -405,18 +433,26 @@ void TDynamicMeshOverlay<RealType, ElementSize>::UnsetTriangle(int TriangleID)
 
 
 template<typename RealType, int ElementSize>
-void TDynamicMeshOverlay<RealType, ElementSize>::InternalSetTriangle(int tid, const FIndex3i& tv, bool bUpdateRefCounts)
+void TDynamicMeshOverlay<RealType, ElementSize>::InternalSetTriangle(int tid, const FIndex3i& tv, bool bUpdateRefCounts, bool bAllowElementFreeing)
 {
 	if (ensure(ParentMesh) == false)
 	{
 		return;
 	}
 
+	// If we have to decrement refcounts, we will do it at the end, because Decrement() frees
+	// elements as soon as they lose their last reference, so a Decrement followed by Increment
+	// can leave things in an invalid state.
+	bool bNeedToDecrement = false;
+	FIndex3i OldTriElements; // only used if need to decrement.
+
 	int i = 3 * tid;
 
+	// See if triangle existed and make it exist if not
 	if (!ElementTriangles.SetMinimumSize(i + 3, FDynamicMesh3::InvalidID) && bUpdateRefCounts)
 	{
-		UnsetTriangle(tid);
+		OldTriElements = GetTriangle(tid);
+		bNeedToDecrement = (OldTriElements[0] != FDynamicMesh3::InvalidID);
 	}
 
 	ElementTriangles[i + 2] = tv[2];
@@ -428,6 +464,20 @@ void TDynamicMeshOverlay<RealType, ElementSize>::InternalSetTriangle(int tid, co
 		ElementsRefCounts.Increment(tv[0]);
 		ElementsRefCounts.Increment(tv[1]);
 		ElementsRefCounts.Increment(tv[2]);
+
+		if (bNeedToDecrement)
+		{
+			for (int j = 0; j < 3; ++j) 
+			{
+				ElementsRefCounts.Decrement(OldTriElements[j]);
+
+				if (bAllowElementFreeing && ElementsRefCounts.GetRefCount(OldTriElements[j]) == 1)
+				{
+					ElementsRefCounts.Decrement(OldTriElements[j]);
+					ParentVertices[OldTriElements[j]] = FDynamicMesh3::InvalidID;
+				}
+			};
+		}
 	}
 
 	if (tv != FDynamicMesh3::InvalidTriangle)
@@ -1413,10 +1463,11 @@ bool TDynamicMeshOverlay<RealType, ElementSize>::CheckValidity(bool bAllowNonMan
 		}
 	}
 	// verify that refcount list counts are same as actual reference counts
-	for (int elemid : ElementsRefCounts.Indices())
+	for (int32 ElementID = 0; ElementID < RealRefCounts.Num(); ++ElementID)
 	{
-		int CurRefCount = ElementsRefCounts.GetRefCount(elemid);
-		CheckOrFailF(CurRefCount == RealRefCounts[elemid] + 1);
+		int32 CurRefCount = ElementsRefCounts.GetRefCount(ElementID);
+		CheckOrFailF((RealRefCounts[ElementID] == 0 && CurRefCount == 0)
+			|| (RealRefCounts[ElementID] != 0 && CurRefCount == RealRefCounts[ElementID] + 1));
 	}
 
 	return is_ok;
diff --git a/Engine/Source/Runtime/GeometryCore/Public/DynamicMesh/DynamicMeshOverlay.h b/Engine/Source/Runtime/GeometryCore/Public/DynamicMesh/DynamicMeshOverlay.h
--- a/Engine/Source/Runtime/GeometryCore/Public/DynamicMesh/DynamicMeshOverlay.h
+++ b/Engine/Source/Runtime/GeometryCore/Public/DynamicMesh/DynamicMeshOverlay.h
@@ -246,8 +246,25 @@ public:
 
 	/** Initialize the triangle list to the given size, and set all triangles to InvalidID */
 	void InitializeTriangles(int MaxTriangleID);
-	/** Set the triangle to the given Element index tuple, and increment element reference counts */
-	EMeshResult SetTriangle(int TriangleID, const FIndex3i& TriElements);
+
+	/**
+	 * Set the triangle to the given Element index tuple, and increment element reference counts 
+	 * 
+	 * @param bAllowElementFreeing If true, then any elements that were only referenced by this triangle
+	 *  become immediately unallocated if the triangle no longer references them. This can be set to false
+	 *  when remeshing across existing elements to avoid them being freed while temporarily unreferenced,
+	 *  but then it should eventually be followed by a call to FreeUnusedElements().
+	 */
+	EMeshResult SetTriangle(int TriangleID, const FIndex3i& TriElements, bool bAllowElementFreeing = true);
+
+	/**
+	 * Goes through elements and frees any whose reference counts indicate that they are not being used. This
+	 * is usually not necessary since most operations that remove references will go ahead and do this, but
+	 * it may be used, for instance, after SetTriangle is called with bAllowElementFreeing set to false.
+	 * 
+	 * @param ElementsToCheck If provided, only these element ID's will be checked.
+	 */
+	void FreeUnusedElements(const TSet<int>* ElementsToCheck = nullptr);
 
 	/**
 	 * Set the triangle to have InvalidID element IDs, decrementing element reference counts if needed.
@@ -516,7 +533,7 @@ protected:
 	void SetElementFromBary(int SetElement, int ElementA, int ElementB, int ElementC, const FVector3<double>& BaryCoords);
 
 	/** updates the triangles array and optionally the element reference counts */
-	void InternalSetTriangle(int TriangleID, const FIndex3i& TriElements, bool bUpdateRefCounts);
+	void InternalSetTriangle(int TriangleID, const FIndex3i& TriElements, bool bUpdateRefCounts, bool bAllowElementFreeing = true);
 
 };
 
diff --git a/Engine/Source/Runtime/GeometryCore/Public/Util/RefCountVector.h b/Engine/Source/Runtime/GeometryCore/Public/Util/RefCountVector.h
--- a/Engine/Source/Runtime/GeometryCore/Public/Util/RefCountVector.h
+++ b/Engine/Source/Runtime/GeometryCore/Public/Util/RefCountVector.h
@@ -105,6 +105,7 @@ public:
 
 	int Increment(int Index, unsigned short IncrementCount = 1)
 	{
+		checkSlow(RefCounts[Index] != INVALID_REF_COUNT);
 		RefCounts[Index] += IncrementCount;
 		return RefCounts[Index];
 	}
@@ -154,7 +155,8 @@ public:
 			{
 				if (FreeIndices[i] == Index)
 				{
-					FreeIndices[i] = INDEX_NONE;
+					FreeIndices[i] = FreeIndices.Back();
+					FreeIndices.PopBack();
 					RefCounts[Index] = 1;
 					UsedCount++;
 					return true;
