fix 'center' vertex of midpoint-fan caps for arrows and generalized cylinders

#rnx

diff --git a/Engine/Source/Runtime/GeometryCore/Public/Generators/SweepGenerator.h b/Engine/Source/Runtime/GeometryCore/Public/Generators/SweepGenerator.h
--- a/Engine/Source/Runtime/GeometryCore/Public/Generators/SweepGenerator.h
+++ b/Engine/Source/Runtime/GeometryCore/Public/Generators/SweepGenerator.h
@@ -474,7 +474,7 @@ public:
 			{
 				if (Caps[CapIdx] == ECapType::FlatMidpointFan)
 				{
-					Vertices[CapVertStart[CapIdx]] = FVector3d::UnitZ() * (double)Heights[CapIdx];
+					Vertices[CapVertStart[CapIdx]] = FVector3d::UnitZ() * (double)Heights[CapIdx * (Heights.Num() - 1)];
 				}
 			}
 
@@ -706,6 +706,15 @@ public:
 		}
 		if (bCapped && !bLoop)
 		{
+			// if capped, set vertices.
+			for (int CapIdx = 0; CapIdx < 2; CapIdx++)
+			{
+				if (Caps[CapIdx] == ECapType::FlatMidpointFan)
+				{
+					Vertices[CapVertStart[CapIdx]] = Path[CapIdx * (Path.Num() - 1)];
+				}
+			}
+
 			for (int CapIdx = 0; CapIdx < 2; CapIdx++)
 			{
 				FVector3d Normal = CurveUtil::Tangent<double, FVector3d>(Path, CapIdx * (PathNum - 1), bLoop) * (double)(CapIdx * 2 - 1);
