Remove duplicate FDynamicMesh vertex triangle one ring function, optimize the remaining function some, and remove bool that doesn't change result.

#rnx

diff --git a/Engine/Plugins/Experimental/UVEditor/Source/UVEditorTools/Private/UVSelectTool.cpp b/Engine/Plugins/Experimental/UVEditor/Source/UVEditorTools/Private/UVSelectTool.cpp
--- a/Engine/Plugins/Experimental/UVEditor/Source/UVEditorTools/Private/UVSelectTool.cpp
+++ b/Engine/Plugins/Experimental/UVEditor/Source/UVEditorTools/Private/UVSelectTool.cpp
@@ -593,7 +593,7 @@ void UUVSelectTool::OnSelectionChanged()
 					}
 
 					TArray<int> TidOneRing;
-					Selection.Mesh->GetVertexOneRingTriangles(EdgeVids[i], TidOneRing);
+					Selection.Mesh->GetVtxTriangles(EdgeVids[i], TidOneRing);
 					for (int32 Tid : TidOneRing)
 					{
 						if (!TidSet.Contains(Tid))
@@ -616,7 +616,7 @@ void UUVSelectTool::OnSelectionChanged()
 				}
 
 				TArray<int> TidOneRing;
-				Selection.Mesh->GetVertexOneRingTriangles(Vid, TidOneRing);
+				Selection.Mesh->GetVtxTriangles(Vid, TidOneRing);
 				for (int32 Tid : TidOneRing)
 				{
 					if (!TidSet.Contains(Tid))
diff --git a/Engine/Plugins/Runtime/GeometryProcessing/Source/DynamicMesh/Private/NormalFlowRemesher.cpp b/Engine/Plugins/Runtime/GeometryProcessing/Source/DynamicMesh/Private/NormalFlowRemesher.cpp
--- a/Engine/Plugins/Runtime/GeometryProcessing/Source/DynamicMesh/Private/NormalFlowRemesher.cpp
+++ b/Engine/Plugins/Runtime/GeometryProcessing/Source/DynamicMesh/Private/NormalFlowRemesher.cpp
@@ -296,7 +296,7 @@ void FNormalFlowRemesher::TrackedFaceProjectionPass(double& MaxDistanceMoved, bo
 		TempPosBuffer[VertexIndex] = { 0.0, 0.0, 0.0 };
 
 		TArray<int> IncidentTriangles;
-		Mesh->GetVertexOneRingTriangles(VertexIndex, IncidentTriangles);
+		Mesh->GetVtxTriangles(VertexIndex, IncidentTriangles);
 
 		for (int TriID : IncidentTriangles)
 		{
diff --git a/Engine/Plugins/Runtime/GeometryProcessing/Source/DynamicMesh/Private/Operations/EmbedSurfacePath.cpp b/Engine/Plugins/Runtime/GeometryProcessing/Source/DynamicMesh/Private/Operations/EmbedSurfacePath.cpp
--- a/Engine/Plugins/Runtime/GeometryProcessing/Source/DynamicMesh/Private/Operations/EmbedSurfacePath.cpp
+++ b/Engine/Plugins/Runtime/GeometryProcessing/Source/DynamicMesh/Private/Operations/EmbedSurfacePath.cpp
@@ -892,7 +892,7 @@ bool UE::Geometry::EmbedProjectedPaths(FDynamicMesh3* Mesh, const TArrayView<con
 				return false;
 			}
 			TArray<int> TrianglesOut;
-			Mesh->GetVertexOneRingTriangles(OutPathVertices.Last(), TrianglesOut);
+			Mesh->GetVtxTriangles(OutPathVertices.Last(), TrianglesOut);
 
 			check(TrianglesOut.Num());
 			CurrentSeedTriID = TrianglesOut[0];
diff --git a/Engine/Source/Runtime/GeometryCore/Private/DynamicMesh/DynamicMesh3.cpp b/Engine/Source/Runtime/GeometryCore/Private/DynamicMesh/DynamicMesh3.cpp
--- a/Engine/Source/Runtime/GeometryCore/Private/DynamicMesh/DynamicMesh3.cpp
+++ b/Engine/Source/Runtime/GeometryCore/Private/DynamicMesh/DynamicMesh3.cpp
@@ -610,45 +610,28 @@ FIndex3i FDynamicMesh3::GetTriNeighbourTris(int tID) const
 }
 
 
-void FDynamicMesh3::GetVertexOneRingTriangles(int VertexID, TArray<int>& TrianglesOut) const
+void FDynamicMesh3::EnumerateVertexTriangles(int32 VertexID, TFunctionRef<void(int32)> ApplyFunc) const
 {
 	checkSlow(VertexRefCounts.IsValid(VertexID));
-	if (VertexRefCounts.IsValid(VertexID))
+	if (!IsVertex(VertexID))
 	{
-		VertexEdgeLists.Enumerate(VertexID, [&](int32 EdgeID)
-		{
-			FIndex2i EdgePair = GetOrderedOneRingEdgeTris(VertexID, EdgeID);
-			if (EdgePair.A != InvalidID)
-			{
-				TrianglesOut.Add(EdgePair.A);
-			}
-			if (EdgePair.B != InvalidID)
-			{
-				TrianglesOut.Add(EdgePair.B);
-			}
-		});
+		return;
 	}
-}
-
 
-void FDynamicMesh3::EnumerateVertexTriangles(int32 VertexID, TFunctionRef<void(int32)> ApplyFunc) const
-{
-	checkSlow(VertexRefCounts.IsValid(VertexID));
-	if (VertexRefCounts.IsValid(VertexID))
+	VertexEdgeLists.Enumerate(VertexID, [&](int32 eid)
 	{
-		VertexEdgeLists.Enumerate(VertexID, [&](int32 EdgeID)
+		const FEdge Edge = Edges[eid];
+		const int vOther = Edge.Vert.A == VertexID ? Edge.Vert.B : Edge.Vert.A;
+		if (TriHasSequentialVertices(Edge.Tri[0], VertexID, vOther))
 		{
-			FIndex2i EdgePair = GetOrderedOneRingEdgeTris(VertexID, EdgeID);
-			if (EdgePair.A != InvalidID)
-			{
-				ApplyFunc(EdgePair.A);
-			}
-			if (EdgePair.B != InvalidID)
-			{
-				ApplyFunc(EdgePair.B);
-			}
-		});
-	}
+			ApplyFunc(Edge.Tri[0]);
+		}
+		const int et1 = Edge.Tri[1];
+		if (Edge.Tri[1] != InvalidID && TriHasSequentialVertices(Edge.Tri[1], VertexID, vOther))
+		{
+			ApplyFunc(Edge.Tri[1]);
+		}
+	});
 }
 
 
@@ -965,10 +948,8 @@ bool FDynamicMesh3::CheckValidity(FValidityOptions ValidityOptions, EValidityChe
 			CheckOrFailF(IsEdge(edge));
 		}
 
-		TArray<int> vTris, vTris2;
-		GetVtxTriangles(vID, vTris, false);
-		GetVtxTriangles(vID, vTris2, true);
-		CheckOrFailF(vTris.Num() == vTris2.Num());
+		TArray<int> vTris;
+		GetVtxTriangles(vID, vTris);
 		//System.Console.WriteLine(string.Format("{0} {1} {2}", vID, vTris.Count, GetVtxEdges(vID).Count));
 		if (ValidityOptions.bAllowNonManifoldVertices)
 		{
diff --git a/Engine/Source/Runtime/GeometryCore/Private/DynamicMesh/DynamicMesh3_Edits.cpp b/Engine/Source/Runtime/GeometryCore/Private/DynamicMesh/DynamicMesh3_Edits.cpp
--- a/Engine/Source/Runtime/GeometryCore/Private/DynamicMesh/DynamicMesh3_Edits.cpp
+++ b/Engine/Source/Runtime/GeometryCore/Private/DynamicMesh/DynamicMesh3_Edits.cpp
@@ -621,7 +621,7 @@ EMeshResult FDynamicMesh3::RemoveVertex(int vID, bool bRemoveAllTriangles, bool
 		}
 
 		TArray<int> tris;
-		GetVtxTriangles(vID, tris, true);
+		GetVtxTriangles(vID, tris);
 		for (int tID : tris)
 		{
 			EMeshResult result = RemoveTriangle(tID, false, bPreserveManifold);
diff --git a/Engine/Source/Runtime/GeometryCore/Private/DynamicMesh/DynamicMesh3_Queries.cpp b/Engine/Source/Runtime/GeometryCore/Private/DynamicMesh/DynamicMesh3_Queries.cpp
--- a/Engine/Source/Runtime/GeometryCore/Private/DynamicMesh/DynamicMesh3_Queries.cpp
+++ b/Engine/Source/Runtime/GeometryCore/Private/DynamicMesh/DynamicMesh3_Queries.cpp
@@ -90,36 +90,27 @@ void FDynamicMesh3::GetVtxNbrhood(int eID, int vID, int& vOther, int& oppV1, int
 }
 
 
-int FDynamicMesh3::GetVtxTriangleCount(int vID, bool bBruteForce) const
+int FDynamicMesh3::GetVtxTriangleCount(int vID) const
 {
-	if (bBruteForce)
-	{
-		TArray<int> vTriangles;
-		if (GetVtxTriangles(vID, vTriangles, false) != EMeshResult::Ok)
-		{
-			return -1;
-		}
-		return (int)vTriangles.Num();
-	}
-
 	if (!IsVertex(vID))
 	{
 		return -1;
 	}
 	int N = 0;
-	for (int eid : VertexEdgeLists.Values(vID))
+	VertexEdgeLists.Enumerate(vID, [&](int32 eid)
 	{
-		int vOther = GetOtherEdgeVertex(eid, vID);
 		const FEdge Edge = Edges[eid];
+		const int vOther = Edge.Vert.A == vID ? Edge.Vert.B : Edge.Vert.A;
 		if (TriHasSequentialVertices(Edge.Tri[0], vID, vOther))
 		{
 			N++;
 		}
+		const int et1 = Edge.Tri[1];
 		if (Edge.Tri[1] != InvalidID && TriHasSequentialVertices(Edge.Tri[1], vID, vOther))
 		{
 			N++;
 		}
-	}
+	});
 	return N;
 }
 
@@ -140,19 +131,19 @@ int FDynamicMesh3::GetVtxSingleTriangle(int VertexID) const
 }
 
 
-EMeshResult FDynamicMesh3::GetVtxTriangles(int vID, TArray<int>& TrianglesOut, bool bUseOrientation) const
+EMeshResult FDynamicMesh3::GetVtxTriangles(int vID, TArray<int>& TrianglesOut) const
 {
 	if (!IsVertex(vID))
 	{
 		return EMeshResult::Failed_NotAVertex;
 	}
 
-	if (bUseOrientation)
+	if (VertexEdgeLists.GetCount(vID) > 20)
 	{
-		for (int eid : VertexEdgeLists.Values(vID))
+		VertexEdgeLists.Enumerate(vID, [&](int32 eid)
 		{
 			const FEdge Edge = Edges[eid];
-			const int vOther = GetOtherEdgeVertex(eid, vID);
+			const int vOther = Edge.Vert.A == vID ? Edge.Vert.B : Edge.Vert.A;
 			if (TriHasSequentialVertices(Edge.Tri[0], vID, vOther))
 			{
 				TrianglesOut.Add(Edge.Tri[0]);
@@ -162,12 +153,11 @@ EMeshResult FDynamicMesh3::GetVtxTriangles(int vID, TArray<int>& TrianglesOut, b
 			{
 				TrianglesOut.Add(Edge.Tri[1]);
 			}
-		}
+		});
 	}
 	else
 	{
-		// brute-force method
-		for (int eid : VertexEdgeLists.Values(vID))
+		VertexEdgeLists.Enumerate(vID, [&](int32 eid)
 		{
 			const FEdge Edge = Edges[eid];
 			TrianglesOut.AddUnique(Edge.Tri[0]);
@@ -175,7 +165,7 @@ EMeshResult FDynamicMesh3::GetVtxTriangles(int vID, TArray<int>& TrianglesOut, b
 			{
 				TrianglesOut.AddUnique(Edge.Tri[1]);
 			}
-		}
+		});
 	}
 	return EMeshResult::Ok;
 }
diff --git a/Engine/Source/Runtime/GeometryCore/Public/DynamicMesh/DynamicMesh3.h b/Engine/Source/Runtime/GeometryCore/Public/DynamicMesh/DynamicMesh3.h
--- a/Engine/Source/Runtime/GeometryCore/Public/DynamicMesh/DynamicMesh3.h
+++ b/Engine/Source/Runtime/GeometryCore/Public/DynamicMesh/DynamicMesh3.h
@@ -598,9 +598,6 @@ public:
 	/** @return the max valence of all vertices in the mesh */
 	int GetMaxVtxEdgeCount() const;
 
-	/** adds triangles touching VertexID to the TrianglesOut list */
-	void GetVertexOneRingTriangles(int VertexID, TArray<int>& TrianglesOut) const;
-
 	/** Get triangle vertices */
 	inline FIndex3i GetTriangle(int TriangleID) const
 	{
@@ -882,16 +879,13 @@ public:
 
 	/**
 	 * return # of triangles attached to vID, or -1 if invalid vertex
-	 * if bBruteForce = true, explicitly checks, which creates a list and is expensive
-	 * default is false, uses orientation, no memory allocation
 	 */
-	int GetVtxTriangleCount(int VertexID, bool bBruteForce = false) const;
+	int GetVtxTriangleCount(int VertexID) const;
 
 	/**
 	 * Get triangle one-ring at vertex.
-	 * bUseOrientation is more efficient but returns incorrect result if vertex is a bowtie
 	 */
-	EMeshResult GetVtxTriangles(int VertexID, TArray<int>& TrianglesOut, bool bUseOrientation) const;
+	EMeshResult GetVtxTriangles(int VertexID, TArray<int>& TrianglesOut) const;
 
 	/**
 	 * @return Triangle ID for a single triangle connected to VertexID, or InvalidID if VertexID does not exist or has no attached triangles
@@ -1404,8 +1398,8 @@ protected:
 
 	inline bool TriHasSequentialVertices(int TriangleID, int vA, int vB) const
 	{
-		int v0 = Triangles[TriangleID][0], v1 = Triangles[TriangleID][1], v2 = Triangles[TriangleID][2];
-		return ((v0 == vA && v1 == vB) || (v1 == vA && v2 == vB) || (v2 == vA && v0 == vB));
+		const FIndex3i& Tri = Triangles[TriangleID];
+		return ((Tri.A == vA && Tri.B == vB) || (Tri.B == vA && Tri.C == vB) || (Tri.C == vA && Tri.A == vB));
 	}
 
 	int FindTriangleEdge(int TriangleID, int vA, int vB) const;
