GeometryProcessing: remove FDynamicMesh3::CachedIsClosed()
#rnx

diff --git a/Engine/Source/Runtime/GeometryCore/Private/DynamicMesh/DynamicMesh3.cpp b/Engine/Source/Runtime/GeometryCore/Private/DynamicMesh/DynamicMesh3.cpp
--- a/Engine/Source/Runtime/GeometryCore/Private/DynamicMesh/DynamicMesh3.cpp
+++ b/Engine/Source/Runtime/GeometryCore/Private/DynamicMesh/DynamicMesh3.cpp
@@ -48,10 +48,7 @@ FDynamicMesh3::FDynamicMesh3(const FDynamicMesh3& Other)
 	TopologyTimestamp{ Other.TopologyTimestamp },
 
 	CachedBoundingBox{ Other.CachedBoundingBox },
-	CachedBoundingBoxTimestamp{ Other.CachedBoundingBoxTimestamp },
-
-	bIsClosedCached{ Other.bIsClosedCached },
-	CachedIsClosedTimestamp{ Other.CachedIsClosedTimestamp }
+	CachedBoundingBoxTimestamp{ Other.CachedBoundingBoxTimestamp }
 {
 	if (Other.HasAttributes())
 	{
@@ -84,10 +81,7 @@ FDynamicMesh3::FDynamicMesh3(FDynamicMesh3&& Other)
 	TopologyTimestamp{ MoveTemp( Other.TopologyTimestamp ) },
 
 	CachedBoundingBox{ MoveTemp( Other.CachedBoundingBox ) },
-	CachedBoundingBoxTimestamp{ MoveTemp( Other.CachedBoundingBoxTimestamp ) },
-
-	bIsClosedCached{ Other.bIsClosedCached },
-	CachedIsClosedTimestamp{ MoveTemp( Other.CachedIsClosedTimestamp ) }
+	CachedBoundingBoxTimestamp{ MoveTemp( Other.CachedBoundingBoxTimestamp ) }
 {
 	if (AttributeSet)
 	{
@@ -133,8 +127,6 @@ const FDynamicMesh3 & FDynamicMesh3::operator=(FDynamicMesh3 && Other)
 	CachedBoundingBox = MoveTemp(Other.CachedBoundingBox);
 	CachedBoundingBoxTimestamp = MoveTemp(Other.CachedBoundingBoxTimestamp);
 
-	bIsClosedCached = MoveTemp(Other.bIsClosedCached);
-	CachedIsClosedTimestamp = MoveTemp( Other.CachedIsClosedTimestamp);
 	return *this;
 }
 
@@ -364,9 +356,6 @@ void FDynamicMesh3::Clear()
 
 	CachedBoundingBox = FAxisAlignedBox3d();
 	CachedBoundingBoxTimestamp = -1;
-
-	bIsClosedCached = false;
-	CachedIsClosedTimestamp = -1;
 }
 
 
@@ -646,7 +635,7 @@ FString FDynamicMesh3::MeshInfoString()
 	FString AttribString = FString::Printf(TEXT("VtxNormals %d  VtxColors %d  VtxUVs %d  TriGroups %d  Attributes %d"),
 		HasVertexNormals(), HasVertexColors(), HasVertexUVs(), HasTriangleGroups(), HasAttributes());
 	FString InfoString = FString::Printf(TEXT("Closed %d  Compact %d  Timestamp %d  ShapeTimestamp %d  TopologyTimestamp %d  MaxGroupID %d"),
-		GetCachedIsClosed(), IsCompact(), GetTimestamp(), GetShapeTimestamp(), GetTopologyTimestamp(), MaxGroupID());
+		IsClosed(), IsCompact(), GetTimestamp(), GetShapeTimestamp(), GetTopologyTimestamp(), MaxGroupID());
 
 	return VtxString + "\n" + TriString + "\n" + EdgeString + "\n" + AttribString + "\n" + InfoString;
 }
diff --git a/Engine/Source/Runtime/GeometryCore/Private/DynamicMesh/DynamicMesh3_Queries.cpp b/Engine/Source/Runtime/GeometryCore/Private/DynamicMesh/DynamicMesh3_Queries.cpp
--- a/Engine/Source/Runtime/GeometryCore/Private/DynamicMesh/DynamicMesh3_Queries.cpp
+++ b/Engine/Source/Runtime/GeometryCore/Private/DynamicMesh/DynamicMesh3_Queries.cpp
@@ -668,17 +668,6 @@ bool FDynamicMesh3::IsClosed() const
 	return true;
 }
 
-
-bool FDynamicMesh3::GetCachedIsClosed()
-{
-	if (CachedIsClosedTimestamp != GetTopologyTimestamp())
-	{
-		bIsClosedCached = IsClosed();
-		CachedIsClosedTimestamp = GetTopologyTimestamp();
-	}
-	return bIsClosedCached;
-}
-
 // average of 1 or 2 face normals
 FVector3d FDynamicMesh3::GetEdgeNormal(int eID) const
 {
diff --git a/Engine/Source/Runtime/GeometryCore/Public/DynamicMesh/DynamicMesh3.h b/Engine/Source/Runtime/GeometryCore/Public/DynamicMesh/DynamicMesh3.h
--- a/Engine/Source/Runtime/GeometryCore/Public/DynamicMesh/DynamicMesh3.h
+++ b/Engine/Source/Runtime/GeometryCore/Public/DynamicMesh/DynamicMesh3.h
@@ -182,10 +182,6 @@ protected:
 	FAxisAlignedBox3d CachedBoundingBox;
 	/** timestamp for CachedBoundingBox, if less than current timestamp, cache is invalid */
 	int CachedBoundingBoxTimestamp = -1;
-	/** Cached value of IsClosed() */
-	bool bIsClosedCached = false;
-	/** timestamp for bIsClosedCached, if less than current timestamp, cache is invalid */
-	int CachedIsClosedTimestamp = -1;
 
 public:
 	/** Default constructor */
@@ -948,10 +944,6 @@ public:
 	/** @return true if mesh has no boundary edges */
 	bool IsClosed() const;
 
-	/** @return value of IsClosed() and cache result. cache is invalidated and recomputed if topology has changed since last call */
-	bool GetCachedIsClosed();
-
-
 	/** Timestamp is incremented any time any change is made to the mesh */
 	inline int GetTimestamp() const
 	{
