GeometryProcessing: remove unused cached bounding-box support in FDynamicMesh3
#rnx

diff --git a/Engine/Source/Runtime/GeometryCore/Private/DynamicMesh/DynamicMesh3.cpp b/Engine/Source/Runtime/GeometryCore/Private/DynamicMesh/DynamicMesh3.cpp
--- a/Engine/Source/Runtime/GeometryCore/Private/DynamicMesh/DynamicMesh3.cpp
+++ b/Engine/Source/Runtime/GeometryCore/Private/DynamicMesh/DynamicMesh3.cpp
@@ -45,10 +45,7 @@ FDynamicMesh3::FDynamicMesh3(const FDynamicMesh3& Other)
 
 	Timestamp{ Other.Timestamp },
 	ShapeTimestamp{ Other.ShapeTimestamp },
-	TopologyTimestamp{ Other.TopologyTimestamp },
-
-	CachedBoundingBox{ Other.CachedBoundingBox },
-	CachedBoundingBoxTimestamp{ Other.CachedBoundingBoxTimestamp }
+	TopologyTimestamp{ Other.TopologyTimestamp }
 {
 	if (Other.HasAttributes())
 	{
@@ -78,10 +75,7 @@ FDynamicMesh3::FDynamicMesh3(FDynamicMesh3&& Other)
 	AttributeSet{ MoveTemp( Other.AttributeSet ) },
 	Timestamp{ MoveTemp( Other.Timestamp ) },
 	ShapeTimestamp{ MoveTemp( Other.ShapeTimestamp ) },
-	TopologyTimestamp{ MoveTemp( Other.TopologyTimestamp ) },
-
-	CachedBoundingBox{ MoveTemp( Other.CachedBoundingBox ) },
-	CachedBoundingBoxTimestamp{ MoveTemp( Other.CachedBoundingBoxTimestamp ) }
+	TopologyTimestamp{ MoveTemp( Other.TopologyTimestamp ) }
 {
 	if (AttributeSet)
 	{
@@ -124,9 +118,6 @@ const FDynamicMesh3 & FDynamicMesh3::operator=(FDynamicMesh3 && Other)
 	ShapeTimestamp = MoveTemp(Other.ShapeTimestamp);
 	TopologyTimestamp = MoveTemp(Other.TopologyTimestamp);
 
-	CachedBoundingBox = MoveTemp(Other.CachedBoundingBox);
-	CachedBoundingBoxTimestamp = MoveTemp(Other.CachedBoundingBoxTimestamp);
-
 	return *this;
 }
 
@@ -353,9 +344,6 @@ void FDynamicMesh3::Clear()
 	Timestamp = 0;
 	ShapeTimestamp = 0;
 	TopologyTimestamp = 0;
-
-	CachedBoundingBox = FAxisAlignedBox3d();
-	CachedBoundingBoxTimestamp = -1;
 }
 
 
diff --git a/Engine/Source/Runtime/GeometryCore/Private/DynamicMesh/DynamicMesh3_Queries.cpp b/Engine/Source/Runtime/GeometryCore/Private/DynamicMesh/DynamicMesh3_Queries.cpp
--- a/Engine/Source/Runtime/GeometryCore/Private/DynamicMesh/DynamicMesh3_Queries.cpp
+++ b/Engine/Source/Runtime/GeometryCore/Private/DynamicMesh/DynamicMesh3_Queries.cpp
@@ -637,19 +637,6 @@ FAxisAlignedBox3d FDynamicMesh3::GetBounds(bool bParallel) const
 }
 
 
-FAxisAlignedBox3d FDynamicMesh3::GetCachedBounds()
-{
-	if (CachedBoundingBoxTimestamp != GetShapeTimestamp())
-	{
-		CachedBoundingBox = GetBounds();
-		CachedBoundingBoxTimestamp = GetShapeTimestamp();
-	}
-	return CachedBoundingBox;
-}
-
-
-
-
 bool FDynamicMesh3::IsClosed() const
 {
 	if (TriangleCount() == 0)
diff --git a/Engine/Source/Runtime/GeometryCore/Public/DynamicMesh/DynamicMesh3.h b/Engine/Source/Runtime/GeometryCore/Public/DynamicMesh/DynamicMesh3.h
--- a/Engine/Source/Runtime/GeometryCore/Public/DynamicMesh/DynamicMesh3.h
+++ b/Engine/Source/Runtime/GeometryCore/Public/DynamicMesh/DynamicMesh3.h
@@ -177,12 +177,6 @@ protected:
 	/** The topology timestamp is incremented any time a function that modifies the mesh topology is called */
 	int TopologyTimestamp = 0;
 
-
-	/** Cached vertex bounding box (includes unreferenced vertices) */
-	FAxisAlignedBox3d CachedBoundingBox;
-	/** timestamp for CachedBoundingBox, if less than current timestamp, cache is invalid */
-	int CachedBoundingBoxTimestamp = -1;
-
 public:
 	/** Default constructor */
 	FDynamicMesh3() : FDynamicMesh3(false, false, false, false) {}
@@ -969,9 +963,6 @@ public:
 	/** Returns bounding box of all mesh vertices (including unreferenced vertices) */
 	FAxisAlignedBox3d GetBounds(bool bParallel = false) const;
 
-	/** Returns GetBounds() and saves result, cache is invalidated and recomputed if topology has changed since last call */
-	FAxisAlignedBox3d GetCachedBounds();
-
 	/**
 	 * Compute a normal/tangent frame at vertex that is "stable" as long as
 	 * the mesh topology doesn't change, meaning that one axis of the frame
