GeometryCore: Refactor FCompactMap

#rnx

diff --git a/Engine/Plugins/Experimental/MeshModelingToolsetExp/Source/ModelingComponents/Private/Selection/GroupTopologyStorableSelection.cpp b/Engine/Plugins/Experimental/MeshModelingToolsetExp/Source/ModelingComponents/Private/Selection/GroupTopologyStorableSelection.cpp
--- a/Engine/Plugins/Experimental/MeshModelingToolsetExp/Source/ModelingComponents/Private/Selection/GroupTopologyStorableSelection.cpp
+++ b/Engine/Plugins/Experimental/MeshModelingToolsetExp/Source/ModelingComponents/Private/Selection/GroupTopologyStorableSelection.cpp
@@ -60,7 +60,7 @@ void UGroupTopologyStorableSelection::SetSelection(const FGroupTopology& Topolog
 
 	for (int32 CornerID : SelectionIn.SelectedCornerIDs)
 	{
-		CornerVids.Add(CompactMaps.GetVertex(TopologyIn.GetCornerVertexID(CornerID)));
+		CornerVids.Add(CompactMaps.GetVertexMapping(TopologyIn.GetCornerVertexID(CornerID)));
 	}
 	for (int32 EdgeID : SelectionIn.SelectedEdgeIDs)
 	{
@@ -134,10 +134,10 @@ FIndex2i UGroupTopologyStorableSelection::GetGroupEdgeRepresentativeVerts(const
 	if (GroupEdge.EndpointCorners.A != IndexConstants::InvalidID)
 	{
 		// Use remapped vids
-		int32 FirstVid = CompactMaps.GetVertex(Verts[0]);
-		int32 FirstNeighbor = CompactMaps.GetVertex(Verts[1]);
-		int32 LastVid = CompactMaps.GetVertex(Verts.Last());
-		int32 LastNeighbor = CompactMaps.GetVertex(Verts[Verts.Num() - 2]);
+		int32 FirstVid = CompactMaps.GetVertexMapping(Verts[0]);
+		int32 FirstNeighbor = CompactMaps.GetVertexMapping(Verts[1]);
+		int32 LastVid = CompactMaps.GetVertexMapping(Verts.Last());
+		int32 LastNeighbor = CompactMaps.GetVertexMapping(Verts[Verts.Num() - 2]);
 
 		return FirstVid < LastVid ?
 			FIndex2i(FMath::Min(FirstVid, FirstNeighbor), FMath::Max(FirstVid, FirstNeighbor))
@@ -149,7 +149,7 @@ FIndex2i UGroupTopologyStorableSelection::GetGroupEdgeRepresentativeVerts(const
 		RemappedVerts.SetNum(Verts.Num());
 		for (int32 i = 0; i < Verts.Num(); ++i)
 		{
-			RemappedVerts[i] = CompactMaps.GetVertex(Verts[i]);
+			RemappedVerts[i] = CompactMaps.GetVertexMapping(Verts[i]);
 		}
 		return GetLoopRepresentativeVerts(RemappedVerts);
 	}
diff --git a/Engine/Plugins/Experimental/MeshModelingToolsetExp/Source/ModelingComponents/Public/Selection/GroupTopologyStorableSelection.h b/Engine/Plugins/Experimental/MeshModelingToolsetExp/Source/ModelingComponents/Public/Selection/GroupTopologyStorableSelection.h
--- a/Engine/Plugins/Experimental/MeshModelingToolsetExp/Source/ModelingComponents/Public/Selection/GroupTopologyStorableSelection.h
+++ b/Engine/Plugins/Experimental/MeshModelingToolsetExp/Source/ModelingComponents/Public/Selection/GroupTopologyStorableSelection.h
@@ -15,7 +15,7 @@ class UPrimitiveComponent;
 class UInteractiveTool;
 PREDECLARE_USE_GEOMETRY_STRUCT(FGroupTopologySelection);
 PREDECLARE_USE_GEOMETRY_CLASS(FGroupTopology);
-PREDECLARE_USE_GEOMETRY_STRUCT(FCompactMaps);
+PREDECLARE_USE_GEOMETRY_CLASS(FCompactMaps);
 
 /**
  * Class used to represent a group topology selection independently of a FGroupTopology. Relies
diff --git a/Engine/Plugins/Experimental/MeshModelingToolsetExp/Source/ModelingComponents/Public/Selection/PolygonSelectionMechanic.h b/Engine/Plugins/Experimental/MeshModelingToolsetExp/Source/ModelingComponents/Public/Selection/PolygonSelectionMechanic.h
--- a/Engine/Plugins/Experimental/MeshModelingToolsetExp/Source/ModelingComponents/Public/Selection/PolygonSelectionMechanic.h
+++ b/Engine/Plugins/Experimental/MeshModelingToolsetExp/Source/ModelingComponents/Public/Selection/PolygonSelectionMechanic.h
@@ -22,7 +22,7 @@ class URectangleMarqueeMechanic;
 class USingleClickInputBehavior;
 
 using UE::Geometry::FDynamicMeshAABBTree3;
-PREDECLARE_USE_GEOMETRY_STRUCT(FCompactMaps);
+PREDECLARE_USE_GEOMETRY_CLASS(FCompactMaps);
 
 UCLASS()
 class MODELINGCOMPONENTS_API UPolygonSelectionMechanicProperties : public UInteractiveToolPropertySet
diff --git a/Engine/Source/Runtime/GeometryCore/Private/DynamicMesh/DynamicMesh3.cpp b/Engine/Source/Runtime/GeometryCore/Private/DynamicMesh/DynamicMesh3.cpp
--- a/Engine/Source/Runtime/GeometryCore/Private/DynamicMesh/DynamicMesh3.cpp
+++ b/Engine/Source/Runtime/GeometryCore/Private/DynamicMesh/DynamicMesh3.cpp
@@ -240,40 +240,28 @@ void FDynamicMesh3::CompactCopy(const FDynamicMesh3& copy, bool bNormals, bool b
 	// currently cannot re-use existing attribute buffers
 	Clear();
 
-	// use a local map if none passed in
-	FCompactMaps LocalMapsVar;
-	FCompactMaps* UseMaps = CompactInfo;
-	bool bNeedClearTriangleMap = false;
-	if (!UseMaps)
-	{
-		UseMaps = &LocalMapsVar;
-		UseMaps->bKeepTriangleMap = bAttributes && copy.HasAttributes();
-	}
-	else
+	// Use a triangle map if we have a CompactInfo or we need to copy attributes.
+	const bool bUseTriangleMap = CompactInfo != nullptr || (bAttributes && copy.HasAttributes());
+
+	// If we don't have a CompactInfo, we'll make it refer to a local one.
+	FCompactMaps LocalCompactInfo;
+	if (!CompactInfo)
 	{
-		// check if we need to temporarily keep the triangle map and clear it after
-		bool bNeedTriangleMap = bAttributes && copy.HasAttributes();
-		if (bNeedTriangleMap && !UseMaps->bKeepTriangleMap)
-		{
-			UseMaps->bKeepTriangleMap = true;
-			bNeedClearTriangleMap = true;
-		}
+		CompactInfo = &LocalCompactInfo;
 	}
-	UseMaps->Reset();
 
+	CompactInfo->ResetVertexMap(copy.MaxVertexID(), false);
 	FVertexInfo vinfo;
-	TArray<int>& mapV = UseMaps->MapV; mapV.SetNumUninitialized(copy.MaxVertexID());
-
-	for (int vid = 0; vid < copy.MaxVertexID(); vid++)
+	for (int vid = 0, NumVid = copy.MaxVertexID(); vid < NumVid; vid++)
 	{
 		if (copy.IsVertex(vid))
 		{
 			copy.GetVertex(vid, vinfo, bNormals, bColors, bUVs);
-			mapV[vid] = AppendVertex(vinfo);
+			CompactInfo->SetVertexMapping(vid, AppendVertex(vinfo));
 		}
 		else
 		{
-			mapV[vid] = -1;
+			CompactInfo->SetVertexMapping(vid, FCompactMaps::InvalidID);
 		}
 	}
 
@@ -284,25 +272,20 @@ void FDynamicMesh3::CompactCopy(const FDynamicMesh3& copy, bool bNormals, bool b
 	}
 
 	// need the triangle map to be computed if we have attributes and/or the FCompactMaps flag was set to request it
-	bool bNeedsTriangleMap = (bAttributes && copy.HasAttributes()) || UseMaps->bKeepTriangleMap;
-	if (bNeedsTriangleMap)
+	
+	if (bUseTriangleMap)
 	{
-		UseMaps->MapT.SetNumUninitialized(copy.MaxTriangleID());
-		for (int tid = 0; tid < copy.MaxTriangleID(); tid++)
-		{
-			UseMaps->MapT[tid] = -1;
-		}
+		CompactInfo->ResetTriangleMap(bUseTriangleMap ? copy.MaxTriangleID() : 0, true);
 	}
 	for (int tid : copy.TriangleIndicesItr())
 	{
-		FIndex3i t = copy.GetTriangle(tid);
-		t = FIndex3i(mapV[t.A], mapV[t.B], mapV[t.C]);
-		int g = (copy.HasTriangleGroups()) ? copy.GetTriangleGroup(tid) : InvalidID;
-		int NewTID = AppendTriangle(t, g);
+		const FIndex3i t = CompactInfo->GetVertexMapping(copy.GetTriangle(tid));
+		const int g = (copy.HasTriangleGroups()) ? copy.GetTriangleGroup(tid) : InvalidID;
+		const int NewTID = AppendTriangle(t, g);
 		GroupIDCounter = FMath::Max(GroupIDCounter, g + 1);
-		if (bNeedsTriangleMap)
+		if (bUseTriangleMap)
 		{
-			UseMaps->MapT[tid] = NewTID;
+			CompactInfo->SetTriangleMapping(tid, NewTID);
 		}
 	}
 
@@ -310,12 +293,7 @@ void FDynamicMesh3::CompactCopy(const FDynamicMesh3& copy, bool bNormals, bool b
 	{
 		EnableAttributes();
 		AttributeSet->EnableMatchingAttributes(*copy.Attributes());
-		AttributeSet->CompactCopy(*UseMaps, *copy.Attributes());
-	}
-
-	if (bNeedClearTriangleMap)
-	{
-		CompactInfo->ClearTriangleMap();
+		AttributeSet->CompactCopy(*CompactInfo, *copy.Attributes());
 	}
 
 	Timestamp = FMath::Max(Timestamp + 1, copy.Timestamp);
diff --git a/Engine/Source/Runtime/GeometryCore/Private/DynamicMesh/DynamicMesh3_Edits.cpp b/Engine/Source/Runtime/GeometryCore/Private/DynamicMesh/DynamicMesh3_Edits.cpp
--- a/Engine/Source/Runtime/GeometryCore/Private/DynamicMesh/DynamicMesh3_Edits.cpp
+++ b/Engine/Source/Runtime/GeometryCore/Private/DynamicMesh/DynamicMesh3_Edits.cpp
@@ -292,38 +292,26 @@ EMeshResult FDynamicMesh3::InsertTriangle(int tid, const FIndex3i& tv, int gid,
 
 void FDynamicMesh3::CompactInPlace(FCompactMaps* CompactInfo)
 {
-
 	// Initialize CompactInfo
 
-	// if have attributes, need CompactInfo -- we'll make it refer to a local one if none was passed in
-	FCompactMaps LocalCompact;
-	// if we have attributes, but didn't request triangle map, we're going to put the triangle map anyway for attributes use and then clear it after
-	bool bNeedClearTriangleMap = false;
-
-	if (HasAttributes())
+	// If we need a CompactInfo for compacting attributes but we don't have one, we'll make it refer to a local one.
+	FCompactMaps LocalCompactInfo;
+	if (HasAttributes() && !CompactInfo)
 	{
-		if (!CompactInfo)
-		{
-			CompactInfo = &LocalCompact;
-		}
-		if (!CompactInfo->bKeepTriangleMap)
-		{
-			bNeedClearTriangleMap = true;
-			CompactInfo->bKeepTriangleMap = true;
-		}
+		CompactInfo = &LocalCompactInfo;
 	}
 
 	if (CompactInfo)
 	{
 		// starts as identity (except at gaps); sparsely remapped below
-		CompactInfo->Reset(MaxVertexID(), MaxTriangleID());
-		for (int VID = 0; VID < MaxVertexID(); VID++)
+		CompactInfo->Reset(MaxVertexID(), MaxTriangleID(), false);
+		for (int VID = 0, NumVID = MaxVertexID(); VID < NumVID; VID++)
 		{
-			CompactInfo->MapV[VID] = IsVertex(VID) ? VID : -1;
+			CompactInfo->SetVertexMapping(VID, IsVertex(VID) ? VID : -1);
 		}
-		for (int TID = 0; TID < CompactInfo->MapT.Num(); TID++)
+		for (int TID = 0, NumTID = MaxTriangleID(); TID < NumTID; TID++)
 		{
-			CompactInfo->MapT[TID] = IsTriangle(TID) ? TID : -1;
+			CompactInfo->SetTriangleMapping(TID, IsTriangle(TID) ? TID : -1);
 		}
 	}
 
@@ -385,7 +373,7 @@ void FDynamicMesh3::CompactInPlace(FCompactMaps* CompactInfo)
 
 		if (CompactInfo != nullptr)
 		{
-			CompactInfo->SetVertex(iLastV, iCurV);
+			CompactInfo->SetVertexMapping(iLastV, iCurV);
 		}
 
 		// move cur forward one, last back one, and  then search for next valid
@@ -455,9 +443,9 @@ void FDynamicMesh3::CompactInPlace(FCompactMaps* CompactInfo)
 		tref[iCurT] = tref[iLastT];
 		tref[iLastT] = FRefCountVector::INVALID_REF_COUNT;
 
-		if (CompactInfo != nullptr && CompactInfo->bKeepTriangleMap)
+		if (CompactInfo != nullptr)
 		{
-			CompactInfo->SetTriangle(iLastT, iCurT);
+			CompactInfo->SetTriangleMapping(iLastT, iCurT);
 		}
 
 		// move cur forward one, last back one, and  then search for next valid
@@ -537,11 +525,6 @@ void FDynamicMesh3::CompactInPlace(FCompactMaps* CompactInfo)
 		checkSlow(CompactInfo);		// can this ever fail?
 		AttributeSet->CompactInPlace(*CompactInfo);
 	}
-
-	if (bNeedClearTriangleMap)
-	{
-		CompactInfo->ClearTriangleMap();
-	}
 }
 
 
diff --git a/Engine/Source/Runtime/GeometryCore/Public/DynamicMesh/DynamicMesh3.h b/Engine/Source/Runtime/GeometryCore/Public/DynamicMesh/DynamicMesh3.h
--- a/Engine/Source/Runtime/GeometryCore/Public/DynamicMesh/DynamicMesh3.h
+++ b/Engine/Source/Runtime/GeometryCore/Public/DynamicMesh/DynamicMesh3.h
@@ -219,6 +219,13 @@ public:
 	 * Copy input mesh while compacting, i.e. removing unused vertices/triangles/edges. 
 	 * Note that it is currently unsafe to hold on to the attribute set pointer (returned from Attributes())
 	 * across the CompactCopy call, as the latter may rebuild an entirely new attribute set.
+	 * 
+	 * @param CopyMesh Mesh to copy
+	 * @param bNormals if true, will copy normals
+	 * @param bColors if true, will copy colors
+	 * @param bUVs if true, will copy UVs
+	 * @param bAttributes if true, will copy attributes
+	 * @param CompactInfo if not nullptr, will be filled with mapping indicating how vertex and triangle IDs were changed during compaction
 	 */
 	void CompactCopy(const FDynamicMesh3& CopyMesh, bool bNormals = true, bool bColors = true, bool bUVs = true,
 	                 bool bAttributes = true, FCompactMaps* CompactInfo = nullptr);
@@ -1091,7 +1098,7 @@ public:
 	 * Compact mesh in-place, by moving vertices around and rewriting indices.
 	 * Should be faster if the amount of compacting is not too significant, and is useful in some places.
 	 *
-	 * @param CompactInfo if not nullptr, will be filled with mapping indicating how vertex and (optionally) triangle IDs were changed during compaction
+	 * @param CompactInfo if not nullptr, will be filled with mapping indicating how vertex and triangle IDs were changed during compaction
 	 * @todo VertexEdgeLists is not compacted. does not affect indices, but does keep memory.
 	 */
 	void CompactInPlace(FCompactMaps* CompactInfo = nullptr);
diff --git a/Engine/Source/Runtime/GeometryCore/Public/DynamicMesh/DynamicMeshOverlay.h b/Engine/Source/Runtime/GeometryCore/Public/DynamicMesh/DynamicMeshOverlay.h
--- a/Engine/Source/Runtime/GeometryCore/Public/DynamicMesh/DynamicMeshOverlay.h
+++ b/Engine/Source/Runtime/GeometryCore/Public/DynamicMesh/DynamicMeshOverlay.h
@@ -120,14 +120,14 @@ public:
 		}
 
 		// copy triangles across
-		check(CompactMaps.bKeepTriangleMap && CompactMaps.MapT.Num() == Copy.GetParentMesh()->MaxTriangleID()); // must have valid triangle map
+		check(CompactMaps.NumTriangleMappings() == Copy.GetParentMesh()->MaxTriangleID()); // must have valid triangle map
 		for (int FromTID : Copy.GetParentMesh()->TriangleIndicesItr())
 		{
 			if (!Copy.IsSetTriangle(FromTID))
 			{
 				continue;
 			}
-			int ToTID = CompactMaps.MapT[FromTID];
+			const int ToTID = CompactMaps.GetTriangleMapping(FromTID);
 			FIndex3i FromTriElements = Copy.GetTriangle(FromTID);
 			SetTriangle(ToTID, FIndex3i(MapE[FromTriElements.A], MapE[FromTriElements.B], MapE[FromTriElements.C]));
 		}
@@ -156,7 +156,7 @@ public:
 			// remap all parents
 			if (ParentVertices[ID] >= 0)
 			{
-				ParentVertices[ID] = CompactMaps.GetVertex(ParentVertices[ID]);
+				ParentVertices[ID] = CompactMaps.GetVertexMapping(ParentVertices[ID]);
 			}
 		}
 
@@ -194,7 +194,7 @@ public:
 			{
 				continue; // triangle was not set; skip it
 			}
-			int NewStart = CompactMaps.GetTriangle(TID) * 3;
+			int NewStart = CompactMaps.GetTriangleMapping(TID) * 3;
 			for (int SubIdx = 0; SubIdx < 3; SubIdx++)
 			{
 				ElementTriangles[NewStart + SubIdx] = MapE[ElementTriangles[OldStart + SubIdx]];
diff --git a/Engine/Source/Runtime/GeometryCore/Public/DynamicMesh/DynamicMeshTriangleAttribute.h b/Engine/Source/Runtime/GeometryCore/Public/DynamicMesh/DynamicMeshTriangleAttribute.h
--- a/Engine/Source/Runtime/GeometryCore/Public/DynamicMesh/DynamicMeshTriangleAttribute.h
+++ b/Engine/Source/Runtime/GeometryCore/Public/DynamicMesh/DynamicMeshTriangleAttribute.h
@@ -124,10 +124,10 @@ public:
 
 	void CompactInPlace(const FCompactMaps& CompactMaps)
 	{
-		for (int TID = 0; TID < CompactMaps.MapT.Num(); TID++)
+		for (int TID = 0, NumTID = CompactMaps.NumTriangleMappings(); TID < NumTID; TID++)
 		{
-			int ToTID = CompactMaps.MapT[TID];
-			if (ToTID < 0)
+			const int ToTID = CompactMaps.GetTriangleMapping(TID);
+			if (ToTID == FCompactMaps::InvalidID)
 			{
 				continue;
 			}
@@ -142,12 +142,12 @@ public:
 	void CompactCopy(const FCompactMaps& CompactMaps, const TDynamicMeshTriangleAttribute<AttribValueType, AttribDimension>& ToCopy)
 	{
 		CopyParentClassData(ToCopy);
-		check(CompactMaps.MapT.Num() <= int(ToCopy.AttribValues.Num() / AttribDimension));
+		check(CompactMaps.NumTriangleMappings() <= int(ToCopy.AttribValues.Num() / AttribDimension));
 		AttribValueType Data[AttribDimension];
-		for (int TID = 0; TID < CompactMaps.MapT.Num(); TID++)
+		for (int TID = 0, NumTID = CompactMaps.NumTriangleMappings(); TID < NumTID; TID++)
 		{
-			int ToTID = CompactMaps.MapT[TID];
-			if (ToTID == -1)
+			const int ToTID = CompactMaps.GetTriangleMapping(TID);
+			if (ToTID == FCompactMaps::InvalidID)
 			{
 				continue;
 			}
diff --git a/Engine/Source/Runtime/GeometryCore/Public/DynamicMesh/DynamicVertexAttribute.h b/Engine/Source/Runtime/GeometryCore/Public/DynamicMesh/DynamicVertexAttribute.h
--- a/Engine/Source/Runtime/GeometryCore/Public/DynamicMesh/DynamicVertexAttribute.h
+++ b/Engine/Source/Runtime/GeometryCore/Public/DynamicMesh/DynamicVertexAttribute.h
@@ -111,10 +111,10 @@ public:
 
 	void CompactInPlace(const FCompactMaps& CompactMaps)
 	{
-		for (int VID = 0; VID < CompactMaps.MapV.Num(); VID++)
+		for (int VID = 0, NumVID = CompactMaps.NumVertexMappings(); VID < NumVID; VID++)
 		{
-			int ToVID = CompactMaps.MapV[VID];
-			if (ToVID < 0)
+			const int ToVID = CompactMaps.GetVertexMapping(VID);
+			if (ToVID == FCompactMaps::InvalidID)
 			{
 				continue;
 			}
@@ -129,13 +129,13 @@ public:
 	void CompactCopy(const FCompactMaps& CompactMaps, const TDynamicVertexAttribute<AttribValueType, AttribDimension, ParentType>& ToCopy)
 	{
 		TDynamicAttributeBase<ParentType>::CopyParentClassData(ToCopy);
-		check(CompactMaps.MapV.Num() >= 0 && static_cast<size_t>(CompactMaps.MapV.Num()) <= AttribValues.Num() / AttribDimension);
+		check(CompactMaps.NumVertexMappings() >= 0 && static_cast<size_t>(CompactMaps.NumVertexMappings()) <= AttribValues.Num() / AttribDimension);
 
 		AttribValueType Data[AttribDimension];
-		for (int VID = 0; VID < CompactMaps.MapV.Num(); VID++)
+		for (int VID = 0, NumVID = CompactMaps.NumVertexMappings(); VID < NumVID; VID++)
 		{
-			int ToVID = CompactMaps.MapV[VID];
-			if (ToVID == -1)
+			int ToVID = CompactMaps.GetVertexMapping(VID);
+			if (ToVID == FCompactMaps::InvalidID)
 			{
 				continue;
 			}
diff --git a/Engine/Source/Runtime/GeometryCore/Public/DynamicMesh/DynamicVertexSkinWeightsAttribute.h b/Engine/Source/Runtime/GeometryCore/Public/DynamicMesh/DynamicVertexSkinWeightsAttribute.h
--- a/Engine/Source/Runtime/GeometryCore/Public/DynamicMesh/DynamicVertexSkinWeightsAttribute.h
+++ b/Engine/Source/Runtime/GeometryCore/Public/DynamicMesh/DynamicVertexSkinWeightsAttribute.h
@@ -113,10 +113,10 @@ public:
 
 	void CompactInPlace(const FCompactMaps& CompactMaps) override
 	{
-		for (int32 VID = 0; VID < CompactMaps.MapV.Num(); VID++)
+		for (int32 VID = 0, NumVID = CompactMaps.NumVertexMappings(); VID < NumVID; VID++)
 		{
-			const int32 ToVID = CompactMaps.MapV[VID];
-			if (ToVID < 0)
+			const int32 ToVID = CompactMaps.GetVertexMapping(VID);
+			if (ToVID == FCompactMaps::InvalidID)
 			{
 				continue;
 			}
@@ -131,12 +131,12 @@ public:
 	void CompactCopy(const FCompactMaps& CompactMaps, const TDynamicVertexSkinWeightsAttribute<ParentType>& ToCopy)
 	{
 		TDynamicAttributeBase<ParentType>::CopyParentClassData(ToCopy);
-		check(CompactMaps.MapV.Num() <= VertexBoneWeights.Num());
+		check(CompactMaps.NumVertexMappings() <= VertexBoneWeights.Num());
 		FBoneWeights Data;
-		for (int32 VID = 0; VID < CompactMaps.MapV.Num(); VID++)
+		for (int32 VID = 0, NumVID = CompactMaps.NumVertexMappings(); VID < NumVID; VID++)
 		{
-			const int32 ToVID = CompactMaps.MapV[VID];
-			if (ToVID == -1)
+			const int32 ToVID = CompactMaps.GetVertexMapping(VID);
+			if (ToVID == FCompactMaps::InvalidID)
 			{
 				continue;
 			}
diff --git a/Engine/Source/Runtime/GeometryCore/Public/Util/CompactMaps.h b/Engine/Source/Runtime/GeometryCore/Public/Util/CompactMaps.h
--- a/Engine/Source/Runtime/GeometryCore/Public/Util/CompactMaps.h
+++ b/Engine/Source/Runtime/GeometryCore/Public/Util/CompactMaps.h
@@ -12,146 +12,186 @@ namespace UE
 {
 namespace Geometry
 {
-
 /**
- * Stores index remapping arrays for different types of elements
+ * Stores index remapping for vertices and triangles.
  * Should only be used for compacting, and should maintain invariant that Map*[Idx] <= Idx for all maps
  */
-struct FCompactMaps
+class FCompactMaps
 {
-	constexpr static int InvalidID = IndexConstants::InvalidID;
+	TArray<int32> VertMap;
+	TArray<int32> TriMap;
 
-	TArray<int> MapV;
-	TArray<int> MapT;
-	bool bKeepTriangleMap = true; // if set to false, MapT should be empty
+public:
+	constexpr static int32 InvalidID = IndexConstants::InvalidID;
 
 	/**
 	 * Set up maps as identity maps.
-	 * Note triangle map will not be set if bKeepTriangleMap is false.
 	 *
-	 * @param MaxVID Vertex map will be created from 0 to MaxVID
-	 * @param MaxTID If bKeepTriangleMap is true, triangle map will be created from 0 to MaxTID
+	 * @param NumVertMappings Vertex map will be created from 0 to NumVertMappings - 1
+	 * @param NumTriMappings Triangle map will be created from 0 to NumTriMappings - 1
 	 */
-	void SetIdentity(int MaxVID, int MaxTID)
+	void SetIdentity(int32 NumVertMappings, int32 NumTriMappings)
 	{
-		MapV.SetNumUninitialized(MaxVID);
-		for (int ID = 0; ID < MaxVID; ID++)
-		{
-			MapV[ID] = ID;
-		}
-		if (bKeepTriangleMap)
+		SetIdentityVertexMap(NumVertMappings);
+		SetIdentityTriangleMap(NumTriMappings);
+	}
+
+	/**
+	* Set up vertex map as identity map.
+	*
+	* @param NumVertMappings Vertex map will be created from 0 to NumVertMappings - 1
+	*/
+	void SetIdentityVertexMap(int32 NumVertMappings)
+	{
+		VertMap.SetNumUninitialized(NumVertMappings);
+		for (int32 i = 0; i < NumVertMappings; ++i)
 		{
-			MapT.SetNumUninitialized(MaxTID);
-			for (int ID = 0; ID < MaxTID; ID++)
-			{
-				MapT[ID] = ID;
-			}
+			VertMap[i] = i;
 		}
-		else
+	}
+
+	/**
+	* Set up triangle map as identity map.
+	*
+	* @param NumTriMappings Vertex map will be created from 0 to NumTriMappings - 1
+	*/
+	void SetIdentityTriangleMap(int32 NumTriMappings)
+	{
+		TriMap.SetNumUninitialized(NumTriMappings);
+		for (int32 i = 0; i < NumTriMappings; ++i)
 		{
-			MapT.Empty();
+			TriMap[i] = i;
 		}
 	}
 
 	/**
-	 * Reset all maps, and initialize with InvalidID
-	 * @param MaxVID Size of post-reset vertex map (filled w/ InvalidID)
-	 * @param MaxTID Size of post-reset triangle map (filled w/ InvalidID)
+	 * Resize vertex and triangle maps, and initialize with InvalidID.
+	 * @param NumVertMappings Size of post-reset vertex map
+	 * @param NumTriMappings Size of post-reset triangle map
+	* @param bInitializeWithInvalidID If true, initializes maps with InvalidID 
 	 */
-	void Reset(int MaxVID, int MaxTID = 0)
+	void Reset(int32 NumVertMappings, int32 NumTriMappings, bool bInitializeWithInvalidID)
+	{
+		ResetVertexMap(NumVertMappings, bInitializeWithInvalidID);
+		ResetTriangleMap(NumTriMappings, bInitializeWithInvalidID);
+	}
+
+	/**
+	* Resize vertex map, and optionally initialize with InvalidID.
+	* @param NumVertMappings Size of post-reset vertex map
+	* @param bInitializeWithInvalidID If true, initializes map with InvalidID 
+	*/
+	void ResetVertexMap(int32 NumVertMappings, bool bInitializeWithInvalidID)
 	{
-		MapV.Reset();
-		MapT.Reset();
-		MapV.SetNumUninitialized(MaxVID);
-		MapT.SetNumUninitialized(MaxTID);
-		for (int VID = 0; VID < MaxVID; VID++)
+		VertMap.SetNumUninitialized(NumVertMappings);
+		if (bInitializeWithInvalidID)
 		{
-			MapV[VID] = InvalidID;
+			for (int32 i = 0; i < NumVertMappings; i++)
+			{
+				VertMap[i] = InvalidID;
+			}
 		}
-		for (int TID = 0; TID < MaxTID; TID++)
+	}
+
+	/**
+	* Resize triangle map, and optionally initialize with InvalidID.
+	* @param NumTriMappings Size of post-reset triangle map
+	* @param bInitializeWithInvalidID If true, initializes map with InvalidID 
+	*/
+	void ResetTriangleMap(int32 NumTriMappings, bool bInitializeWithInvalidID)
+	{
+		TriMap.SetNumUninitialized(NumTriMappings);
+		if (bInitializeWithInvalidID)
 		{
-			MapT[TID] = InvalidID;
+			for (int32 i = 0; i < NumTriMappings; ++i)
+			{
+				TriMap[i] = InvalidID;
+			}
 		}
 	}
 
 	/** Reset all maps, leaving them empty */
 	void Reset()
 	{
-		MapV.Reset();
-		MapT.Reset();
+		VertMap.Reset();
+		TriMap.Reset();
+	}
+
+	/** Returns true if there are vertex mappings. */
+	bool VertexMapIsSet() const
+	{
+		return !VertMap.IsEmpty();
+	}
+
+	/** Returns true if there are triangle mappings. */
+	bool TriangleMapIsSet() const
+	{
+		return !TriMap.IsEmpty();
+	}
+
+	/** Get number of vertex mappings */
+	int32 NumVertexMappings() const
+	{
+		return VertMap.Num();
+	}
+
+	/** Get number of triangle mappings */
+	int32 NumTriangleMappings() const
+	{
+		return TriMap.Num();
 	}
 
 	/** Set mapping for a vertex */
-	inline void SetVertex(int FromID, int ToID)
+	void SetVertexMapping(int32 FromID, int32 ToID)
 	{
 		checkSlow(FromID >= ToID);
-		MapV[FromID] = ToID;
+		VertMap[FromID] = ToID;
 	}
 
 	/** Set mapping for a triangle */
-	inline void SetTriangle(int FromID, int ToID)
+	void SetTriangleMapping(int32 FromID, int32 ToID)
 	{
-		checkSlow(bKeepTriangleMap);
 		checkSlow(FromID >= ToID);
-		MapT[FromID] = ToID;
+		TriMap[FromID] = ToID;
 	}
 
 	/** Get mapping for a vertex */
-	inline int GetVertex(int FromID) const
+	int32 GetVertexMapping(int32 FromID) const
 	{
-		int ToID = MapV[FromID];
-		checkSlow(ToID != InvalidID);
-		return ToID;
+		return VertMap[FromID];
 	}
 
-	/** Get mapping for a triangle */
-	inline int GetTriangle(int FromID) const
+	/** Get mapping for three vertices, e.g. a triangle */
+	FIndex3i GetVertexMapping(FIndex3i FromIDs) const
 	{
-		checkSlow(bKeepTriangleMap);
-		int ToID = MapT[FromID];
-		checkSlow(ToID != InvalidID);
-		return ToID;
+		return {VertMap[FromIDs[0]], VertMap[FromIDs[1]], VertMap[FromIDs[2]]};
 	}
 
-	/** 
-	 * Clear all triangles
-	 *
-	 * @param bPermanent If true, set flag to not create triangle map
-	 */
-	void ClearTriangleMap(bool bPermanent = true)
+	/** Get mapping for a triangle */
+	int32 GetTriangleMapping(int32 FromID) const
 	{
-		if (bPermanent)
-		{
-			bKeepTriangleMap = false;
-		}
-		MapT.Empty();
+		return TriMap[FromID];
 	}
 
 	/** Check data for validity; for testing */
 	bool Validate() const
 	{
-		for (int Idx = 0; Idx < MapV.Num(); Idx++)
+		for (int32 Idx = 0; Idx < VertMap.Num(); Idx++)
 		{
-			if (MapV[Idx] > Idx)
+			if (VertMap[Idx] > Idx)
 			{
 				return false;
 			}
 		}
-		for (int Idx = 0; Idx < MapT.Num(); Idx++)
+		for (int32 Idx = 0; Idx < TriMap.Num(); Idx++)
 		{
-			if (MapT[Idx] > Idx)
+			if (TriMap[Idx] > Idx)
 			{
 				return false;
 			}
 		}
-		if (!bKeepTriangleMap && MapT.Num() > 0)
-		{
-			return false;
-		}
 		return true;
 	}
 };
-
-
 } // end namespace UE::Geometry
-} // end namespace UE
\ No newline at end of file
+} // end namespace UE
