GeometryCore: Add unit test for loading all serialization versions of FDynamicMesh3

This CL adds a unit test and some related basic infrastructure for testing the backwards compatibility of our FDynamicMesh3 serialization. This is in anticipation of the incoming serialization changes.

#rnx

diff --git a/Engine/Source/Runtime/GeometryCore/Private/DynamicMesh/DynamicMesh3_Serialization.cpp b/Engine/Source/Runtime/GeometryCore/Private/DynamicMesh/DynamicMesh3_Serialization.cpp
--- a/Engine/Source/Runtime/GeometryCore/Private/DynamicMesh/DynamicMesh3_Serialization.cpp
+++ b/Engine/Source/Runtime/GeometryCore/Private/DynamicMesh/DynamicMesh3_Serialization.cpp
@@ -9,17 +9,26 @@ using namespace UE::Geometry;
 
 enum class EDynamicMeshSerializationVersion
 {
-	InitialVersion = 1
+	InitialVersion = 1,
+
+	// ----- new versions to be added above this line -------------------------------------------------
+	VersionPlusOne,
+	LatestVersion = VersionPlusOne - 1
 };
 
+int FDynamicMesh3::SerializeInternal_LatestVersion()
+{
+	return static_cast<int>(EDynamicMeshSerializationVersion::LatestVersion);
+}
 
 void FDynamicMesh3::Serialize(FArchive& Ar)
 {
-	checkSlow(CheckValidity(FValidityOptions(), EValidityCheckFailMode::Ensure));
+	// Check validity before saving data.
+	checkSlow(Ar.IsLoading() || CheckValidity(FValidityOptions(), EValidityCheckFailMode::Ensure));
 
 	Ar.UsingCustomVersion(FUE5MainStreamObjectVersion::GUID);
 
-	int32 SerializationVersion = (int32)EDynamicMeshSerializationVersion::InitialVersion;
+	int32 SerializationVersion = static_cast<int>(EDynamicMeshSerializationVersion::LatestVersion);
 	Ar << SerializationVersion;
 
 	Ar << Vertices;
@@ -46,4 +55,7 @@ void FDynamicMesh3::Serialize(FArchive& Ar)
 		}
 		Ar << *AttributeSet;
 	}
+
+	// Check validity after loading data.
+	checkSlow(!Ar.IsLoading() || CheckValidity(FValidityOptions(), EValidityCheckFailMode::Ensure));
 }
\ No newline at end of file
diff --git a/Engine/Source/Runtime/GeometryCore/Public/DynamicMesh/DynamicMesh3.h b/Engine/Source/Runtime/GeometryCore/Public/DynamicMesh/DynamicMesh3.h
--- a/Engine/Source/Runtime/GeometryCore/Public/DynamicMesh/DynamicMesh3.h
+++ b/Engine/Source/Runtime/GeometryCore/Public/DynamicMesh/DynamicMesh3.h
@@ -1463,6 +1463,9 @@ protected:
 			TopologyTimestamp++;
 		}
 	}
+
+	/** Returns the internal version number used during serialization. This is used for testing purposes. */
+	static int SerializeInternal_LatestVersion();
 };
 
 
