GeometryProcessing: clean up mesh timestamps.
- remove FDynamicMesh3::Timestamp (unused), rename Shape/Topology Timestamps to Shape/TopologyChangeStamp, change to atomic<uint32>
- add FDynamicMesh3::bEnableShapeChangeStamp, default to false, to disable ShapeChange tracking. Add ::SetShapeChangeStampEnabled() and ::HasShapeChangeStampEnabled() to configure.
- replace FDynamicMesh3::UpdateTimestamps() with UpdateChangeStamps()
- add bTrackChange param to FDynamicMesh3::SetVertex(), optionally updates ShapeChangeStamp (if enabled). Default true. Remove SetVertex_NoTimeStampUpdate(), update call sites.
- add FDynamicMesh3::GetChangeStamp(), returns combination of Shape and Topology stamps as uint64
- rename TTriangleMeshAdapter::GetTimestamp() to GetChangeStamp(), update usages
- remove TPointSetAdapter::Timestamp()   (was not used in code)
- update TMeshAABBTree3 to use GetChangeStamp(), update internal checks to call IsValid() instead
- update TFastWindingTree w/ similar changes
- update calls in UVEditor, may require further updates
#rnx

diff --git a/Engine/Plugins/Experimental/Enterprise/DataprepGeometryOperations/Source/Private/DataprepGeometrySelectionTransforms.cpp b/Engine/Plugins/Experimental/Enterprise/DataprepGeometryOperations/Source/Private/DataprepGeometrySelectionTransforms.cpp
--- a/Engine/Plugins/Experimental/Enterprise/DataprepGeometryOperations/Source/Private/DataprepGeometrySelectionTransforms.cpp
+++ b/Engine/Plugins/Experimental/Enterprise/DataprepGeometryOperations/Source/Private/DataprepGeometrySelectionTransforms.cpp
@@ -98,7 +98,7 @@ namespace DataprepGeometryOperationsUtils
 
 		MergedMeshTree.Build();
 
-		check( MergedMeshTree.IsValid() );
+		check( MergedMeshTree.IsValid(false) );
 
 		using FAxisAlignedBox3d = UE::Geometry::FAxisAlignedBox3d;
 		const FAxisAlignedBox3d MergedMeshBox = MergedMeshTree.GetBoundingBox();
diff --git a/Engine/Plugins/Experimental/MeshModelingToolsetExp/Source/MeshModelingToolsExp/Private/DynamicMeshSculptTool.cpp b/Engine/Plugins/Experimental/MeshModelingToolsetExp/Source/MeshModelingToolsExp/Private/DynamicMeshSculptTool.cpp
--- a/Engine/Plugins/Experimental/MeshModelingToolsetExp/Source/MeshModelingToolsExp/Private/DynamicMeshSculptTool.cpp
+++ b/Engine/Plugins/Experimental/MeshModelingToolsetExp/Source/MeshModelingToolsExp/Private/DynamicMeshSculptTool.cpp
@@ -653,9 +653,9 @@ void UDynamicMeshSculptTool::SyncMeshWithPositionBuffer(FDynamicMesh3* Mesh)
 	{
 		int VertIdx = VertexROI[k];
 		const FVector3d& NewPos = ROIPositionBuffer[k];
-		Mesh->SetVertex_NoTimeStampUpdate(VertIdx, NewPos);
+		Mesh->SetVertex(VertIdx, NewPos, false);
 	});
-	Mesh->IncrementTimeStamps(1, true, false);
+	Mesh->UpdateChangeStamps(true, false);
 }
 
 bool UDynamicMeshSculptTool::ApplySmoothBrush(const FRay& WorldRay)
diff --git a/Engine/Plugins/Experimental/MeshModelingToolsetExp/Source/MeshModelingToolsExp/Private/MeshVertexSculptTool.cpp b/Engine/Plugins/Experimental/MeshModelingToolsetExp/Source/MeshModelingToolsExp/Private/MeshVertexSculptTool.cpp
--- a/Engine/Plugins/Experimental/MeshModelingToolsetExp/Source/MeshModelingToolsExp/Private/MeshVertexSculptTool.cpp
+++ b/Engine/Plugins/Experimental/MeshModelingToolsetExp/Source/MeshModelingToolsExp/Private/MeshVertexSculptTool.cpp
@@ -616,9 +616,9 @@ TFuture<void> UMeshVertexSculptTool::ApplyStamp()
 		{
 			int VertIdx = VertexROI[k];
 			const FVector3d& NewPos = ROIPositionBuffer[k];
-			Mesh->SetVertex_NoTimeStampUpdate(VertIdx, NewPos);
+			Mesh->SetVertex(VertIdx, NewPos, false);
 		});
-		Mesh->IncrementTimeStamps(1, true, false);
+		Mesh->UpdateChangeStamps(true, false);
 	}
 
 	LastStamp = CurrentStamp;
diff --git a/Engine/Plugins/Experimental/MeshModelingToolsetExp/Source/ModelingComponents/Private/PreviewMesh.cpp b/Engine/Plugins/Experimental/MeshModelingToolsetExp/Source/ModelingComponents/Private/PreviewMesh.cpp
--- a/Engine/Plugins/Experimental/MeshModelingToolsetExp/Source/ModelingComponents/Private/PreviewMesh.cpp
+++ b/Engine/Plugins/Experimental/MeshModelingToolsetExp/Source/ModelingComponents/Private/PreviewMesh.cpp
@@ -296,7 +296,7 @@ FDynamicMeshAABBTree3* UPreviewMesh::GetSpatial()
 {
 	if (DynamicMeshComponent != nullptr && bBuildSpatialDataStructure)
 	{
-		if (MeshAABBTree.IsValid())
+		if (MeshAABBTree.IsValid(false))
 		{
 			return &MeshAABBTree;
 		}
diff --git a/Engine/Plugins/Experimental/MeshModelingToolsetExp/Source/ModelingComponents/Private/Scene/MeshSceneAdapter.cpp b/Engine/Plugins/Experimental/MeshModelingToolsetExp/Source/ModelingComponents/Private/Scene/MeshSceneAdapter.cpp
--- a/Engine/Plugins/Experimental/MeshModelingToolsetExp/Source/ModelingComponents/Private/Scene/MeshSceneAdapter.cpp
+++ b/Engine/Plugins/Experimental/MeshModelingToolsetExp/Source/ModelingComponents/Private/Scene/MeshSceneAdapter.cpp
@@ -536,7 +536,7 @@ struct TMeshDescriptionMeshAdapterd : public UE::Geometry::TTriangleMeshAdapter<
 		MaxVertexID = [&]() { return ParentAdapter.MaxVertexID(); };
 		TriangleCount = [&]() { return ParentAdapter.TriangleCount(); };
 		VertexCount = [&]() { return ParentAdapter.VertexCount(); };
-		GetShapeTimestamp = [&]() { return ParentAdapter.GetShapeTimestamp(); };
+		GetChangeStamp = [&]() { return ParentAdapter.GetChangeStamp(); };
 		GetTriangle = [&](int32 TriangleID) { return ParentAdapter.GetTriangle(TriangleID); };
 		GetVertex = [&](int32 VertexID) { return ParentAdapter.GetVertex(VertexID); };
 	}
diff --git a/Engine/Plugins/Experimental/MeshModelingToolsetExp/Source/ModelingComponents/Private/Selection/MeshSelectionMechanic.cpp b/Engine/Plugins/Experimental/MeshModelingToolsetExp/Source/ModelingComponents/Private/Selection/MeshSelectionMechanic.cpp
--- a/Engine/Plugins/Experimental/MeshModelingToolsetExp/Source/ModelingComponents/Private/Selection/MeshSelectionMechanic.cpp
+++ b/Engine/Plugins/Experimental/MeshModelingToolsetExp/Source/ModelingComponents/Private/Selection/MeshSelectionMechanic.cpp
@@ -290,7 +290,9 @@ void UMeshSelectionMechanic::UpdateCentroid()
 		CurrentSelectionCentroid /= CurrentSelection.SelectedIDs.Num();
 	}
 
-	CentroidTimestamp = CurrentSelection.Mesh->GetShapeTimestamp();
+	// disable centroid caching if mesh does not have shape changestamp enabled
+	CentroidTimestamp = CurrentSelection.Mesh->HasShapeChangeStampEnabled() ?
+		CurrentSelection.Mesh->GetShapeChangeStamp() : TNumericLimits<uint32>::Max();
 }
 
 FVector3d UMeshSelectionMechanic::GetCurrentSelectionCentroid()
@@ -299,7 +301,7 @@ FVector3d UMeshSelectionMechanic::GetCurrentSelectionCentroid()
 	{
 		return FVector3d(0);
 	}
-	else if (CurrentSelection.Mesh->GetShapeTimestamp() != CentroidTimestamp)
+	else if (CurrentSelection.Mesh->GetShapeChangeStamp() != CentroidTimestamp)
 	{
 		UpdateCentroid();
 	}
@@ -371,7 +373,7 @@ void UMeshSelectionMechanic::OnClicked(const FInputDeviceRay& ClickPos)
 		{
 			// It should be okay to reassign this. If we're in multiselect, and it misses a hit, this just won't happen and selection won't be altered.
 			CurrentSelection.Mesh = MeshSpatials[i]->GetMesh();
-			CurrentSelection.TopologyTimestamp = CurrentSelection.Mesh->GetTopologyTimestamp();
+			CurrentSelection.TopologyTimestamp = CurrentSelection.Mesh->GetTopologyChangeStamp();
 			CurrentSelectionIndex = i;
 			break;
 		}
diff --git a/Engine/Plugins/Experimental/MeshModelingToolsetExp/Source/ModelingComponents/Public/Selection/DynamicMeshSelection.h b/Engine/Plugins/Experimental/MeshModelingToolsetExp/Source/ModelingComponents/Public/Selection/DynamicMeshSelection.h
--- a/Engine/Plugins/Experimental/MeshModelingToolsetExp/Source/ModelingComponents/Public/Selection/DynamicMeshSelection.h
+++ b/Engine/Plugins/Experimental/MeshModelingToolsetExp/Source/ModelingComponents/Public/Selection/DynamicMeshSelection.h
@@ -38,7 +38,7 @@ public:
 	/** Checks whether the selection's timestamp still matches the meshes topology timestamp. */
 	bool MatchesTimestamp() const
 	{
-		return Mesh && Mesh->GetTopologyTimestamp() == TopologyTimestamp;
+		return Mesh && Mesh->GetTopologyChangeStamp() == TopologyTimestamp;
 	}
 
 	bool IsEmpty() const
diff --git a/Engine/Plugins/Experimental/PlanarCutPlugin/Source/PlanarCut/Private/FractureAutoUV.cpp b/Engine/Plugins/Experimental/PlanarCutPlugin/Source/PlanarCut/Private/FractureAutoUV.cpp
--- a/Engine/Plugins/Experimental/PlanarCutPlugin/Source/PlanarCut/Private/FractureAutoUV.cpp
+++ b/Engine/Plugins/Experimental/PlanarCutPlugin/Source/PlanarCut/Private/FractureAutoUV.cpp
@@ -272,9 +272,9 @@ struct FGeomMesh : public UE::Geometry::FUVPacker::IUVMeshView
 	{
 		return GlobalVertices.Num();
 	}
-	constexpr inline int32 GetShapeTimestamp() const
+	constexpr inline uint64 GetChangeStamp() const
 	{
-		return 0;
+		return 1;
 	}
 
 	inline void GetTriVertices(int TID, UE::Math::TVector<double>& V0, UE::Math::TVector<double>& V1, UE::Math::TVector<double>& V2) const
@@ -354,9 +354,9 @@ struct FGeomFlatUVMesh
 	{
 		return Collection->UVs.Num();
 	}
-	constexpr inline int32 GetShapeTimestamp() const
+	constexpr inline uint64 GetChangeStamp() const
 	{
-		return 0;
+		return 1;
 	}
 
 	inline void GetTriVertices(int TID, UE::Math::TVector<double>& V0, UE::Math::TVector<double>& V1, UE::Math::TVector<double>& V2) const
diff --git a/Engine/Plugins/Experimental/UVEditor/Source/UVEditor/Private/UVEditorMode.cpp b/Engine/Plugins/Experimental/UVEditor/Source/UVEditor/Private/UVEditorMode.cpp
--- a/Engine/Plugins/Experimental/UVEditor/Source/UVEditor/Private/UVEditorMode.cpp
+++ b/Engine/Plugins/Experimental/UVEditor/Source/UVEditor/Private/UVEditorMode.cpp
@@ -221,7 +221,7 @@ void UUVEditorMode::InitializeTargets(TArray<TObjectPtr<UObject>>& AssetsIn, TAr
 				TriangleOpacity));
 
 		// Initialize our timestamp so we can later detect changes
-		MeshChangeStamps.Add(ToolInputObject->UnwrapCanonical->GetShapeTimestamp());
+		MeshChangeStamps.Add(ToolInputObject->UnwrapCanonical->GetShapeChangeStamp());
 
 		// Set up the wireframe display of the unwrapped mesh.
 		UMeshElementsVisualizer* WireframeDisplay = NewObject<UMeshElementsVisualizer>(this);
@@ -273,7 +273,7 @@ bool UUVEditorMode::HaveUnappliedChanges()
 {
 	for (int32 i = 0; i < ToolInputObjects.Num(); ++i)
 	{
-		if (ToolInputObjects[i]->UnwrapCanonical->GetShapeTimestamp() != MeshChangeStamps[i])
+		if (ToolInputObjects[i]->UnwrapCanonical->GetShapeChangeStamp() != MeshChangeStamps[i])
 		{
 			return true;
 		}
@@ -285,7 +285,7 @@ void UUVEditorMode::GetAssetsWithUnappliedChanges(TArray<TObjectPtr<UObject>> Un
 {
 	for (int32 i = 0; i < ToolInputObjects.Num(); ++i)
 	{
-		if (ToolInputObjects[i]->UnwrapCanonical->GetShapeTimestamp() != MeshChangeStamps[i])
+		if (ToolInputObjects[i]->UnwrapCanonical->GetShapeChangeStamp() != MeshChangeStamps[i])
 		{
 			UnappliedAssetsOut.Add(OriginalObjectsToEdit[i]);
 		}
@@ -303,10 +303,10 @@ void UUVEditorMode::ApplyChanges()
 
 	for (int32 i = 0; i < ToolInputObjects.Num(); ++i)
 	{
-		if (ToolInputObjects[i]->UnwrapCanonical->GetShapeTimestamp() != MeshChangeStamps[i])
+		if (ToolInputObjects[i]->UnwrapCanonical->GetShapeChangeStamp() != MeshChangeStamps[i])
 		{
 			Cast<IDynamicMeshCommitter>(ToolTargets[i])->CommitDynamicMesh(*ToolInputObjects[i]->AppliedCanonical, CommitInfo);
-			MeshChangeStamps[i] = ToolInputObjects[i]->UnwrapCanonical->GetShapeTimestamp();
+			MeshChangeStamps[i] = ToolInputObjects[i]->UnwrapCanonical->GetShapeChangeStamp();
 		}
 	}
 
diff --git a/Engine/Plugins/Experimental/UVEditor/Source/UVEditorTools/Private/ToolTargets/UVEditorToolMeshInput.cpp b/Engine/Plugins/Experimental/UVEditor/Source/UVEditorTools/Private/ToolTargets/UVEditorToolMeshInput.cpp
--- a/Engine/Plugins/Experimental/UVEditor/Source/UVEditorTools/Private/ToolTargets/UVEditorToolMeshInput.cpp
+++ b/Engine/Plugins/Experimental/UVEditor/Source/UVEditorTools/Private/ToolTargets/UVEditorToolMeshInput.cpp
@@ -165,6 +165,7 @@ bool UUVEditorToolMeshInput::InitializeMeshes(
 	UVEditorToolUtil::GenerateUVUnwrapMesh(
 		*AppliedCanonical->Attributes()->GetUVLayer(UVLayerIndex),
 		*UnwrapCanonical, UVToVertPosition);
+	UnwrapCanonical->SetShapeChangeStampEnabled(true);
 
 	// Set up the unwrap preview
 	UnwrapPreview = NewObject<UMeshOpPreviewWithBackgroundCompute>();
diff --git a/Engine/Plugins/Experimental/UVEditor/Source/UVEditorTools/Private/UVSelectTool.cpp b/Engine/Plugins/Experimental/UVEditor/Source/UVEditorTools/Private/UVSelectTool.cpp
--- a/Engine/Plugins/Experimental/UVEditor/Source/UVEditorTools/Private/UVSelectTool.cpp
+++ b/Engine/Plugins/Experimental/UVEditor/Source/UVEditorTools/Private/UVSelectTool.cpp
@@ -316,7 +316,7 @@ void UUVSelectTool::Setup()
 			Tree->SetMesh(Target->UnwrapCanonical.Get());
 			TreeStore->Set(Target->UnwrapCanonical.Get(), Tree);
 		}
-		if (!Tree->IsValid())
+		if (!Tree->IsValid(false))
 		{
 			Tree->Build();
 		}
@@ -418,7 +418,7 @@ void UUVSelectTool::Shutdown(EToolShutdownType ShutdownType)
 		// to our undo/redo change implementation sometimes incrementing the topology timestamp
 		// even when the topology hasn't changed.
 		// TODO: Check whether we still need this once we use FMeshVertexChange instead of FDynamicMeshChange
-		SelectionStore->Selection->TopologyTimestamp = SelectionStore->Selection->Mesh->GetTopologyTimestamp();
+		SelectionStore->Selection->TopologyTimestamp = SelectionStore->Selection->Mesh->GetTopologyChangeStamp();
 
 		// Don't issue a transaciton if we're cancelling (currently only possible via undo out
 		// of the "Transform" tool). Otherwise, issue a speculative change so that we can clear
@@ -735,7 +735,7 @@ void UUVSelectTool::GizmoTransformEnded(UTransformProxy* Proxy)
 		Targets[SelectionTargetIndex]->UpdateAllFromUnwrapPreview(&MovingVids, nullptr, &SelectedTids);
 	}
 
-	if (!AABBTrees[SelectionTargetIndex]->IsValid())
+	if (!AABBTrees[SelectionTargetIndex]->IsValid(false))
 	{
 		AABBTrees[SelectionTargetIndex]->Build();
 	}
diff --git a/Engine/Plugins/Runtime/GeometryProcessing/Source/DynamicMesh/Private/Operations/RemoveOccludedTriangles.cpp b/Engine/Plugins/Runtime/GeometryProcessing/Source/DynamicMesh/Private/Operations/RemoveOccludedTriangles.cpp
--- a/Engine/Plugins/Runtime/GeometryProcessing/Source/DynamicMesh/Private/Operations/RemoveOccludedTriangles.cpp
+++ b/Engine/Plugins/Runtime/GeometryProcessing/Source/DynamicMesh/Private/Operations/RemoveOccludedTriangles.cpp
@@ -52,9 +52,9 @@ namespace UE
 				return Mesh->VertexCount();
 			}
 
-			FORCEINLINE int32 GetShapeTimestamp() const
+			FORCEINLINE uint64 GetChangeStamp() const
 			{
-				return 0;
+				return 1;
 			}
 
 			FORCEINLINE FIndex3i GetTriangle(int32 Index) const
diff --git a/Engine/Plugins/Runtime/GeometryProcessing/Source/DynamicMesh/Private/Remesher.cpp b/Engine/Plugins/Runtime/GeometryProcessing/Source/DynamicMesh/Private/Remesher.cpp
--- a/Engine/Plugins/Runtime/GeometryProcessing/Source/DynamicMesh/Private/Remesher.cpp
+++ b/Engine/Plugins/Runtime/GeometryProcessing/Source/DynamicMesh/Private/Remesher.cpp
@@ -555,37 +555,27 @@ void FRemesher::ApplyVertexBuffer(bool bParallel)
 		{
 			if (TempFlagBuffer[vid])
 			{
-				Mesh->SetVertex(vid, TempPosBuffer[vid]);
+				Mesh->SetVertex(vid, TempPosBuffer[vid], false);
 			}
 		}
+		Mesh->UpdateChangeStamps(true, false);
 	}
 	else
 	{
-		// The serial version calls Mesh->SetVertex which sets the vertex position and also updates the timestamps. In
-		// the parallel version, to avoid contention for the timestamps we work in batches, calling 
-		// SetVertex_NoTimeStampUpdate several times for each async task, then finally locking and incrementing the 
-		// timestamp once per task.
-
-		// TODO: Try varying the batch size to observe a trade-off between work done per thread and overhead due to 
-		// locking.
-
+		// TODO: verify that this batching is still necessary, now that timestamps/locking situation has been improved
 		const int BatchSize = 1000;
 		const int NumBatches = FMath::CeilToInt(Mesh->MaxVertexID() / static_cast<float>(BatchSize));
-
 		ParallelFor(NumBatches, [this, BatchSize](int32 BatchID)
 		{
-			int NumVerticesSetThisBatch = 0;
 			for (int VertexID = BatchID * BatchSize; VertexID < (BatchID + 1) * BatchSize; ++VertexID)
 			{
 				if (VertexID < TempFlagBuffer.Num() && TempFlagBuffer[VertexID] && Mesh->IsVertex(VertexID))
 				{
-					Mesh->SetVertex_NoTimeStampUpdate(VertexID, TempPosBuffer[VertexID]);
-					++NumVerticesSetThisBatch;
+					Mesh->SetVertex(VertexID, TempPosBuffer[VertexID], false);
 				}
 			}
-			Mesh->IncrementTimeStamps(NumVerticesSetThisBatch, true, false);
-
 		}, false);
+		Mesh->UpdateChangeStamps(true, false);
 	}
 }
 
diff --git a/Engine/Plugins/Runtime/GeometryProcessing/Source/DynamicMesh/Public/Operations/DisplaceMesh.h b/Engine/Plugins/Runtime/GeometryProcessing/Source/DynamicMesh/Public/Operations/DisplaceMesh.h
--- a/Engine/Plugins/Runtime/GeometryProcessing/Source/DynamicMesh/Public/Operations/DisplaceMesh.h
+++ b/Engine/Plugins/Runtime/GeometryProcessing/Source/DynamicMesh/Public/Operations/DisplaceMesh.h
@@ -37,9 +37,9 @@ public:
 
 			float Offset = VertexWeights[VertexID];
 			FVector3d NewPosition = Mesh.GetVertexRef(VertexID) + (Offset * Intensity * Normals[VertexID]);
-			Mesh.SetVertex_NoTimeStampUpdate(VertexID, NewPosition);
+			Mesh.SetVertex(VertexID, NewPosition, false);
 		});
-		Mesh.IncrementTimeStamps(1, true, false);
+		Mesh.UpdateChangeStamps(true, false);
 	}
 
 
diff --git a/Engine/Plugins/Runtime/GeometryProcessing/Source/DynamicMesh/Public/Sampling/MeshSurfaceSampler.h b/Engine/Plugins/Runtime/GeometryProcessing/Source/DynamicMesh/Public/Sampling/MeshSurfaceSampler.h
--- a/Engine/Plugins/Runtime/GeometryProcessing/Source/DynamicMesh/Public/Sampling/MeshSurfaceSampler.h
+++ b/Engine/Plugins/Runtime/GeometryProcessing/Source/DynamicMesh/Public/Sampling/MeshSurfaceSampler.h
@@ -50,9 +50,9 @@ struct FDynamicMeshUVMesh
 		return UV->ElementCount();
 	}
 
-	inline int32 GetShapeTimestamp() const
+	inline uint64 GetChangeStamp() const
 	{
-		return Mesh->GetShapeTimestamp();
+		return Mesh->GetChangeStamp();
 	}
 
 	inline FIndex3i GetTriangle(int32 TriangleIndex) const
diff --git a/Engine/Plugins/Runtime/GeometryProcessing/Source/DynamicMesh/Public/SubRegionRemesher.h b/Engine/Plugins/Runtime/GeometryProcessing/Source/DynamicMesh/Public/SubRegionRemesher.h
--- a/Engine/Plugins/Runtime/GeometryProcessing/Source/DynamicMesh/Public/SubRegionRemesher.h
+++ b/Engine/Plugins/Runtime/GeometryProcessing/Source/DynamicMesh/Public/SubRegionRemesher.h
@@ -464,10 +464,10 @@ protected:
 		{
 			if (SmoothIndices[i] >= 0)
 			{
-				Mesh->SetVertex_NoTimeStampUpdate(SmoothIndices[i], SmoothedPositions[i]);
+				Mesh->SetVertex(SmoothIndices[i], SmoothedPositions[i], false);
 			}
 		});
-		Mesh->IncrementTimeStamps(1, true, false);
+		Mesh->UpdateChangeStamps(true, false);
 	}
 };
 
diff --git a/Engine/Source/Runtime/GeometryCore/Private/DynamicMesh/DynamicMesh3.cpp b/Engine/Source/Runtime/GeometryCore/Private/DynamicMesh/DynamicMesh3.cpp
--- a/Engine/Source/Runtime/GeometryCore/Private/DynamicMesh/DynamicMesh3.cpp
+++ b/Engine/Source/Runtime/GeometryCore/Private/DynamicMesh/DynamicMesh3.cpp
@@ -41,12 +41,11 @@ FDynamicMesh3::FDynamicMesh3(const FDynamicMesh3& Other)
 	GroupIDCounter{ Other.GroupIDCounter },
 
 	Edges{ Other.Edges },
-	EdgeRefCounts{ Other.EdgeRefCounts },
-
-	Timestamp{ Other.Timestamp },
-	ShapeTimestamp{ Other.ShapeTimestamp },
-	TopologyTimestamp{ Other.TopologyTimestamp }
+	EdgeRefCounts{ Other.EdgeRefCounts }
 {
+	ShapeChangeStamp.store(Other.ShapeChangeStamp);
+	TopologyChangeStamp.store(Other.TopologyChangeStamp);
+
 	if (Other.HasAttributes())
 	{
 		EnableAttributes();
@@ -72,11 +71,11 @@ FDynamicMesh3::FDynamicMesh3(FDynamicMesh3&& Other)
 	Edges{ MoveTemp( Other.Edges ) },
 	EdgeRefCounts{ MoveTemp( Other.EdgeRefCounts ) },
 
-	AttributeSet{ MoveTemp( Other.AttributeSet ) },
-	Timestamp{ MoveTemp( Other.Timestamp ) },
-	ShapeTimestamp{ MoveTemp( Other.ShapeTimestamp ) },
-	TopologyTimestamp{ MoveTemp( Other.TopologyTimestamp ) }
+	AttributeSet{ MoveTemp( Other.AttributeSet ) }
 {
+	ShapeChangeStamp.store(Other.ShapeChangeStamp);
+	TopologyChangeStamp.store(Other.TopologyChangeStamp);
+
 	if (AttributeSet)
 	{
 		AttributeSet->Reparent(this);
@@ -114,9 +113,9 @@ const FDynamicMesh3 & FDynamicMesh3::operator=(FDynamicMesh3 && Other)
 	{
 		AttributeSet->Reparent(this);
 	}
-	Timestamp = MoveTemp(Other.Timestamp);
-	ShapeTimestamp = MoveTemp(Other.ShapeTimestamp);
-	TopologyTimestamp = MoveTemp(Other.TopologyTimestamp);
+	
+	ShapeChangeStamp.store(Other.ShapeChangeStamp);
+	TopologyChangeStamp.store(Other.TopologyChangeStamp);
 
 	return *this;
 }
@@ -222,8 +221,8 @@ void FDynamicMesh3::Copy(const FDynamicMesh3& copy, bool bNormals, bool bColors,
 		DiscardAttributes();
 	}
 
-	Timestamp = FMath::Max(Timestamp + 1, copy.Timestamp);
-	ShapeTimestamp = TopologyTimestamp = Timestamp;
+	ShapeChangeStamp.store(copy.ShapeChangeStamp);
+	TopologyChangeStamp.store(copy.TopologyChangeStamp);
 }
 
 void FDynamicMesh3::CompactCopy(const FDynamicMesh3& copy, bool bNormals, bool bColors, bool bUVs, bool bAttributes, FCompactMaps* CompactInfo)
@@ -296,8 +295,8 @@ void FDynamicMesh3::CompactCopy(const FDynamicMesh3& copy, bool bNormals, bool b
 		AttributeSet->CompactCopy(*CompactInfo, *copy.Attributes());
 	}
 
-	Timestamp = FMath::Max(Timestamp + 1, copy.Timestamp);
-	ShapeTimestamp = TopologyTimestamp = Timestamp;
+	ShapeChangeStamp.store(copy.ShapeChangeStamp);
+	TopologyChangeStamp.store(copy.TopologyChangeStamp);
 }
 
 void FDynamicMesh3::Clear()
@@ -319,9 +318,9 @@ void FDynamicMesh3::Clear()
 	EdgeRefCounts.Clear();
 
 	AttributeSet.Reset();
-	Timestamp = 0;
-	ShapeTimestamp = 0;
-	TopologyTimestamp = 0;
+
+	ShapeChangeStamp.store(1);
+	TopologyChangeStamp.store(1);
 }
 
 
@@ -600,8 +599,8 @@ FString FDynamicMesh3::MeshInfoString()
 		EdgeCount(), MaxEdgeID(), *EdgeRefCounts.UsageStats());
 	FString AttribString = FString::Printf(TEXT("VtxNormals %d  VtxColors %d  VtxUVs %d  TriGroups %d  Attributes %d"),
 		HasVertexNormals(), HasVertexColors(), HasVertexUVs(), HasTriangleGroups(), HasAttributes());
-	FString InfoString = FString::Printf(TEXT("Closed %d  Compact %d  Timestamp %d  ShapeTimestamp %d  TopologyTimestamp %d  MaxGroupID %d"),
-		IsClosed(), IsCompact(), GetTimestamp(), GetShapeTimestamp(), GetTopologyTimestamp(), MaxGroupID());
+	FString InfoString = FString::Printf(TEXT("Closed %d  Compact %d  ShapeChangeStamp %d  TopologyChangeStamp %d  MaxGroupID %d"),
+		IsClosed(), IsCompact(), GetShapeChangeStamp(), GetTopologyChangeStamp(), MaxGroupID());
 
 	return VtxString + "\n" + TriString + "\n" + EdgeString + "\n" + AttribString + "\n" + InfoString;
 }
diff --git a/Engine/Source/Runtime/GeometryCore/Private/DynamicMesh/DynamicMesh3_Edits.cpp b/Engine/Source/Runtime/GeometryCore/Private/DynamicMesh/DynamicMesh3_Edits.cpp
--- a/Engine/Source/Runtime/GeometryCore/Private/DynamicMesh/DynamicMesh3_Edits.cpp
+++ b/Engine/Source/Runtime/GeometryCore/Private/DynamicMesh/DynamicMesh3_Edits.cpp
@@ -37,7 +37,7 @@ int FDynamicMesh3::AppendVertex(const FVertexInfo& VtxInfo)
 		Attributes()->OnNewVertex(vid, false);
 	}
 
-	UpdateTimeStamp(true, true);
+	UpdateChangeStamps(true, true);
 	return vid;
 }
 
@@ -94,7 +94,7 @@ int FDynamicMesh3::AppendVertex(const FDynamicMesh3& from, int fromVID)
 		Attributes()->OnNewVertex(vid, false);
 	}
 
-	UpdateTimeStamp(true, true);
+	UpdateChangeStamps(true, true);
 	return vid;
 }
 
@@ -141,7 +141,7 @@ EMeshResult FDynamicMesh3::InsertVertex(int vid, const FVertexInfo& info, bool b
 		Attributes()->OnNewVertex(vid, true);
 	}
 
-	UpdateTimeStamp(true, true);
+	UpdateChangeStamps(true, true);
 	return EMeshResult::Ok;
 }
 
@@ -211,7 +211,7 @@ int FDynamicMesh3::AppendTriangle(const FIndex3i& tv, int gid)
 		Attributes()->OnNewTriangle(tid, false);
 	}
 
-	UpdateTimeStamp(true, true);
+	UpdateChangeStamps(true, true);
 	return tid;
 }
 
@@ -277,7 +277,7 @@ EMeshResult FDynamicMesh3::InsertTriangle(int tid, const FIndex3i& tv, int gid,
 		Attributes()->OnNewTriangle(tid, true);
 	}
 
-	UpdateTimeStamp(true, true);
+	UpdateChangeStamps(true, true);
 	return EMeshResult::Ok;
 }
 
@@ -540,7 +540,7 @@ EMeshResult FDynamicMesh3::ReverseTriOrientation(int tID)
 		return EMeshResult::Failed_NotATriangle;
 	}
 	ReverseTriOrientationInternal(tID);
-	UpdateTimeStamp(true, true);
+	UpdateChangeStamps(true, true);
 	return EMeshResult::Ok;
 }
 
@@ -570,7 +570,7 @@ void FDynamicMesh3::ReverseOrientation(bool bFlipNormals)
 			Normals[vid] = -Normals[vid];
 		}
 	}
-	UpdateTimeStamp(true, true);
+	UpdateChangeStamps(true, true);
 }
 
 EMeshResult FDynamicMesh3::RemoveVertex(int vID, bool bRemoveAllTriangles, bool bPreserveManifold)
@@ -624,7 +624,7 @@ EMeshResult FDynamicMesh3::RemoveVertex(int vID, bool bRemoveAllTriangles, bool
 	ensure(VertexRefCounts.IsValid(vID) == false);
 	VertexEdgeLists.Clear(vID);
 
-	UpdateTimeStamp(true, true);
+	UpdateChangeStamps(true, true);
 	return EMeshResult::Ok;
 }
 
@@ -705,7 +705,7 @@ EMeshResult FDynamicMesh3::RemoveTriangle(int tID, bool bRemoveIsolatedVertices,
 		Attributes()->OnRemoveTriangle(tID);
 	}
 
-	UpdateTimeStamp(true, true);
+	UpdateChangeStamps(true, true);
 	return EMeshResult::Ok;
 }
 
@@ -832,7 +832,7 @@ EMeshResult FDynamicMesh3::SetTriangle(int tID, const FIndex3i& newv, bool bRemo
 		AddTriangleEdge(tID, newv[2], newv[0], 2, e2);
 	}
 
-	UpdateTimeStamp(true, true);
+	UpdateChangeStamps(true, true);
 	return EMeshResult::Ok;
 }
 
@@ -945,7 +945,7 @@ EMeshResult FDynamicMesh3::SplitEdge(int eab, FEdgeSplitInfo& SplitInfo, double
 			Attributes()->OnSplitEdge(SplitInfo);
 		}
 
-		UpdateTimeStamp(true, true);
+		UpdateChangeStamps(true, true);
 		return EMeshResult::Ok;
 
 	}
@@ -1041,7 +1041,7 @@ EMeshResult FDynamicMesh3::SplitEdge(int eab, FEdgeSplitInfo& SplitInfo, double
 			Attributes()->OnSplitEdge(SplitInfo);
 		}
 
-		UpdateTimeStamp(true, true);
+		UpdateChangeStamps(true, true);
 		return EMeshResult::Ok;
 	}
 
@@ -1164,7 +1164,7 @@ EMeshResult FDynamicMesh3::FlipEdge(int eab, FEdgeFlipInfo& FlipInfo)
 		Attributes()->OnFlipEdge(FlipInfo);
 	}
 
-	UpdateTimeStamp(true, true);
+	UpdateChangeStamps(true, true);
 	return EMeshResult::Ok;
 }
 
@@ -1238,7 +1238,7 @@ EMeshResult FDynamicMesh3::SplitVertex(int VertexID, const TArrayView<const int>
 	{
 		Attributes()->OnSplitVertex(SplitInfo, TrianglesToUpdate);
 	}
-	UpdateTimeStamp(true, true);
+	UpdateChangeStamps(true, true);
 	return EMeshResult::Ok;
 }
 
@@ -1602,7 +1602,7 @@ EMeshResult FDynamicMesh3::CollapseEdge(int vKeep, int vRemove, double collapse_
 		Attributes()->OnCollapseEdge(CollapseInfo);
 	}
 
-	UpdateTimeStamp(true, true);
+	UpdateChangeStamps(true, true);
 	return EMeshResult::Ok;
 }
 
@@ -1870,7 +1870,7 @@ EMeshResult FDynamicMesh3::MergeEdges(int eKeep, int eDiscard, FMergeEdgesInfo&
 		Attributes()->OnMergeEdges(MergeInfo);
 	}
 
-	UpdateTimeStamp(true, true);
+	UpdateChangeStamps(true, true);
 	return EMeshResult::Ok;
 }
 
@@ -1945,7 +1945,7 @@ EMeshResult FDynamicMesh3::PokeTriangle(int TriangleID, const FVector3d& BaryCoo
 		Attributes()->OnPokeTriangle(PokeResult);
 	}
 
-	UpdateTimeStamp(true, true);
+	UpdateChangeStamps(true, true);
 	return EMeshResult::Ok;
 }
 
diff --git a/Engine/Source/Runtime/GeometryCore/Private/DynamicMesh/MeshAdapterUtil.cpp b/Engine/Source/Runtime/GeometryCore/Private/DynamicMesh/MeshAdapterUtil.cpp
--- a/Engine/Source/Runtime/GeometryCore/Private/DynamicMesh/MeshAdapterUtil.cpp
+++ b/Engine/Source/Runtime/GeometryCore/Private/DynamicMesh/MeshAdapterUtil.cpp
@@ -12,7 +12,6 @@ FPointSetAdapterd UE::Geometry::MakePointsAdapter(const FDynamicPointSet3d* Poin
 	Adapter.PointCount = [PointSet]() { return PointSet->VertexCount(); };
 	Adapter.IsPoint = [PointSet](int Idx) { return PointSet->IsVertex(Idx); };
 	Adapter.GetPoint = [PointSet](int Idx) { return PointSet->GetVertex(Idx); };
-	Adapter.Timestamp = [PointSet] { return PointSet->GetTimestamp(); };
 
 	Adapter.HasNormals = [PointSet] { return false; };
 	Adapter.GetPointNormal = [PointSet](int Idx) {return FVector3f(0,0,1); };
@@ -29,7 +28,6 @@ FPointSetAdapterd UE::Geometry::MakeVerticesAdapter(const FDynamicMesh3* Mesh)
 	Adapter.PointCount = [Mesh]() { return Mesh->VertexCount(); };
 	Adapter.IsPoint = [Mesh](int Idx) { return Mesh->IsVertex(Idx); };
 	Adapter.GetPoint = [Mesh](int Idx) { return Mesh->GetVertex(Idx); };
-	Adapter.Timestamp = [Mesh] { return Mesh->GetTimestamp(); };
 
 	Adapter.HasNormals = [Mesh] { return Mesh->HasVertexNormals(); };
 	Adapter.GetPointNormal = [Mesh](int Idx) {return Mesh->GetVertexNormal(Idx); };
@@ -45,7 +43,6 @@ FPointSetAdapterd UE::Geometry::MakeTriCentroidsAdapter(const FDynamicMesh3* Mes
 	Adapter.PointCount = [Mesh]() { return Mesh->TriangleCount(); };
 	Adapter.IsPoint = [Mesh](int Idx) { return Mesh->IsTriangle(Idx); };
 	Adapter.GetPoint = [Mesh](int Idx) { return Mesh->GetTriCentroid(Idx); };
-	Adapter.Timestamp = [Mesh] { return Mesh->GetTimestamp(); };
 
 	Adapter.HasNormals = [] { return true; };
 	Adapter.GetPointNormal = [Mesh](int Idx) {return (FVector3f)Mesh->GetTriNormal(Idx); };
@@ -63,7 +60,6 @@ FPointSetAdapterd UE::Geometry::MakeEdgeMidpointsAdapter(const FDynamicMesh3* Me
 	Adapter.PointCount = [Mesh]() { return Mesh->EdgeCount(); };
 	Adapter.IsPoint = [Mesh] (int Idx) { return Mesh->IsEdge(Idx); };
 	Adapter.GetPoint = [Mesh](int Idx) { return Mesh->GetEdgePoint(Idx, 0.5); };
-	Adapter.Timestamp = [Mesh] { return Mesh->GetTimestamp(); };
 
 	Adapter.HasNormals = [] { return false; };
 	Adapter.GetPointNormal = [](int Idx) { return FVector3f::UnitY();};
@@ -86,7 +82,6 @@ FPointSetAdapterd UE::Geometry::MakeBoundaryEdgeMidpointsAdapter(const FDynamicM
 	Adapter.PointCount = [NumBoundaryEdges]() { return NumBoundaryEdges; };
 	Adapter.IsPoint = [Mesh](int Idx) { return Mesh->IsEdge(Idx) && Mesh->IsBoundaryEdge(Idx); };
 	Adapter.GetPoint = [Mesh](int Idx) { return Mesh->GetEdgePoint(Idx, 0.5); };
-	Adapter.Timestamp = [Mesh] { return Mesh->GetTimestamp(); };
 
 	Adapter.HasNormals = [] { return false; };
 	Adapter.GetPointNormal = [](int Idx) { return FVector3f::UnitY(); };
@@ -117,7 +112,7 @@ FTriangleMeshAdapterd UE::Geometry::MakeDynamicMeshAdapter(const FDynamicMesh3*
 		[Mesh]() { return Mesh->MaxVertexID(); },
 		[Mesh]() { return Mesh->TriangleCount(); },
 		[Mesh]() { return Mesh->VertexCount(); },
-		[Mesh]() { return Mesh->GetShapeTimestamp(); },
+		[Mesh]() { return Mesh->GetChangeStamp(); },
 		[Mesh](int Idx) { return Mesh->GetTriangle(Idx); },
 		[Mesh](int Idx) { return Mesh->GetVertex(Idx); }
 	};
diff --git a/Engine/Source/Runtime/GeometryCore/Public/DynamicMesh/DynamicMesh3.h b/Engine/Source/Runtime/GeometryCore/Public/DynamicMesh/DynamicMesh3.h
--- a/Engine/Source/Runtime/GeometryCore/Public/DynamicMesh/DynamicMesh3.h
+++ b/Engine/Source/Runtime/GeometryCore/Public/DynamicMesh/DynamicMesh3.h
@@ -77,19 +77,15 @@ enum class EMeshComponents : uint8
 *
 * The function CheckValidity() does extensive sanity checking on the mesh data structure.
 * Use this to test your code, both for mesh construction and editing!!
-*
-*
-* TODO:
-*  - Many of the iterators depend on lambda functions, can we replace these with calls to
-*    internal/static functions that do the same thing?
-*  - efficient TriTrianglesItr() implementation?
-*  - additional Topology timestamp?
-*  - CompactInPlace() does not compact VertexEdgeLists
-*  ? TDynamicVector w/ 'stride' option, so that we can guarantee that tuples are in single block.
-*    The can have custom accessor that looks up entire tuple
 */
 class GEOMETRYCORE_API FDynamicMesh3
 {
+
+// TODO:
+//  - Many of the iterators depend on lambda functions, can we replace these with calls to
+//    internal/static functions that do the same thing?
+//  - CompactInPlace() does not compact VertexEdgeLists
+
 public:
 	using FEdgeFlipInfo     = DynamicMeshInfo::FEdgeFlipInfo;
 	using FEdgeSplitInfo    = DynamicMeshInfo::FEdgeSplitInfo;
@@ -168,14 +164,15 @@ protected:
 	/** Reference counts of edge indices. Ref count is always 1 if the edge exists. Iterate over this to find out which edge indices are valid. */
 	FRefCountVector EdgeRefCounts;
 
+	/** Extended Attributes for the Mesh (UV layers, Hard Normals, additional Polygroup Layers, etc) */
 	TUniquePtr<FDynamicMeshAttributeSet> AttributeSet{};
 
-	/** The mesh timestamp is incremented any time a function that modifies the mesh is called */
-	int Timestamp = 0;
-	/** The shape timestamp is incremented any time a function that modifies the mesh shape or topology is called */
-	int ShapeTimestamp = 0;
-	/** The topology timestamp is incremented any time a function that modifies the mesh topology is called */
-	int TopologyTimestamp = 0;
+	/** Enable/Disable updating of ShapeChangeStamp. This can be problematic in multi-threaded contexts so it is disabled by default. */
+	bool bEnableShapeChangeStamp = false;
+	/** If bEnableShapeChangeStamp=true, The shape change stamp is incremented any time a function that modifies the mesh shape or topology is called */
+	std::atomic<uint32> ShapeChangeStamp = 1;
+	/** The topology change stamp is incremented any time a function that modifies the mesh topology is called */
+	std::atomic<uint32> TopologyChangeStamp = 1;
 
 public:
 	/** Default constructor */
@@ -341,6 +338,61 @@ public:
 	}
 
 
+
+
+	//
+	// Change Tracking support
+	// 
+	// 
+	// 
+public:
+
+	/** Enable/Disable incrementing of the ShapeChangeStamp */
+	void SetShapeChangeStampEnabled(bool bEnabled)
+	{
+		bEnableShapeChangeStamp = bEnabled;
+	}
+
+	/** @return true if Shape ChangeStamp is enabled (disabled by default) */
+	bool HasShapeChangeStampEnabled() const
+	{
+		return bEnableShapeChangeStamp;
+	}
+
+	/** Increment the specified ChangeStamps, if they are enabled. Thread-safe. */
+	void UpdateChangeStamps(bool bShapeChange, bool bTopologyChange)
+	{
+		if (bEnableShapeChangeStamp && (bShapeChange || bTopologyChange))
+		{
+			ShapeChangeStamp++;		// atomic increment
+		}
+		if (bTopologyChange)
+		{
+			TopologyChangeStamp++;	// atomic increment
+		}
+	}
+
+	/** ShapeChangeStamp is incremented any time a mesh vertex position is changed or the mesh topology is modified, if bEnableShapeChangeStamp=true */
+	inline uint32 GetShapeChangeStamp() const
+	{
+		ensureMsgf(bEnableShapeChangeStamp, TEXT("ShapeChange Tracking is not enabled on this mesh. Use SetShapeChangeStampEnabled() to enable, or use GetChangeStamp() to avoid this ensure."));
+		return ShapeChangeStamp;
+	}
+
+	/** TopologyChangeStamp is incremented when the mesh topology is modified */
+	inline uint32 GetTopologyChangeStamp() const
+	{
+		return TopologyChangeStamp;
+	}
+
+	/** ChangeStamp is a combination of the Shape and Topology ChangeStamps */
+	inline uint64 GetChangeStamp() const
+	{
+		return ShapeChangeStamp + TopologyChangeStamp;
+	}
+
+
+
 	//
 	// Mesh Element Iterators
 	//   The functions VertexIndicesItr() / TriangleIndicesItr() / EdgeIndicesItr() allow you to do:
@@ -534,51 +586,24 @@ public:
 		return Vertices[VertexID];
 	}
 
-	/** Set vertex position */
-	inline void SetVertex(int VertexID, const FVector3d& vNewPos)
-	{
-		checkSlow(VectorUtil::IsFinite(vNewPos));
-		checkSlow(IsVertex(VertexID));
-		if (VectorUtil::IsFinite(vNewPos))
-		{
-			Vertices[VertexID] = vNewPos;
-			UpdateTimeStamp(true, false);
-		}
-	}
-
-	/** Set vertex position. DOES NOT UPDATE TIMESTAMPS.
-	This is meant to be used in conjunction with IncrementTimeStamps, below. Set a bunch of vertex positions then 
-	increment the timestamps with the number of vertices that were moved, rather than incrementing the timestamps for
-	each vertex.
-	*/
-	void SetVertex_NoTimeStampUpdate(int VertexID, const FVector3d& vNewPos)
+	/** 
+	 * Set vertex position 
+	 * @param bTrackChange if true, ShapeChangeStamp will be incremented (if enabled)
+	 */
+	inline void SetVertex(int VertexID, const FVector3d& vNewPos, bool bTrackChange = true)
 	{
 		checkSlow(VectorUtil::IsFinite(vNewPos));
 		checkSlow(IsVertex(VertexID));
-
 		if (VectorUtil::IsFinite(vNewPos))
 		{
 			Vertices[VertexID] = vNewPos;
+			if (bTrackChange)
+			{
+				UpdateChangeStamps(true, false);
+			}
 		}
 	}
 
-	/** Increment the specified timestamps by the given amount. Thread-safe. 
-	This is meant to be used in conjunction with SetVertex_NoTimeStampUpdate, above. Set a bunch of vertex positions then
-	increment the timestamps with the number of vertices that were moved, rather than incrementing the timestamps for
-	each vertex.
-	*/
-	void IncrementTimeStamps(int Amount, bool bShapeChange, bool bTopologyChange)
-	{
-		FPlatformAtomics::InterlockedAdd(&Timestamp, Amount);
-		if (bShapeChange)
-		{
-			FPlatformAtomics::InterlockedAdd(&ShapeTimestamp, Amount);
-		}
-		if (bTopologyChange)
-		{
-			FPlatformAtomics::InterlockedAdd(&TopologyTimestamp, Amount);
-		}
-	}
 
 	/** Get extended vertex information */
 	bool GetVertex(int VertexID, FVertexInfo& VertInfo, bool bWantNormals, bool bWantColors, bool bWantUVs) const;
@@ -725,7 +750,6 @@ public:
 			checkSlow(IsVertex(vID));
 			TDynamicVector<FVector3f>& Normals = VertexNormals.GetValue();
 			Normals[vID]                       = vNewNormal;
-			UpdateTimeStamp(false, false);
 		}
 	}
 
@@ -752,7 +776,6 @@ public:
 			checkSlow(IsVertex(vID));
 			TDynamicVector<FVector3f>& Colors = VertexColors.GetValue();
 			Colors[vID]                       = vNewColor;
-			UpdateTimeStamp(false, false);
 		}
 	}
 
@@ -777,7 +800,6 @@ public:
 			checkSlow(IsVertex(vID));
 			TDynamicVector<FVector2f>& UVs = VertexUVs.GetValue();
 			UVs[vID]                       = vNewUV;
-			UpdateTimeStamp(false, false);
 		}
 	}
 
@@ -802,7 +824,6 @@ public:
 			checkSlow(IsTriangle(tid));
 			TriangleGroups.GetValue()[tid] = group_id;
 			GroupIDCounter                 = FMath::Max(GroupIDCounter, group_id + 1);
-			UpdateTimeStamp(false, false);
 		}
 	}
 
@@ -938,30 +959,12 @@ public:
 	/** returns measure of compactness in range [0,1], where 1 is fully compacted */
 	double CompactMetric() const
 	{
-		return ((double)VertexCount() / (double)MaxVertexID() + (double)TriangleCount() / (double)MaxTriangleID()) *
-		    0.5;
+		return ((double)VertexCount() / (double)MaxVertexID() + (double)TriangleCount() / (double)MaxTriangleID()) * 0.5;
 	}
 
 	/** @return true if mesh has no boundary edges */
 	bool IsClosed() const;
 
-	/** Timestamp is incremented any time any change is made to the mesh */
-	inline int GetTimestamp() const
-	{
-		return Timestamp;
-	}
-
-	/** ShapeTimestamp is incremented any time any vertex position is changed or the mesh topology is modified */
-	inline int GetShapeTimestamp() const
-	{
-		return ShapeTimestamp;
-	}
-
-	/** TopologyTimestamp is incremented any time any vertex position is changed or the mesh topology is modified */
-	inline int GetTopologyTimestamp() const
-	{
-		return TopologyTimestamp;
-	}
 
 	//
 	// Geometric queries
@@ -1450,19 +1453,6 @@ protected:
 
 	void ReverseTriOrientationInternal(int TriangleID);
 
-	void UpdateTimeStamp(bool bShapeChange, bool bTopologyChange)
-	{
-		Timestamp++;
-		if (bShapeChange)
-		{
-			ShapeTimestamp++;
-		}
-		if (bTopologyChange)
-		{
-			checkSlow(bShapeChange); // we consider topology change to be a shape change!
-			TopologyTimestamp++;
-		}
-	}
 
 	/** Returns the internal version number used during serialization. This is used for testing purposes. */
 	static int SerializeInternal_LatestVersion();
diff --git a/Engine/Source/Runtime/GeometryCore/Public/Implicit/Blend.h b/Engine/Source/Runtime/GeometryCore/Public/Implicit/Blend.h
--- a/Engine/Source/Runtime/GeometryCore/Public/Implicit/Blend.h
+++ b/Engine/Source/Runtime/GeometryCore/Public/Implicit/Blend.h
@@ -124,7 +124,7 @@ protected:
 		ComputedSpatials.SetNum(Sources.Num());
 		ParallelFor(Sources.Num(), [this, bReuseComputed](int SourceIdx)
 			{
-				if (!bReuseComputed || ComputedSpatials[SourceIdx].GetMesh() != Sources[SourceIdx] || !ComputedSpatials[SourceIdx].IsValid())
+				if (!bReuseComputed || ComputedSpatials[SourceIdx].GetMesh() != Sources[SourceIdx] || !ComputedSpatials[SourceIdx].IsValid(false))
 				{
 					ComputedSpatials[SourceIdx].SetMesh(Sources[SourceIdx], true);
 				}
diff --git a/Engine/Source/Runtime/GeometryCore/Public/Implicit/Morphology.h b/Engine/Source/Runtime/GeometryCore/Public/Implicit/Morphology.h
--- a/Engine/Source/Runtime/GeometryCore/Public/Implicit/Morphology.h
+++ b/Engine/Source/Runtime/GeometryCore/Public/Implicit/Morphology.h
@@ -84,7 +84,7 @@ protected:
 public:
 	bool Validate()
 	{
-		bool bValidMeshAndSpatial = Source != nullptr && SourceSpatial != nullptr && SourceSpatial->IsValid();
+		bool bValidMeshAndSpatial = Source != nullptr && SourceSpatial != nullptr && SourceSpatial->IsValid(false);
 		bool bValidParams = Distance > 0 && GridCellSize > 0 && MeshCellSize > 0 && FMath::IsFinite(MeshCellSize);
 		return bValidMeshAndSpatial && bValidParams;
 	}
diff --git a/Engine/Source/Runtime/GeometryCore/Public/Implicit/Solidify.h b/Engine/Source/Runtime/GeometryCore/Public/Implicit/Solidify.h
--- a/Engine/Source/Runtime/GeometryCore/Public/Implicit/Solidify.h
+++ b/Engine/Source/Runtime/GeometryCore/Public/Implicit/Solidify.h
@@ -80,7 +80,7 @@ public:
 	 */
 	bool Validate()
 	{
-		bool bValidMeshAndSpatial = Source != nullptr && SourceSpatial != nullptr && SourceSpatial->IsValid();
+		bool bValidMeshAndSpatial = Source != nullptr && SourceSpatial != nullptr && SourceSpatial->IsValid(false);
 		bool bValidWinding = SourceWinding != nullptr;
 		bool bValidParams = SurfaceSearchSteps >= 0 && MeshCellSize > 0 && FMath::IsFinite(MeshCellSize);
 		return bValidMeshAndSpatial && bValidWinding && bValidParams;
diff --git a/Engine/Source/Runtime/GeometryCore/Public/MeshAdapter.h b/Engine/Source/Runtime/GeometryCore/Public/MeshAdapter.h
--- a/Engine/Source/Runtime/GeometryCore/Public/MeshAdapter.h
+++ b/Engine/Source/Runtime/GeometryCore/Public/MeshAdapter.h
@@ -26,7 +26,7 @@ struct TTriangleMeshAdapter
 	TFunction<int32()> MaxVertexID;
 	TFunction<int32()> TriangleCount;
 	TFunction<int32()> VertexCount;
-	TFunction<int32()> GetShapeTimestamp;
+	TFunction<uint64()> GetChangeStamp;
 	TFunction<FIndex3i(int32)> GetTriangle;
 	TFunction<FVector3<RealType>(int32)> GetVertex;
 
@@ -116,9 +116,9 @@ struct TIndexMeshArrayAdapter
 		return SourceVertices->Num();
 	}
 
-	FORCEINLINE int32 GetShapeTimestamp() const
+	FORCEINLINE uint64 GetChangeStamp() const
 	{
-		return 0; // source data doesn't have a timestamp concept
+		return 1; // source data doesn't have a timestamp concept
 	}
 
 	FORCEINLINE FIndex3i GetTriangle(int32 Index) const
@@ -198,9 +198,9 @@ struct TIndexVectorMeshArrayAdapter
 		return SourceVertices->Num();
 	}
 
-	FORCEINLINE int32 GetShapeTimestamp() const
+	FORCEINLINE uint64 GetChangeStamp() const
 	{
-		return 0; // source data doesn't have a timestamp concept
+		return 1; // source data doesn't have a timestamp concept
 	}
 
 	FORCEINLINE FIndex3i GetTriangle(int32 Index) const
diff --git a/Engine/Source/Runtime/GeometryCore/Public/PointSetAdapter.h b/Engine/Source/Runtime/GeometryCore/Public/PointSetAdapter.h
--- a/Engine/Source/Runtime/GeometryCore/Public/PointSetAdapter.h
+++ b/Engine/Source/Runtime/GeometryCore/Public/PointSetAdapter.h
@@ -24,8 +24,6 @@ struct TPointSetAdapter
 	TFunction<bool(int32)> IsPoint;
 	/** Get point at this index */
 	TFunction<FVector3<RealType>(int32)> GetPoint;
-	/** Returns a timestamp. If point set changes, timestamp should also change. */
-	TFunction<int32()> Timestamp;
 
 	/** Returns true if this point set has per-point normals */
 	TFunction<bool()> HasNormals;
diff --git a/Engine/Source/Runtime/GeometryCore/Public/Spatial/FastWinding.h b/Engine/Source/Runtime/GeometryCore/Public/Spatial/FastWinding.h
--- a/Engine/Source/Runtime/GeometryCore/Public/Spatial/FastWinding.h
+++ b/Engine/Source/Runtime/GeometryCore/Public/Spatial/FastWinding.h
@@ -316,7 +316,7 @@ class TFastWindingTree
 	};
 
 	TMap<int, FWNInfo> FastWindingCache;
-	int FastWindingCacheTimestamp = -1;
+	uint64 FastWindingCacheMeshChangeStamp = 0;
 
 public:
 	/**
@@ -348,23 +348,23 @@ public:
 		return Tree;
 	}
 
-	void Build(bool bAlwaysBuildRegardlessOfTimestamp = true)
+	void Build(bool bForceRebuild = true)
 	{
 		check(Tree);
-		if (Tree->MeshTimestamp != Tree->Mesh->GetShapeTimestamp())
+		if ( Tree->IsValid(false) == false )
 		{
 			Tree->Build();
 		}
-		if (bAlwaysBuildRegardlessOfTimestamp || FastWindingCacheTimestamp != Tree->MeshTimestamp)
+		if (bForceRebuild || FastWindingCacheMeshChangeStamp != Tree->MeshChangeStamp)
 		{
 			build_fast_winding_cache();
-			FastWindingCacheTimestamp = Tree->MeshTimestamp;
+			FastWindingCacheMeshChangeStamp = Tree->MeshChangeStamp;
 		}
 	}
 
 	bool IsBuilt() const
 	{
-		return Tree->MeshTimestamp == Tree->Mesh->GetShapeTimestamp() && FastWindingCacheTimestamp == Tree->MeshTimestamp;
+		return Tree->IsValid(false) && FastWindingCacheMeshChangeStamp == Tree->MeshChangeStamp;
 	}
 
 	/**
diff --git a/Engine/Source/Runtime/GeometryCore/Public/Spatial/MeshAABBTree3.h b/Engine/Source/Runtime/GeometryCore/Public/Spatial/MeshAABBTree3.h
--- a/Engine/Source/Runtime/GeometryCore/Public/Spatial/MeshAABBTree3.h
+++ b/Engine/Source/Runtime/GeometryCore/Public/Spatial/MeshAABBTree3.h
@@ -61,7 +61,7 @@ class TMeshAABBTree3 : public IMeshSpatial
 
 protected:
 	const TriangleMeshType* Mesh;
-	int MeshTimestamp = -1;
+	uint64 MeshChangeStamp = 0;
 	int TopDownLeafMaxTriCount = 4;
 
 public:
@@ -80,7 +80,7 @@ public:
 	void SetMesh(const TriangleMeshType* SourceMesh, bool bAutoBuild = true)
 	{
 		Mesh = SourceMesh;
-		MeshTimestamp = -1;
+		MeshChangeStamp = 0;
 
 		if (bAutoBuild)
 		{
@@ -94,15 +94,21 @@ public:
 	}
 
 	/**
-	 * @return true if internal timestamp matches mesh timestamp
+	 * @param bAllowUnsafeModifiedMeshQueries if true, then a built tree will still be considered valid even if the mesh change stamp has been modified
+	 * @return true if tree is built and mesh is unchanged (identified based on change stamp)
 	 */
-	bool IsValid() const 
+	bool IsValid(bool bAllowUnsafeModifiedMeshQueries) const 
 	{
-		if (MeshTimestamp == Mesh->GetShapeTimestamp())
+		checkSlow(RootIndex >= 0);
+		if (RootIndex < 0)
+		{
+			return false;
+		}
+		if (! bAllowUnsafeModifiedMeshQueries)
 		{
-			checkSlow(RootIndex >= 0);
+			return (MeshChangeStamp == Mesh->GetChangeStamp());
 		}
-		return MeshTimestamp == Mesh->GetShapeTimestamp();
+		return true;
 	}
 
 	void SetBuildOptions(int32 MaxBoxTriCount)
@@ -113,13 +119,13 @@ public:
 	void Build()
 	{
 		BuildTopDown(false);
-		MeshTimestamp = Mesh->GetShapeTimestamp();
+		MeshChangeStamp = Mesh->GetChangeStamp();
 	}
 
 	void Build(const TArray<int32>& TriangleList)
 	{
 		BuildTopDown(false, TriangleList, TriangleList.Num());
-		MeshTimestamp = Mesh->GetShapeTimestamp();
+		MeshChangeStamp = Mesh->GetChangeStamp();
 	}
 
 	virtual bool SupportsNearestTriangle() const override
@@ -136,8 +142,7 @@ public:
 		const FQueryOptions& Options = FQueryOptions()
 	) const override
 	{
-		checkSlow(RootIndex >= 0);
-		if (RootIndex < 0 || (!Options.bAllowUnsafeModifiedMeshQueries && !ensure(MeshTimestamp == Mesh->GetShapeTimestamp())))
+		if ( ensure(IsValid(Options.bAllowUnsafeModifiedMeshQueries)) == false )
 		{
 			return IndexConstants::InvalidID;
 		}
@@ -258,8 +263,7 @@ public:
 	virtual int FindNearestVertex(const FVector3d& P, double& NearestDistSqr,
 		double MaxDist = TNumericLimits<double>::Max(), const FQueryOptions& Options = FQueryOptions())
 	{
-		checkSlow(RootIndex >= 0);
-		if (RootIndex < 0 || (!Options.bAllowUnsafeModifiedMeshQueries && !ensure(MeshTimestamp == Mesh->GetShapeTimestamp())))
+		if ( ensure(IsValid(Options.bAllowUnsafeModifiedMeshQueries)) == false )
 		{
 			return IndexConstants::InvalidID;
 		}
@@ -373,13 +377,11 @@ public:
 		//   nearestT to TNumericLimits<double>::Max(), then we will test all boxes (!)
 		NearestT = (Options.MaxDistance < TNumericLimits<float>::Max()) ? Options.MaxDistance : TNumericLimits<float>::Max();
 
-		checkSlow(RootIndex >= 0);
-		if (RootIndex < 0 || (!Options.bAllowUnsafeModifiedMeshQueries && !ensure(MeshTimestamp == Mesh->GetShapeTimestamp())))
+		if ( ensure(IsValid(Options.bAllowUnsafeModifiedMeshQueries)) == false )
 		{
 			return false;
 		}
 
-
 		FindHitTriangle(RootIndex, Ray, NearestT, TID, Options);
 		return TID != IndexConstants::InvalidID;
 	}
@@ -468,10 +470,11 @@ public:
 	 */
 	virtual bool TestAnyHitTriangle(const FRay3d& Ray, const FQueryOptions& Options = FQueryOptions()) const
 	{
-		if (RootIndex < 0 || (!Options.bAllowUnsafeModifiedMeshQueries && !ensure(MeshTimestamp == Mesh->GetShapeTimestamp())))
+		if ( ensure(IsValid(Options.bAllowUnsafeModifiedMeshQueries)) == false )
 		{
 			return false;
 		}
+
 		double MaxDistance = (Options.MaxDistance < TNumericLimits<float>::Max()) ? Options.MaxDistance : TNumericLimits<float>::Max();
 		int32 HitTID = IndexConstants::InvalidID;
 		bool bFoundHit = TestAnyHitTriangle(RootIndex, Ray, HitTID, MaxDistance, Options);
@@ -584,8 +587,7 @@ public:
 		double& Distance, const FQueryOptions& Options = FQueryOptions(), const FQueryOptions& OtherTreeOptions = FQueryOptions()
 	)
 	{
-		checkSlow(RootIndex >= 0);
-		if (RootIndex < 0 || (!Options.bAllowUnsafeModifiedMeshQueries && !ensure(MeshTimestamp == Mesh->GetShapeTimestamp())))
+		if ( ensure(IsValid(Options.bAllowUnsafeModifiedMeshQueries)) == false )
 		{
 			return FIndex2i::Invalid();
 		}
@@ -633,8 +635,7 @@ public:
 	 */
 	virtual void DoTraversal(FTreeTraversal& Traversal, const FQueryOptions& Options = FQueryOptions()) const
 	{
-		checkSlow(RootIndex >= 0);
-		if (RootIndex < 0 || (!Options.bAllowUnsafeModifiedMeshQueries && !ensure(MeshTimestamp == Mesh->GetShapeTimestamp())))
+		if ( ensure(IsValid(Options.bAllowUnsafeModifiedMeshQueries)) == false )
 		{
 			return;
 		}
@@ -710,7 +711,7 @@ public:
 		const FQueryOptions& Options = FQueryOptions()
 	) const
 	{
-		if (!Options.bAllowUnsafeModifiedMeshQueries && !ensure(MeshTimestamp == Mesh->GetShapeTimestamp()))
+		if ( ensure(IsValid(Options.bAllowUnsafeModifiedMeshQueries)) == false )
 		{
 			return false;
 		}
@@ -761,7 +762,7 @@ public:
 		const FQueryOptions& Options = FQueryOptions(), const FQueryOptions& OtherTreeOptions = FQueryOptions()
 	) const
 	{
-		if (!Options.bAllowUnsafeModifiedMeshQueries && !ensure(MeshTimestamp == Mesh->GetShapeTimestamp()))
+		if ( ensure(IsValid(Options.bAllowUnsafeModifiedMeshQueries)) == false )
 		{
 			return false;
 		}
@@ -786,7 +787,7 @@ public:
 		const FQueryOptions& Options = FQueryOptions()
 	) const
 	{
-		if (!Options.bAllowUnsafeModifiedMeshQueries && !ensure(MeshTimestamp == Mesh->GetShapeTimestamp()))
+		if ( ensure(IsValid(Options.bAllowUnsafeModifiedMeshQueries)) == false )
 		{
 			return false;
 		}
@@ -824,7 +825,7 @@ public:
 
 		MeshIntersection::FIntersectionsQueryResult result;
 
-		if (!Options.bAllowUnsafeModifiedMeshQueries && !ensure(MeshTimestamp == Mesh->GetShapeTimestamp()))
+		if ( ensure(IsValid(Options.bAllowUnsafeModifiedMeshQueries)) == false )
 		{
 			return result;
 		}
@@ -858,7 +859,7 @@ public:
 
 		MeshIntersection::FIntersectionsQueryResult Result;
 
-		if (!Options.bAllowUnsafeModifiedMeshQueries && !ensure(MeshTimestamp == Mesh->GetShapeTimestamp()))
+		if ( ensure(IsValid(Options.bAllowUnsafeModifiedMeshQueries)) == false )
 		{
 			return Result;
 		}
@@ -870,7 +871,7 @@ public:
 
 	virtual bool TestSelfIntersection(bool bIgnoreTopoConnected = true, const FQueryOptions& Options = FQueryOptions()) const
 	{
-		if (!Options.bAllowUnsafeModifiedMeshQueries && !ensure(MeshTimestamp == Mesh->GetShapeTimestamp()))
+		if ( ensure(IsValid(Options.bAllowUnsafeModifiedMeshQueries)) == false )
 		{
 			return false;
 		}
diff --git a/Engine/Source/Runtime/MeshConversion/Public/MeshDescriptionAdapter.h b/Engine/Source/Runtime/MeshConversion/Public/MeshDescriptionAdapter.h
--- a/Engine/Source/Runtime/MeshConversion/Public/MeshDescriptionAdapter.h
+++ b/Engine/Source/Runtime/MeshConversion/Public/MeshDescriptionAdapter.h
@@ -69,11 +69,11 @@ public:
 	{
 		return Mesh->Vertices().Num();
 	}
-	int32 GetShapeTimestamp() const
+	uint64 GetChangeStamp() const
 	{
-		// MeshDescription doesn't provide any mechanism to know if it's been modified so just return 0
+		// MeshDescription doesn't provide any mechanism to know if it's been modified so just return 1
 		// and leave it to the caller to not build an aabb and then change the underlying mesh
-		return 0;
+		return 1;
 	}
 	FIndex3i GetTriangle(int32 IDValue) const
 	{
@@ -179,11 +179,11 @@ public:
 	{
 		return Mesh->Vertices().Num();
 	}
-	int32 GetShapeTimestamp() const
+	uint64 GetChangeStamp() const
 	{
-		// MeshDescription doesn't provide any mechanism to know if it's been modified so just return 0
+		// MeshDescription doesn't provide any mechanism to know if it's been modified so just return 1
 		// and leave it to the caller to not build an aabb and then change the underlying mesh
-		return 0;
+		return 1;
 	}
 	FIndex3i GetTriangle(int32 IDValue) const
 	{
@@ -251,7 +251,7 @@ struct FMeshDescriptionMeshAdapterd : public UE::Geometry::TTriangleMeshAdapter<
 		MaxVertexID = [&]() { return ParentAdapter.MaxVertexID();};
 		TriangleCount = [&]() { return ParentAdapter.TriangleCount();};
 		VertexCount = [&]() { return ParentAdapter.VertexCount();};
-		GetShapeTimestamp = [&]() { return ParentAdapter.GetShapeTimestamp();};
+		GetChangeStamp = [&]() { return ParentAdapter.GetChangeStamp();};
 		GetTriangle = [&](int32 TriangleID) { return ParentAdapter.GetTriangle(TriangleID); };
 		GetVertex = [&](int32 VertexID) { return ParentAdapter.GetVertex(VertexID); };
 	}
@@ -264,7 +264,7 @@ struct FMeshDescriptionMeshAdapterd : public UE::Geometry::TTriangleMeshAdapter<
 		MaxVertexID = [&]() { return ParentAdapter.MaxVertexID();};
 		TriangleCount = [&]() { return ParentAdapter.TriangleCount();};
 		VertexCount = [&]() { return ParentAdapter.VertexCount();};
-		GetShapeTimestamp = [&]() { return ParentAdapter.GetShapeTimestamp();};
+		GetChangeStamp = [&]() { return ParentAdapter.GetChangeStamp();};
 		GetTriangle = [&](int32 TriangleID) { return ParentAdapter.GetTriangle(TriangleID); };
 		GetVertex = [&](int32 VertexID) { return ParentAdapter.GetVertex(VertexID); };
 	}
