GeometryProcessing: clean up TVector4 and usage in preparation for LWC conversion
#rnx

diff --git a/Engine/Plugins/Experimental/MeshModelingToolsetExp/Source/MeshModelingToolsExp/Private/MeshAttributePaintTool.cpp b/Engine/Plugins/Experimental/MeshModelingToolsetExp/Source/MeshModelingToolsExp/Private/MeshAttributePaintTool.cpp
--- a/Engine/Plugins/Experimental/MeshModelingToolsetExp/Source/MeshModelingToolsExp/Private/MeshAttributePaintTool.cpp
+++ b/Engine/Plugins/Experimental/MeshModelingToolsetExp/Source/MeshModelingToolsExp/Private/MeshAttributePaintTool.cpp
@@ -422,7 +422,7 @@ void UMeshAttributePaintTool::UpdateVisibleAttribute()
 			{
 				const int32 vid = ColorOverlay->GetParentVertex(elid);
 				const float Value = AttribData.CurrentValues[vid];
-				const FVector4f Color4f(ColorMapper->ToColor(Value));
+				const FVector4f Color4f = ToVector4<float>(ColorMapper->ToColor(Value));
 				ColorOverlay->SetElement(elid,  Color4f);
 			}
 		});
@@ -495,7 +495,7 @@ void UMeshAttributePaintTool::ApplyStamp(const FBrushStampData& Stamp)
 		{
 			int32 vid = ActionData.ROIVertices[k];
 			AttribData.CurrentValues[vid] = ActionData.ROIAfter[k];
-			FVector4f NewColor(ColorMapper->ToColor(ActionData.ROIAfter[k]));
+			FVector4f NewColor( ToVector4<float>(ColorMapper->ToColor(ActionData.ROIAfter[k])) );
 			ColorOverlay->GetVertexElements(vid, ElIDs);
 			for (int elid : ElIDs)
 			{
@@ -713,7 +713,7 @@ void UMeshAttributePaintTool::ExternalUpdateValues(int32 AttribIndex, const TArr
 			FDynamicMeshColorOverlay* ColorOverlay = Mesh.Attributes()->PrimaryColors();
 			for (int32 vid : VertexIndices)
 			{
-				FVector4f NewColor(ColorMapper->ToColor(AttribData.CurrentValues[vid]));
+				FVector4f NewColor( ToVector4<float>(ColorMapper->ToColor(AttribData.CurrentValues[vid])) );
 				ColorOverlay->GetVertexElements(vid, ElIDs);
 				for (int elid : ElIDs)
 				{
diff --git a/Engine/Plugins/Experimental/MeshModelingToolsetExp/Source/MeshModelingToolsExp/Private/SkinWeightsBindingTool.cpp b/Engine/Plugins/Experimental/MeshModelingToolsetExp/Source/MeshModelingToolsExp/Private/SkinWeightsBindingTool.cpp
--- a/Engine/Plugins/Experimental/MeshModelingToolsetExp/Source/MeshModelingToolsExp/Private/SkinWeightsBindingTool.cpp
+++ b/Engine/Plugins/Experimental/MeshModelingToolsetExp/Source/MeshModelingToolsExp/Private/SkinWeightsBindingTool.cpp
@@ -1037,7 +1037,7 @@ FVector4f USkinWeightsBindingTool::WeightToColor(float Value)
 	{
 		// A close approximation of the skeletal mesh editor's bone weight ramp. 
 		const FLinearColor HSV((1.0f - Value) * 285.0f, 100.0f, 85.0f);
-		return FVector4f(HSV.HSVToLinearRGB());
+		return UE::Geometry::ToVector4<float>(HSV.HSVToLinearRGB());
 	}
 }
 
diff --git a/Engine/Plugins/Experimental/MeshModelingToolsetExp/Source/MeshModelingToolsExp/Private/SkinWeightsPaintTool.cpp b/Engine/Plugins/Experimental/MeshModelingToolsetExp/Source/MeshModelingToolsExp/Private/SkinWeightsPaintTool.cpp
--- a/Engine/Plugins/Experimental/MeshModelingToolsetExp/Source/MeshModelingToolsExp/Private/SkinWeightsPaintTool.cpp
+++ b/Engine/Plugins/Experimental/MeshModelingToolsetExp/Source/MeshModelingToolsExp/Private/SkinWeightsPaintTool.cpp
@@ -299,7 +299,7 @@ FVector4f USkinWeightsPaintTool::WeightToColor(float Value)
 	{
 		// A close approximation of the skeletal mesh editor's bone weight ramp. 
 		const FLinearColor HSV((1.0f - Value) * 285.0f, 100.0f, 85.0f);
-		return FVector4f(HSV.HSVToLinearRGB());
+		return UE::Geometry::ToVector4<float>(HSV.HSVToLinearRGB());
 	}
 }
 
diff --git a/Engine/Plugins/Experimental/MeshModelingToolsetExp/Source/ModelingComponents/Private/AssetUtils/Texture2DBuilder.cpp b/Engine/Plugins/Experimental/MeshModelingToolsetExp/Source/ModelingComponents/Private/AssetUtils/Texture2DBuilder.cpp
--- a/Engine/Plugins/Experimental/MeshModelingToolsetExp/Source/ModelingComponents/Private/AssetUtils/Texture2DBuilder.cpp
+++ b/Engine/Plugins/Experimental/MeshModelingToolsetExp/Source/ModelingComponents/Private/AssetUtils/Texture2DBuilder.cpp
@@ -336,12 +336,12 @@ bool FTexture2DBuilder::Copy(const TImageBuilder<FVector4f>& SourceImage, const
 			Pixel.Y = FMathf::Clamp(Pixel.Y, 0.0, 1.0);
 			Pixel.Z = FMathf::Clamp(Pixel.Z, 0.0, 1.0);
 			Pixel.W = FMathf::Clamp(Pixel.W, 0.0, 1.0);
-			FColor Texel = ((FLinearColor)Pixel).ToFColor(bConvertToSRGB);
+			FColor Texel = ToLinearColor(Pixel).ToFColor(bConvertToSRGB);
 			SetTexel(i, Texel);
 		}
 		else
 		{
-			FFloat16Color Texel((FLinearColor)Pixel);
+			FFloat16Color Texel(ToLinearColor(Pixel));
 			SetTexel(i, Texel);
 		}
 	}
@@ -362,13 +362,13 @@ bool FTexture2DBuilder::CopyTo(TImageBuilder<FVector4f>& DestImage) const
 		{
 			FColor ByteColor = GetTexel(i);
 			FLinearColor FloatColor(ByteColor);
-			DestImage.SetPixel(i, FVector4f(FloatColor));
+			DestImage.SetPixel(i, ToVector4<float>(FloatColor));
 		}
 		else
 		{
 			FFloat16Color Float16Color = GetTexelFloat16(i);
 			FLinearColor FloatColor = Float16Color.GetFloats();
-			DestImage.SetPixel(i, FVector4f(FloatColor));
+			DestImage.SetPixel(i, ToVector4<float>(FloatColor));
 		}
 	}
 	return true;
diff --git a/Engine/Plugins/Experimental/MeshModelingToolsetExp/Source/ModelingComponents/Private/AssetUtils/Texture2DUtil.cpp b/Engine/Plugins/Experimental/MeshModelingToolsetExp/Source/ModelingComponents/Private/AssetUtils/Texture2DUtil.cpp
--- a/Engine/Plugins/Experimental/MeshModelingToolsetExp/Source/ModelingComponents/Private/AssetUtils/Texture2DUtil.cpp
+++ b/Engine/Plugins/Experimental/MeshModelingToolsetExp/Source/ModelingComponents/Private/AssetUtils/Texture2DUtil.cpp
@@ -51,7 +51,7 @@ static bool ReadTexture_PlatformData(
 		FLinearColor FloatColor = (TextureMap->SRGB) ?
 				FLinearColor::FromSRGBColor(ByteColor) :
 				ByteColor.ReinterpretAsLinear();
-		DestImage.SetPixel(i, FVector4f(FloatColor));
+		DestImage.SetPixel(i, ToVector4<float>(FloatColor));
 	}
 
 	// restore built platform texture data to initial state
@@ -98,7 +98,7 @@ static bool ReadTexture_SourceData(
 				FLinearColor::FromSRGBColor(PixelColor) :
 				PixelColor.ReinterpretAsLinear();
 
-			DestImage.SetPixel(i, FVector4f(FloatColor));
+			DestImage.SetPixel(i, ToVector4<float>(FloatColor));
 		}
 	}
 	// code below is also derived from FImage::CopyTo (CopyImage) - invoked during BuildTexture on import.
@@ -123,7 +123,8 @@ static bool ReadTexture_SourceData(
 		for (int64 i = 0; i < Num; ++i)
 		{
 			const uint8* PixelPtr = SourceDataPtr + (i * BytesPerPixel);
-			DestImage.SetPixel(i, FVector4f( ((const FFloat16Color*)PixelPtr)->GetFloats() ));
+			FLinearColor LinearColor = ((const FFloat16Color*)PixelPtr)->GetFloats();
+			DestImage.SetPixel(i, ToVector4<float>(LinearColor) );
 		}
 	}
 	else if ((SourceFormat == TSF_G16))
@@ -148,7 +149,7 @@ static bool ReadTexture_SourceData(
 			FLinearColor FloatColor = (TextureMap->SRGB) ?
 				FLinearColor::FromSRGBColor(FColor(PixelColor, PixelColor, PixelColor, 255)) :
 				FLinearColor(PixelColorf, PixelColorf, PixelColorf, 1.0);
-			DestImage.SetPixel(i, FVector4f(FloatColor));
+			DestImage.SetPixel(i, ToVector4<float>(FloatColor));
 		}
 	}
 
@@ -358,7 +359,7 @@ bool UE::AssetUtils::SaveDebugImage(
 	ConvertedColor.Reserve(N);
 	for ( int64 i = 0; i < N; ++i )
 	{
-		FLinearColor LinearColor = (FLinearColor)Image.GetPixel(i);
+		FLinearColor LinearColor = ToLinearColor(Image.GetPixel(i));
 		ConvertedColor.Add(LinearColor.ToFColor(bConvertToSRGB));
 	}
 
diff --git a/Engine/Plugins/Experimental/MeshModelingToolsetExp/Source/ModelingComponents/Private/Scene/WorldRenderCapture.cpp b/Engine/Plugins/Experimental/MeshModelingToolsetExp/Source/ModelingComponents/Private/Scene/WorldRenderCapture.cpp
--- a/Engine/Plugins/Experimental/MeshModelingToolsetExp/Source/ModelingComponents/Private/Scene/WorldRenderCapture.cpp
+++ b/Engine/Plugins/Experimental/MeshModelingToolsetExp/Source/ModelingComponents/Private/Scene/WorldRenderCapture.cpp
@@ -275,7 +275,7 @@ static bool ReadPixelDataToImage(TUniquePtr<FImagePixelData>& PixelData, FImageA
 					const uint8* BufferPixel = &SourceBuffer[(yi * Width + xi) * 4];
 					FColor Color(BufferPixel[2], BufferPixel[1], BufferPixel[0], BufferPixel[3]);		// BGRA
 					FLinearColor PixelColorf = (bLinear) ? Color.ReinterpretAsLinear() : FLinearColor(Color);
-					ResultImageOut.SetPixel(FVector2i(xi, yi), FVector4f(PixelColorf));
+					ResultImageOut.SetPixel(FVector2i(xi, yi), PixelColorf);
 				}
 			}
 			return true;
@@ -291,7 +291,7 @@ static bool ReadPixelDataToImage(TUniquePtr<FImagePixelData>& PixelData, FImageA
 				{
 					const FFloat16Color* BufferPixel = &SourceBuffer[yi * Width + xi];
 					FLinearColor PixelColorf = BufferPixel->GetFloats();
-					ResultImageOut.SetPixel(FVector2i(xi, yi), FVector4f(PixelColorf));
+					ResultImageOut.SetPixel(FVector2i(xi, yi), PixelColorf);
 				}
 			}
 			return true;
@@ -306,7 +306,7 @@ static bool ReadPixelDataToImage(TUniquePtr<FImagePixelData>& PixelData, FImageA
 				for (int32 xi = 0; xi < Width; ++xi)
 				{
 					FLinearColor PixelColorf = SourceBuffer[yi * Width + xi];
-					ResultImageOut.SetPixel(FVector2i(xi, yi), FVector4f(PixelColorf));
+					ResultImageOut.SetPixel(FVector2i(xi, yi), PixelColorf);
 				}
 			}
 			return true;
@@ -591,7 +591,7 @@ bool FWorldRenderCapture::CaptureEmissiveFromPosition(
 		for (int32 xi = 0; xi < Width; ++xi)
 		{
 			FLinearColor PixelColorf = ReadImageBuffer[yi * Width + xi];
-			ResultImageOut.SetPixel(FVector2i(xi, yi), FVector4f(PixelColorf));
+			ResultImageOut.SetPixel(FVector2i(xi, yi), PixelColorf);
 		}
 	}
 
@@ -753,7 +753,7 @@ bool FWorldRenderCapture::CaptureFromPosition(
 		{
 			FLinearColor PixelColorf = ReadImageBuffer[yi * Width + xi];
 			PixelColorf.A = 1.0f;		// ?
-			ResultImageOut.SetPixel(FVector2i(xi, yi), FVector4f(PixelColorf));
+			ResultImageOut.SetPixel(FVector2i(xi, yi), PixelColorf);
 		}
 	}
 
diff --git a/Engine/Source/Runtime/GeometryCore/Public/Image/ImageBuilder.h b/Engine/Source/Runtime/GeometryCore/Public/Image/ImageBuilder.h
--- a/Engine/Source/Runtime/GeometryCore/Public/Image/ImageBuilder.h
+++ b/Engine/Source/Runtime/GeometryCore/Public/Image/ImageBuilder.h
@@ -349,6 +349,22 @@ public:
 		}
 	}
 
+	void SetPixel(const FVector2i& PixelCoords, const FLinearColor& FloatPixel)
+	{
+		switch (ImageType)
+		{
+		case EImageType::Float1:
+			Image1f->SetPixel(PixelCoords, FloatPixel.R);
+			break;
+		case EImageType::Float3:
+			Image3f->SetPixel(PixelCoords, FVector3f(FloatPixel.R, FloatPixel.G, FloatPixel.B));
+			break;
+		case EImageType::Float4:
+			Image4f->SetPixel(PixelCoords, FVector4f(FloatPixel.R, FloatPixel.G, FloatPixel.B, FloatPixel.A));
+			break;
+		}
+	}
+
 	FVector4f GetPixel(int64 LinearIndex) const
 	{
 		switch (ImageType)
diff --git a/Engine/Source/Runtime/GeometryCore/Public/VectorTypes.h b/Engine/Source/Runtime/GeometryCore/Public/VectorTypes.h
--- a/Engine/Source/Runtime/GeometryCore/Public/VectorTypes.h
+++ b/Engine/Source/Runtime/GeometryCore/Public/VectorTypes.h
@@ -905,11 +905,6 @@ struct TVector4
 	{
 	}
 
-	constexpr TVector4(const T* Data)
-		: X(Data[0]), Y(Data[1]), Z(Data[2]), W(Data[3])
-	{
-	}
-
 	constexpr TVector4(const TVector4& Vec) = default;
 
 	template<typename RealType2>
@@ -918,34 +913,6 @@ struct TVector4
 	{
 	}
 
-	explicit constexpr operator const T*() const
-	{
-		return &X;
-	};
-	explicit constexpr operator T*()
-	{
-		return &X;
-	}
-
-	explicit constexpr operator FLinearColor() const
-	{
-		return FLinearColor((float)X, (float)Y, (float)Z, (float)W);
-	}
-	constexpr TVector4(const FLinearColor& Color)
-		: X((T)Color.R), Y((T)Color.G), Z((T)Color.B), W((T)Color.A)
-	{
-	}
-
-	explicit constexpr operator FColor() const
-	{
-		return FColor(
-			FMathf::Clamp((int)((float)X*255.0f), 0, 255),
-			FMathf::Clamp((int)((float)Y*255.0f), 0, 255),
-			FMathf::Clamp((int)((float)Z*255.0f), 0, 255),
-			FMathf::Clamp((int)((float)W*255.0f), 0, 255));
-	}
-
-
 	static TVector4<T> Zero()
 	{
 		return TVector4<T>((T)0, (T)0, (T)0, (T)0);
@@ -974,16 +941,6 @@ struct TVector4
 		return (&X)[Idx];
 	}
 
-	T Length() const
-	{
-		return TMathUtil<T>::Sqrt(X * X + Y * Y + Z * Z + W * W);
-	}
-	T SquaredLength() const
-	{
-		return X * X + Y * Y + Z * Z + W * W;
-	}
-
-
 	constexpr TVector4<T> operator-() const
 	{
 		return TVector4<T>(-X, -Y, -Z, -W);
@@ -1020,21 +977,11 @@ struct TVector4
 		return TVector4<T>(X * (T)Scalar, Y * (T)Scalar, Z * (T)Scalar, W * (T)Scalar);
 	}
 
-	constexpr TVector4<T> operator*(const TVector4<T>& V2) const // component-wise
-	{
-		return TVector4<T>(X * V2.X, Y * V2.Y, Z * V2.Z, W * V2.W);
-	}
-
 	constexpr TVector4<T> operator/(const T& Scalar) const
 	{
 		return TVector4<T>(X / Scalar, Y / Scalar, Z / Scalar, W / Scalar);
 	}
 
-	constexpr TVector4<T> operator/(const TVector4<T>& V2) const // component-wise
-	{
-		return TVector4<T>(X / V2.X, Y / V2.Y, Z / V2.Z, W / V2.W);
-	}
-
 	constexpr TVector4<T>& operator+=(const TVector4<T>& V2)
 	{
 		X += V2.X;
@@ -1183,6 +1130,28 @@ inline TVector4<RealType> operator*(RealType2 Scalar, const TVector4<RealType>&
 	return TVector4<RealType>((RealType)Scalar * V.X, (RealType)Scalar * V.Y, (RealType)Scalar * V.Z, (RealType)Scalar * V.W);
 }
 
+template<typename T>
+FLinearColor ToLinearColor(const TVector4<T>& V)
+{
+	return FLinearColor((float)V.X, (float)V.Y, (float)V.Z, (float)V.W);
+}
+
+template<typename T>
+TVector4<T> ToVector4(const FLinearColor& Color)
+{
+	return TVector4<T>( (T)Color.R, (T)Color.G, (T)Color.B, (T)Color.A );
+}
+
+template<typename T>
+FColor ToFColor(const TVector4<T>& V)
+{
+	return FColor(
+		FMathf::Clamp((int)((float)V.X*255.0f), 0, 255),
+		FMathf::Clamp((int)((float)V.Y*255.0f), 0, 255),
+		FMathf::Clamp((int)((float)V.Z*255.0f), 0, 255),
+		FMathf::Clamp((int)((float)V.W*255.0f), 0, 255));
+}
+
 template <typename RealType>
 std::ostream& operator<<(std::ostream& os, const TVector4<RealType>& Vec)
 {
diff --git a/Engine/Source/Runtime/GeometryFramework/Public/Components/BaseDynamicMeshSceneProxy.h b/Engine/Source/Runtime/GeometryFramework/Public/Components/BaseDynamicMeshSceneProxy.h
--- a/Engine/Source/Runtime/GeometryFramework/Public/Components/BaseDynamicMeshSceneProxy.h
+++ b/Engine/Source/Runtime/GeometryFramework/Public/Components/BaseDynamicMeshSceneProxy.h
@@ -512,7 +512,7 @@ public:
 				}
 
 				FColor VertexFColor = (bHaveColors && TriColor[j] != FDynamicMesh3::InvalidID) ?
-					FColor(ColorOverlay->GetElement(TriColor[j])) : UniformTriColor;
+					ToFColor(ColorOverlay->GetElement(TriColor[j])) : UniformTriColor;
 
 				RenderBuffers->ColorVertexBuffer.VertexColor(VertIdx) = VertexFColor;
 
@@ -697,7 +697,7 @@ public:
 				if (bUpdateColors)
 				{
 					FColor VertexFColor = (bHaveColors  && TriColor[j] != FDynamicMesh3::InvalidID) ?
-						FColor(ColorOverlay->GetElement(TriColor[j])) : UniformTriColor;
+						ToFColor(ColorOverlay->GetElement(TriColor[j])) : UniformTriColor;
 					RenderBuffers->ColorVertexBuffer.VertexColor(VertIdx) = VertexFColor;
 				}
 
