GeometryCore: Add minimal UE::Math::TVector2<T> implementation and FVector2f/2d typedefs locally in GeometryCore VectorTypes.h. Replace existing FVector2<T> with this class, in preparation for LWC conversion. FVector2 members ::Distance() and ::DistanceSquared() moved to free functions. Update all call sites and references to FVector2.
#rnx

diff --git a/Engine/Plugins/Experimental/MeshModelingToolsetExp/Source/MeshModelingToolsExp/Private/MeshGroupPaintTool.cpp b/Engine/Plugins/Experimental/MeshModelingToolsetExp/Source/MeshModelingToolsExp/Private/MeshGroupPaintTool.cpp
--- a/Engine/Plugins/Experimental/MeshModelingToolsetExp/Source/MeshModelingToolsExp/Private/MeshGroupPaintTool.cpp
+++ b/Engine/Plugins/Experimental/MeshModelingToolsetExp/Source/MeshModelingToolsExp/Private/MeshGroupPaintTool.cpp
@@ -600,8 +600,8 @@ bool UMeshGroupPaintTool::SyncMeshWithGroupBuffer(FDynamicMesh3* Mesh)
 
 template<typename RealType>
 static bool FindPolylineSelfIntersection(
-	const TArray<FVector2<RealType>>& Polyline, 
-	FVector2<RealType>& IntersectionPointOut, 
+	const TArray<UE::Math::TVector2<RealType>>& Polyline, 
+	UE::Math::TVector2<RealType>& IntersectionPointOut, 
 	FIndex2i& IntersectionIndexOut,
 	bool bParallel = true)
 {
@@ -635,9 +635,9 @@ static bool FindPolylineSelfIntersection(
 
 template<typename RealType>
 static bool FindPolylineSegmentIntersection(
-	const TArray<FVector2<RealType>>& Polyline,
+	const TArray<UE::Math::TVector2<RealType>>& Polyline,
 	const TSegment2<RealType>& Segment,
-	FVector2<RealType>& IntersectionPointOut,
+	UE::Math::TVector2<RealType>& IntersectionPointOut,
 	int& IntersectionIndexOut)
 {
 
@@ -664,7 +664,7 @@ bool ApproxSelfClipPolyline(TArray<FVector2f>& Polyline)
 	int32 N = Polyline.Num();
 
 	// handle already-closed polylines
-	if (Polyline[0].Distance(Polyline[N-1]) < 0.0001f)
+	if (Distance(Polyline[0], Polyline[N-1]) < 0.0001f)
 	{
 		return true;
 	}
@@ -711,7 +711,7 @@ bool ApproxSelfClipPolyline(TArray<FVector2f>& Polyline)
 	}
 	float Size = Bounds.DiagonalLength();
 
-	FVector2f StartPos = Polyline[0] + 0.001 * StartDirOut;
+	FVector2f StartPos = Polyline[0] + 0.001f * StartDirOut;
 	if (FindPolylineSegmentIntersection(Polyline, FSegment2f(StartPos, StartPos + 2*Size*StartDirOut), IntersectPoint, IntersectionIndex.A))
 	{
 		//TArray<FVector2f> NewPolyline;
@@ -725,7 +725,7 @@ bool ApproxSelfClipPolyline(TArray<FVector2f>& Polyline)
 		return true;
 	}
 
-	FVector2f EndPos = Polyline[N-1] + 0.001 * EndDirOut;
+	FVector2f EndPos = Polyline[N-1] + 0.001f * EndDirOut;
 	if (FindPolylineSegmentIntersection(Polyline, FSegment2f(EndPos, EndPos + 2*Size*EndDirOut), IntersectPoint, IntersectionIndex.A))
 	{
 		//TArray<FVector2f> NewPolyline;
diff --git a/Engine/Plugins/Experimental/MeshModelingToolsetExp/Source/ModelingOperators/Private/CleaningOps/HoleFillOp.cpp b/Engine/Plugins/Experimental/MeshModelingToolsetExp/Source/ModelingOperators/Private/CleaningOps/HoleFillOp.cpp
--- a/Engine/Plugins/Experimental/MeshModelingToolsetExp/Source/ModelingOperators/Private/CleaningOps/HoleFillOp.cpp
+++ b/Engine/Plugins/Experimental/MeshModelingToolsetExp/Source/ModelingOperators/Private/CleaningOps/HoleFillOp.cpp
@@ -161,7 +161,7 @@ namespace
 		FIndex3i TriangleVertices = Mesh->GetTriangle(NewTriangleID);
 		FVector2f NewElement = 0.5f * (EdgeElements[0] + EdgeElements[1]);
 		
-		bool bIsDegen = (EdgeElements[0].Distance(EdgeElements[1]) == 0.0f);
+		bool bIsDegen = ( Distance(EdgeElements[0], EdgeElements[1]) == 0.0f);
 		if (!bIsDegen)
 		{
 			// push the new element off the edge slightly to avoid degenerate triangle (assuming the adjacent triangle 
diff --git a/Engine/Plugins/Experimental/MeshModelingToolsetExp/Source/ModelingOperatorsEditorOnly/Private/CuttingOps/EmbedPolygonsOp.cpp b/Engine/Plugins/Experimental/MeshModelingToolsetExp/Source/ModelingOperatorsEditorOnly/Private/CuttingOps/EmbedPolygonsOp.cpp
--- a/Engine/Plugins/Experimental/MeshModelingToolsetExp/Source/ModelingOperatorsEditorOnly/Private/CuttingOps/EmbedPolygonsOp.cpp
+++ b/Engine/Plugins/Experimental/MeshModelingToolsetExp/Source/ModelingOperatorsEditorOnly/Private/CuttingOps/EmbedPolygonsOp.cpp
@@ -473,7 +473,7 @@ void FEmbedPolygonsOp::CalculateResult(FProgressCancel* Progress)
 			for (int UIdx = 0; UIdx < UValues.Num(); UIdx++)
 			{
 				UValues[UIdx] = Along;
-				Along += Polygon[UIdx % Polygon.VertexCount()].Distance(Polygon[(UIdx + 1) % Polygon.VertexCount()]);
+				Along += Distance(Polygon[UIdx % Polygon.VertexCount()], Polygon[(UIdx + 1) % Polygon.VertexCount()]);
 			}
 
 			for (int UVIdx = 0, NumUVLayers = ResultMesh->Attributes()->NumUVLayers(); UVIdx < NumUVLayers; UVIdx++)
diff --git a/Engine/Plugins/Experimental/Text3D/Source/Text3D/Private/Arrangement.cpp b/Engine/Plugins/Experimental/Text3D/Source/Text3D/Private/Arrangement.cpp
--- a/Engine/Plugins/Experimental/Text3D/Source/Text3D/Private/Arrangement.cpp
+++ b/Engine/Plugins/Experimental/Text3D/Source/Text3D/Private/Arrangement.cpp
@@ -81,7 +81,7 @@ bool FArrangement::insert_segment(FVector2d A, FVector2d B, double Tol)
     }
 
     // handle tiny-segment case
-    double SegLenSq = A.DistanceSquared(B);
+    double SegLenSq = DistanceSquared(A, B);
     if (SegLenSq <= VertexSnapTol*VertexSnapTol)
     {
         // seg is too short and was already on an existing vertex; just consider that vertex to be the inserted segment
@@ -247,7 +247,7 @@ int FArrangement::find_existing_vertex(FVector2d Pt)
 
 int FArrangement::find_nearest_vertex(FVector2d Pt, double SearchRadius, int IgnoreVID)
 {
-    auto FuncDistSq = [&](int B) { return Pt.DistanceSquared(Graph.GetVertex(B)); };
+    auto FuncDistSq = [&](int B) { return DistanceSquared(Pt, Graph.GetVertex(B)); };
     auto FuncIgnore = [&](int VID) { return VID == IgnoreVID; };
     TPair<int, double> found = (IgnoreVID == -1) ? PointHash.FindNearestInRadius(Pt, SearchRadius, FuncDistSq)
                                                  : PointHash.FindNearestInRadius(Pt, SearchRadius, FuncDistSq, FuncIgnore);
diff --git a/Engine/Plugins/Runtime/GeometryProcessing/Source/DynamicMesh/Private/DynamicMeshEditor.cpp b/Engine/Plugins/Runtime/GeometryProcessing/Source/DynamicMesh/Private/DynamicMeshEditor.cpp
--- a/Engine/Plugins/Runtime/GeometryProcessing/Source/DynamicMesh/Private/DynamicMeshEditor.cpp
+++ b/Engine/Plugins/Runtime/GeometryProcessing/Source/DynamicMesh/Private/DynamicMeshEditor.cpp
@@ -1182,7 +1182,7 @@ void FDynamicMeshEditor::RescaleAttributeUVs(float UVScale, bool bWorldSpace, in
 			}
 			for (int j = 2, i = 0; i < 3; j = i++)
 			{
-				TotalEdgeUVLen += TriUVs[j].Distance(TriUVs[i]);
+				TotalEdgeUVLen += Distance(TriUVs[j], TriUVs[i]);
 				TotalEdgeLen += Distance(TriVs[j], TriVs[i]);
 			}
 		}
diff --git a/Engine/Plugins/Runtime/GeometryProcessing/Source/DynamicMesh/Private/Operations/JoinMeshLoops.cpp b/Engine/Plugins/Runtime/GeometryProcessing/Source/DynamicMesh/Private/Operations/JoinMeshLoops.cpp
--- a/Engine/Plugins/Runtime/GeometryProcessing/Source/DynamicMesh/Private/Operations/JoinMeshLoops.cpp
+++ b/Engine/Plugins/Runtime/GeometryProcessing/Source/DynamicMesh/Private/Operations/JoinMeshLoops.cpp
@@ -53,7 +53,7 @@ static double CalculateAverageUVScale(const FDynamicMeshUVOverlay* UVOverlay, in
 			{
 				int32 jn = (j + 1) % 3;
 				double Len3D = Distance(Pos[j], Pos[jn]);
-				double LenUV = UV[j].Distance(UV[jn]);
+				double LenUV = Distance(UV[j], UV[jn]);
 				double Scale = LenUV / FMathd::Max(FMathf::ZeroTolerance, Len3D);
 				AverageScale += Scale;
 				Count++;
diff --git a/Engine/Plugins/Runtime/GeometryProcessing/Source/DynamicMesh/Private/Operations/MeshBoolean.cpp b/Engine/Plugins/Runtime/GeometryProcessing/Source/DynamicMesh/Private/Operations/MeshBoolean.cpp
--- a/Engine/Plugins/Runtime/GeometryProcessing/Source/DynamicMesh/Private/Operations/MeshBoolean.cpp
+++ b/Engine/Plugins/Runtime/GeometryProcessing/Source/DynamicMesh/Private/Operations/MeshBoolean.cpp
@@ -896,7 +896,7 @@ bool FMeshBoolean::CollapseWouldChangeShapeOrUVs(
 					FVector2f OtherUV = Mesh.GetVertexUV(OtherV);
 					FVector2f RemoveUV = Mesh.GetVertexUV(RemoveV);
 					FVector2f KeepUV = Mesh.GetVertexUV(KeepV);
-					if ( Lerp(OtherUV, KeepUV, LerpT).DistanceSquared(RemoveUV) > UVEqualThresholdSq)
+					if ( DistanceSquared( Lerp(OtherUV, KeepUV, LerpT), RemoveUV) > UVEqualThresholdSq)
 					{
 						bHasBadEdge = true;
 						return;
@@ -954,7 +954,7 @@ bool FMeshBoolean::CollapseWouldChangeShapeOrUVs(
 						FVector2f OtherUV = UVs->GetElement(OtherE);
 						FVector2f RemoveUV = UVs->GetElement(RemoveE);
 						FVector2f KeepUV = UVs->GetElement(KeepE);
-						if (Lerp(OtherUV, KeepUV, LerpT).DistanceSquared(RemoveUV) > UVEqualThresholdSq)
+						if ( DistanceSquared( Lerp(OtherUV, KeepUV, LerpT), RemoveUV) > UVEqualThresholdSq)
 						{
 							bHasBadEdge = true;
 							return;
diff --git a/Engine/Plugins/Runtime/GeometryProcessing/Source/DynamicMesh/Private/Parameterization/PatchBasedMeshUVGenerator.cpp b/Engine/Plugins/Runtime/GeometryProcessing/Source/DynamicMesh/Private/Parameterization/PatchBasedMeshUVGenerator.cpp
--- a/Engine/Plugins/Runtime/GeometryProcessing/Source/DynamicMesh/Private/Parameterization/PatchBasedMeshUVGenerator.cpp
+++ b/Engine/Plugins/Runtime/GeometryProcessing/Source/DynamicMesh/Private/Parameterization/PatchBasedMeshUVGenerator.cpp
@@ -284,7 +284,7 @@ static FUVAreaMetrics Compute_UV_Area_Metrics(const FDynamicMesh3* Mesh, const T
 		{
 			if (Submesh.IsBoundaryEdge(TriEdges[j]))
 			{
-				Metrics.UVPerimeter += UVTri[j].Distance(UVTri[(j + 1) % 3]);
+				Metrics.UVPerimeter += Distance( UVTri[j], UVTri[(j + 1) % 3] );
 			}
 		}
 
diff --git a/Engine/Plugins/Runtime/GeometryProcessing/Source/GeometryAlgorithms/Private/ConstrainedDelaunay2.cpp b/Engine/Plugins/Runtime/GeometryProcessing/Source/GeometryAlgorithms/Private/ConstrainedDelaunay2.cpp
--- a/Engine/Plugins/Runtime/GeometryProcessing/Source/GeometryAlgorithms/Private/ConstrainedDelaunay2.cpp
+++ b/Engine/Plugins/Runtime/GeometryProcessing/Source/GeometryAlgorithms/Private/ConstrainedDelaunay2.cpp
@@ -33,12 +33,12 @@ using namespace UE::Geometry;
 //
 //		for (int32 VertexIdx = 0; VertexIdx < Delaunay.Vertices.Num(); VertexIdx++)
 //		{
-//			const FVector2<RealType>& Vertex = Delaunay.Vertices[VertexIdx];
+//			const TVector2<RealType>& Vertex = Delaunay.Vertices[VertexIdx];
 //			f << "v " << Vertex.X << " " << Vertex.Y << " 0" << std::endl;
 //		}
 //		for (int32 VertexIdx = 0; VertexIdx < Delaunay.Vertices.Num(); VertexIdx++)
 //		{
-//			const FVector2<RealType>& Vertex = Delaunay.Vertices[VertexIdx];
+//			const TVector2<RealType>& Vertex = Delaunay.Vertices[VertexIdx];
 //			f << "v " << Vertex.X << " " << Vertex.Y << " .5" << std::endl;
 //		}
 //		for (const FIndex2i& Edge : Delaunay.Edges)
@@ -56,7 +56,7 @@ using namespace UE::Geometry;
 //		std::ofstream f(*Path);
 //		for (int32 VertexIdx = 0; VertexIdx < Delaunay.Vertices.Num(); VertexIdx++)
 //		{
-//			const FVector2<RealType>& Vertex = Delaunay.Vertices[VertexIdx];
+//			const TVector2<RealType>& Vertex = Delaunay.Vertices[VertexIdx];
 //			f << "v " << Vertex.X << " " << Vertex.Y << " 0" << std::endl;
 //		}
 //		for (const FIndex3i& Tri : Delaunay.Triangles)
@@ -120,7 +120,7 @@ bool HasUnorderedEdge(const TMap<TPair<int, int>, bool>& EdgeMap, int VertA, int
 namespace ConstrainedDelaunay2Internal
 {
 	template<class RealType>
-	void SplitBowtiesHelper(TArray<TPair<int, int>>& NeedUpdates, TArray<FVector2<RealType>>& Vertices, TArray<int8>& Keep, const std::vector<int>& Indices, const std::vector<int>& Adj)
+	void SplitBowtiesHelper(TArray<TPair<int, int>>& NeedUpdates, TArray<TVector2<RealType>>& Vertices, TArray<int8>& Keep, const std::vector<int>& Indices, const std::vector<int>& Adj)
 	{
 		int TriNum = int(Adj.size() / 3);
 		int32 OrigNumVertices = Vertices.Num();
@@ -187,7 +187,7 @@ namespace ConstrainedDelaunay2Internal
 				int NewVertID = -1;
 				if (bSeenSource)
 				{
-					FVector2<RealType> VertexToCopy = Vertices[VertID];
+					TVector2<RealType> VertexToCopy = Vertices[VertID];
 					NewVertID = Vertices.Add(VertexToCopy);
 				}
 
@@ -281,7 +281,7 @@ namespace ConstrainedDelaunay2Internal
 
 
 template<class RealType>
-bool TConstrainedDelaunay2<RealType>::Triangulate(TFunctionRef<bool(const TArray<FVector2<RealType>>&, const FIndex3i&)> KeepTriangle)
+bool TConstrainedDelaunay2<RealType>::Triangulate(TFunctionRef<bool(const TArray<TVector2<RealType>>&, const FIndex3i&)> KeepTriangle)
 {
 	Triangles.Empty();
 
diff --git a/Engine/Plugins/Runtime/GeometryProcessing/Source/GeometryAlgorithms/Private/GteUtil.h b/Engine/Plugins/Runtime/GeometryProcessing/Source/GeometryAlgorithms/Private/GteUtil.h
--- a/Engine/Plugins/Runtime/GeometryProcessing/Source/GeometryAlgorithms/Private/GteUtil.h
+++ b/Engine/Plugins/Runtime/GeometryProcessing/Source/GeometryAlgorithms/Private/GteUtil.h
@@ -19,15 +19,15 @@ namespace Geometry
 {
 
 template<typename Real>
-gte::Vector2<Real> Convert(const FVector2<Real>& Vec)
+gte::Vector2<Real> Convert(const UE::Math::TVector2<Real>& Vec)
 {
 	return gte::Vector2<Real>({ Vec.X, Vec.Y});
 }
 
 template<typename Real>
-FVector2<Real> Convert(const gte::Vector2<Real>& Vec)
+UE::Math::TVector2<Real> Convert(const gte::Vector2<Real>& Vec)
 {
-	return FVector2<Real>(Vec[0], Vec[1]);
+	return UE::Math::TVector2<Real>(Vec[0], Vec[1]);
 }
 
 template<typename Real>
diff --git a/Engine/Plugins/Runtime/GeometryProcessing/Source/GeometryAlgorithms/Public/Arrangement2d.h b/Engine/Plugins/Runtime/GeometryProcessing/Source/GeometryAlgorithms/Public/Arrangement2d.h
--- a/Engine/Plugins/Runtime/GeometryProcessing/Source/GeometryAlgorithms/Public/Arrangement2d.h
+++ b/Engine/Plugins/Runtime/GeometryProcessing/Source/GeometryAlgorithms/Public/Arrangement2d.h
@@ -270,7 +270,7 @@ protected:
 		}
 
 		// handle tiny-segment case
-		double SegLenSq = A.DistanceSquared(B);
+		double SegLenSq = DistanceSquared( A, B );
 		if (SegLenSq <= VertexSnapTol*VertexSnapTol)
 		{
 			// seg is too short and was already on an existing vertex; just consider that vertex to be the inserted segment
@@ -441,7 +441,7 @@ protected:
 	 */
 	int find_nearest_vertex(FVector2d Pt, double SearchRadius, int IgnoreVID = -1)
 	{
-		auto FuncDistSq = [&](int B) { return Pt.DistanceSquared(Graph.GetVertex(B)); };
+		auto FuncDistSq = [&](int B) { return DistanceSquared(Pt, Graph.GetVertex(B)); };
 		auto FuncIgnore = [&](int VID) { return VID == IgnoreVID; };
 		TPair<int, double> found = (IgnoreVID == -1) ? PointHash.FindNearestInRadius(Pt, SearchRadius, FuncDistSq)
 													 : PointHash.FindNearestInRadius(Pt, SearchRadius, FuncDistSq, FuncIgnore);
@@ -457,7 +457,7 @@ protected:
 	 */
 	int find_nearest_boundary_vertex(FVector2d Pt, double SearchRadius, int IgnoreVID = -1)
 	{
-		auto FuncDistSq = [&](int B) { return Pt.DistanceSquared(Graph.GetVertex(B)); };
+		auto FuncDistSq = [&](int B) { return DistanceSquared(Pt, Graph.GetVertex(B)); };
 		auto FuncIgnore = [&](int VID) { return Graph.IsBoundaryVertex(VID) == false || VID == IgnoreVID; };
 		TPair<int, double> found =
 			PointHash.FindNearestInRadius(Pt, SearchRadius, FuncDistSq, FuncIgnore);
diff --git a/Engine/Plugins/Runtime/GeometryProcessing/Source/GeometryAlgorithms/Public/ConstrainedDelaunay2.h b/Engine/Plugins/Runtime/GeometryProcessing/Source/GeometryAlgorithms/Public/ConstrainedDelaunay2.h
--- a/Engine/Plugins/Runtime/GeometryProcessing/Source/GeometryAlgorithms/Public/ConstrainedDelaunay2.h
+++ b/Engine/Plugins/Runtime/GeometryProcessing/Source/GeometryAlgorithms/Public/ConstrainedDelaunay2.h
@@ -14,6 +14,7 @@
 namespace UE {
 namespace Geometry {
 
+using namespace UE::Math;
 
 
 template<typename RealType>
@@ -22,7 +23,7 @@ struct TConstrainedDelaunay2
 	//
 	// inputs
 	//
-	TArray<FVector2<RealType>> Vertices;
+	TArray<TVector2<RealType>> Vertices;
 
 	// Edges and HoleEdges must not be intersecting; use Arrangment2D to pre-process any input w/ intersecting edges
 	TArray<FIndex2i> Edges;	// Edges can be boundaries or not based on the EFillRule setting
@@ -73,7 +74,7 @@ struct TConstrainedDelaunay2
 			if (Graph.IsVertex(Idx))
 			{
 				GraphToDTVertIdxMap[Idx] = Vertices.Num();
-				Vertices.Add((FVector2<RealType>)Graph.GetVertex(Idx));
+				Vertices.Add((TVector2<RealType>)Graph.GetVertex(Idx));
 			}
 			else
 			{
@@ -90,9 +91,9 @@ struct TConstrainedDelaunay2
 	{
 		int32 VertexStart = Vertices.Num();
 		int32 VertexEnd = VertexStart + Polygon.VertexCount();
-		for (const FVector2<InputRealType> &Vertex : Polygon.GetVertices())
+		for (const TVector2<InputRealType> &Vertex : Polygon.GetVertices())
 		{
-			Vertices.Add((FVector2<RealType>)Vertex);
+			Vertices.Add((TVector2<RealType>)Vertex);
 		}
 
 		TArray<FIndex2i>* EdgeArr;
@@ -142,7 +143,7 @@ struct TConstrainedDelaunay2
 	 * @param KeepTriangle Function to check whether the given triangle should be kept in the output
 	 * @return false if Triangulation failed
 	 */
-	bool GEOMETRYALGORITHMS_API Triangulate(TFunctionRef<bool(const TArray<FVector2<RealType>>&, const FIndex3i&)> KeepTriangle);
+	bool GEOMETRYALGORITHMS_API Triangulate(TFunctionRef<bool(const TArray<TVector2<RealType>>&, const FIndex3i&)> KeepTriangle);
 };
 
 template<typename RealType>
diff --git a/Engine/Source/Runtime/GeometryCore/Private/CompGeom/ConvexHull2.cpp b/Engine/Source/Runtime/GeometryCore/Private/CompGeom/ConvexHull2.cpp
--- a/Engine/Source/Runtime/GeometryCore/Private/CompGeom/ConvexHull2.cpp
+++ b/Engine/Source/Runtime/GeometryCore/Private/CompGeom/ConvexHull2.cpp
@@ -15,7 +15,7 @@ namespace Geometry
 
 
 template<class RealType>
-bool TConvexHull2<RealType>::Solve(int32 NumPoints, TFunctionRef<FVector2<RealType>(int32)> GetPointFunc, TFunctionRef<bool(int32)> FilterFunc)
+bool TConvexHull2<RealType>::Solve(int32 NumPoints, TFunctionRef<TVector2<RealType>(int32)> GetPointFunc, TFunctionRef<bool(int32)> FilterFunc)
 {
 	Hull.Reset();
 
@@ -33,7 +33,7 @@ bool TConvexHull2<RealType>::Solve(int32 NumPoints, TFunctionRef<FVector2<RealTy
 	}
 	Hull.Sort([&GetPointFunc](int32 A, int32 B)
 		{
-			FVector2<RealType> PA = GetPointFunc(A), PB = GetPointFunc(B);
+			TVector2<RealType> PA = GetPointFunc(A), PB = GetPointFunc(B);
 			return (PA.X < PB.X) || (PA.X == PB.X && PA.Y < PB.Y);
 		});
 
@@ -84,7 +84,7 @@ bool TConvexHull2<RealType>::Solve(int32 NumPoints, TFunctionRef<FVector2<RealTy
 }
 
 template<class RealType>
-void TConvexHull2<RealType>::GetHull(TFunctionRef<FVector2<RealType>(int32)> GetPointFunc, TArray<int32>& Merged, int32& IdxFirst, int32& IdxLast)
+void TConvexHull2<RealType>::GetHull(TFunctionRef<TVector2<RealType>(int32)> GetPointFunc, TArray<int32>& Merged, int32& IdxFirst, int32& IdxLast)
 {
 	int32 NumVertices = IdxLast - IdxFirst + 1;
 	if (NumVertices > 1)
@@ -104,7 +104,7 @@ void TConvexHull2<RealType>::GetHull(TFunctionRef<FVector2<RealType>(int32)> Get
 }
 
 template<class RealType>
-void TConvexHull2<RealType>::Merge(TFunctionRef<FVector2<RealType>(int32)> GetPointFunc, TArray<int32>& Merged, int32 j0, int32 j1, int32 j2, int32 j3, int32& i0, int32& i1)
+void TConvexHull2<RealType>::Merge(TFunctionRef<TVector2<RealType>(int32)> GetPointFunc, TArray<int32>& Merged, int32 j0, int32 j1, int32 j2, int32 j3, int32& i0, int32& i1)
 {
 	// Subhull0 is to the left of subhull1 because of the initial sorting of
 	// the points by x-components.  We need to find two mutually visible
@@ -237,7 +237,7 @@ EPointOrdering PointOnLine(FVector2d P, FVector2d L0, FVector2d L1)
 }
 
 template<class RealType>
-void TConvexHull2<RealType>::GetTangent(TFunctionRef<FVector2<RealType>(int32)> GetPointFunc, TArray<int32>& Merged, int32 j0, int32 j1, int32 j2, int32 j3, int32& i0, int32& i1)
+void TConvexHull2<RealType>::GetTangent(TFunctionRef<TVector2<RealType>(int32)> GetPointFunc, TArray<int32>& Merged, int32 j0, int32 j1, int32 j2, int32 j3, int32& i0, int32& i1)
 {
 	// In theory the loop terminates in a finite number of steps, but the
 	// upper bound for the loop variable is used to trap problems caused by
diff --git a/Engine/Source/Runtime/GeometryCore/Private/CompGeom/PolygonTriangulation.cpp b/Engine/Source/Runtime/GeometryCore/Private/CompGeom/PolygonTriangulation.cpp
--- a/Engine/Source/Runtime/GeometryCore/Private/CompGeom/PolygonTriangulation.cpp
+++ b/Engine/Source/Runtime/GeometryCore/Private/CompGeom/PolygonTriangulation.cpp
@@ -9,8 +9,8 @@ using namespace UE::Geometry;
 // explicit instantiations
 namespace PolygonTriangulation
 {
-	template GEOMETRYCORE_API void TriangulateSimplePolygon<float>(const TArray<FVector2<float>>& VertexPositions, TArray<FIndex3i>& OutTriangles);
-	template GEOMETRYCORE_API void TriangulateSimplePolygon<double>(const TArray<FVector2<double>>& VertexPositions, TArray<FIndex3i>& OutTriangles);
+	template GEOMETRYCORE_API void TriangulateSimplePolygon<float>(const TArray<TVector2<float>>& VertexPositions, TArray<FIndex3i>& OutTriangles);
+	template GEOMETRYCORE_API void TriangulateSimplePolygon<double>(const TArray<TVector2<double>>& VertexPositions, TArray<FIndex3i>& OutTriangles);
 
 	template GEOMETRYCORE_API void ComputePolygonPlane<float>(const TArray<FVector3<float>>& VertexPositions, FVector3<float>& NormalOut, FVector3<float>& PlanePointOut);
 	template GEOMETRYCORE_API void ComputePolygonPlane<double>(const TArray<FVector3<double>>& VertexPositions, FVector3<double>& NormalOut, FVector3<double>& PlanePointOut);
@@ -27,11 +27,11 @@ namespace PolygonTriangulation
 // This is based on the 3D triangulation code from MeshDescription.cpp, simplified for 2D polygons
 // 
 template<typename T>
-void PolygonTriangulation::TriangulateSimplePolygon(const TArray<FVector2<T>>& VertexPositions, TArray<FIndex3i>& OutTriangles)
+void PolygonTriangulation::TriangulateSimplePolygon(const TArray<TVector2<T>>& VertexPositions, TArray<FIndex3i>& OutTriangles)
 {
 	struct Local
 	{
-		static inline bool IsTriangleFlipped(T OrientationSign, const FVector2<T>& VertexPositionA, const FVector2<T>& VertexPositionB, const FVector2<T>& VertexPositionC)
+		static inline bool IsTriangleFlipped(T OrientationSign, const TVector2<T>& VertexPositionA, const TVector2<T>& VertexPositionB, const TVector2<T>& VertexPositionC)
 		{
 			T TriSignedArea = TTriangle2<T>::SignedArea(VertexPositionA, VertexPositionB, VertexPositionC);
 			return TriSignedArea * OrientationSign < 0;
@@ -47,8 +47,8 @@ void PolygonTriangulation::TriangulateSimplePolygon(const TArray<FVector2<T>>& V
 	double PolySignedArea = 0;
 	for (int i = 0; i < PolygonVertexCount; ++i)
 	{
-		const FVector2<T>& v1 = VertexPositions[i];
-		const FVector2<T>& v2 = VertexPositions[(i + 1) % PolygonVertexCount];
+		const TVector2<T>& v1 = VertexPositions[i];
+		const TVector2<T>& v2 = VertexPositions[(i + 1) % PolygonVertexCount];
 		PolySignedArea += v1.X*v2.Y - v1.Y*v2.X;
 	}
 	PolySignedArea *= 0.5;
@@ -93,9 +93,9 @@ void PolygonTriangulation::TriangulateSimplePolygon(const TArray<FVector2<T>>& V
 		// vertices are collinear or other degenerate cases.
 		if (RemainingVertexCount > 3 && EarTestCount < RemainingVertexCount)
 		{
-			const FVector2<T>& PrevVertexPosition = VertexPositions[PrevVertexNumbers[EarVertexNumber]];
-			const FVector2<T>& EarVertexPosition = VertexPositions[EarVertexNumber];
-			const FVector2<T>& NextVertexPosition = VertexPositions[NextVertexNumbers[EarVertexNumber]];
+			const TVector2<T>& PrevVertexPosition = VertexPositions[PrevVertexNumbers[EarVertexNumber]];
+			const TVector2<T>& EarVertexPosition = VertexPositions[EarVertexNumber];
+			const TVector2<T>& NextVertexPosition = VertexPositions[NextVertexNumbers[EarVertexNumber]];
 
 			// Figure out whether the potential ear triangle is facing the same direction as the polygon
 			// itself.  If it's facing the opposite direction, then we're dealing with a concave triangle
@@ -109,7 +109,7 @@ void PolygonTriangulation::TriangulateSimplePolygon(const TArray<FVector2<T>>& V
 				{
 					// Test every other remaining vertex to make sure that it doesn't lie inside our potential ear
 					// triangle.  If we find a vertex that's inside the triangle, then it cannot actually be an ear.
-					const FVector2<T>& TestVertexPosition = VertexPositions[TestVertexNumber];
+					const TVector2<T>& TestVertexPosition = VertexPositions[TestVertexNumber];
 					if (TTriangle2<T>::IsInside(PrevVertexPosition, EarVertexPosition, NextVertexPosition, TestVertexPosition))
 					{
 						bIsEar = false;
diff --git a/Engine/Source/Runtime/GeometryCore/Public/BoxTypes.h b/Engine/Source/Runtime/GeometryCore/Public/BoxTypes.h
--- a/Engine/Source/Runtime/GeometryCore/Public/BoxTypes.h
+++ b/Engine/Source/Runtime/GeometryCore/Public/BoxTypes.h
@@ -12,6 +12,7 @@ namespace UE
 namespace Geometry
 {
 
+using namespace UE::Math;
 
 template <typename RealType>
 struct TInterval1
@@ -539,15 +540,15 @@ struct TAxisAlignedBox3
 template <typename RealType>
 struct TAxisAlignedBox2
 {
-	FVector2<RealType> Min;
-	FVector2<RealType> Max;
+	TVector2<RealType> Min;
+	TVector2<RealType> Max;
 
 	TAxisAlignedBox2() : 
 		TAxisAlignedBox2(Empty())
 	{
 	}
 
-	TAxisAlignedBox2(const FVector2<RealType>& Min, const FVector2<RealType>& Max)
+	TAxisAlignedBox2(const TVector2<RealType>& Min, const TVector2<RealType>& Max)
 		: Min(Min), Max(Max)
 	{
 	}
@@ -557,8 +558,8 @@ struct TAxisAlignedBox2
 	template<typename OtherRealType>
 	explicit TAxisAlignedBox2(const TAxisAlignedBox2<OtherRealType>& OtherBox)
 	{
-		this->Min = FVector2<RealType>(OtherBox.Min);
-		this->Max = FVector2<RealType>(OtherBox.Max);
+		this->Min = TVector2<RealType>(OtherBox.Min);
+		this->Max = TVector2<RealType>(OtherBox.Max);
 	}
 
 	TAxisAlignedBox2(RealType SquareSize)
@@ -570,16 +571,16 @@ struct TAxisAlignedBox2
 	{
 	}
 
-	TAxisAlignedBox2(const TArray<FVector2<RealType>>& Pts)
+	TAxisAlignedBox2(const TArray<TVector2<RealType>>& Pts)
 	{
 		*this = Empty();
 		Contain(Pts);
 	}
 
-	TAxisAlignedBox2(const FVector2<RealType>& Center, RealType HalfWidth)
+	TAxisAlignedBox2(const TVector2<RealType>& Center, RealType HalfWidth)
 	{
-		this->Min = FVector2<RealType>(Center.X - HalfWidth, Center.Y - HalfWidth);
-		this->Max = FVector2<RealType>(Center.X + HalfWidth, Center.Y + HalfWidth);
+		this->Min = TVector2<RealType>(Center.X - HalfWidth, Center.Y - HalfWidth);
+		this->Max = TVector2<RealType>(Center.X + HalfWidth, Center.Y + HalfWidth);
 	}
 
 	explicit operator FBox2D() const
@@ -597,18 +598,18 @@ struct TAxisAlignedBox2
 	static TAxisAlignedBox2<RealType> Empty()
 	{
 		return TAxisAlignedBox2(
-			FVector2<RealType>(TNumericLimits<RealType>::Max(), TNumericLimits<RealType>::Max()),
-			FVector2<RealType>(-TNumericLimits<RealType>::Max(), -TNumericLimits<RealType>::Max()));
+			TVector2<RealType>(TNumericLimits<RealType>::Max(), TNumericLimits<RealType>::Max()),
+			TVector2<RealType>(-TNumericLimits<RealType>::Max(), -TNumericLimits<RealType>::Max()));
 	}
 
-	FVector2<RealType> Center() const
+	TVector2<RealType> Center() const
 	{
-		return FVector2<RealType>(
+		return TVector2<RealType>(
 			(Min.X + Max.X) * (RealType)0.5,
 			(Min.Y + Max.Y) * (RealType)0.5);
 	}
 
-	FVector2<RealType> Extents() const
+	TVector2<RealType> Extents() const
 	{
 		return (Max - Min) * (RealType).5;
 	}
@@ -618,15 +619,15 @@ struct TAxisAlignedBox2
 	 * @param Index which corner to return, must be in range [0,3]
 	 * @return Corner of the bounding rectangle
 	 */
-	FVector2<RealType> GetCorner(int Index) const
+	TVector2<RealType> GetCorner(int Index) const
 	{
 		check(Index >= 0 && Index <= 3);
 		RealType X = ((Index % 3) == 0) ? (Min.X) : (Max.X);
 		RealType Y = ((Index & 2) == 0) ? (Min.Y) : (Max.Y);
-		return FVector2<RealType>(X, Y);
+		return TVector2<RealType>(X, Y);
 	}
 
-	inline void Contain(const FVector2<RealType>& V)
+	inline void Contain(const TVector2<RealType>& V)
 	{
 		if (V.X < Min.X)
 		{
@@ -654,15 +655,15 @@ struct TAxisAlignedBox2
 		Max.Y = Max.Y > Other.Max.Y ? Max.Y : Other.Max.Y;
 	}
 
-	void Contain(const TArray<FVector2<RealType>>& Pts)
+	void Contain(const TArray<TVector2<RealType>>& Pts)
 	{
-		for (const FVector2<RealType>& Pt : Pts)
+		for (const TVector2<RealType>& Pt : Pts)
 		{
 			Contain(Pt);
 		}
 	}
 
-	bool Contains(const FVector2<RealType>& V) const
+	bool Contains(const TVector2<RealType>& V) const
 	{
 		return (Min.X <= V.X) && (Min.Y <= V.Y) && (Max.X >= V.X) && (Max.Y >= V.Y);
 	}
@@ -680,8 +681,8 @@ struct TAxisAlignedBox2
 	TAxisAlignedBox2<RealType> Intersect(const TAxisAlignedBox2<RealType> &Box) const
 	{
 		TAxisAlignedBox2<RealType> Intersection(
-			FVector2<RealType>(TMathUtil<RealType>::Max(Min.X, Box.Min.X), TMathUtil<RealType>::Max(Min.Y, Box.Min.Y)),
-			FVector2<RealType>(TMathUtil<RealType>::Min(Max.X, Box.Max.X), TMathUtil<RealType>::Min(Max.Y, Box.Max.Y)));
+			TVector2<RealType>(TMathUtil<RealType>::Max(Min.X, Box.Min.X), TMathUtil<RealType>::Max(Min.Y, Box.Min.Y)),
+			TVector2<RealType>(TMathUtil<RealType>::Min(Max.X, Box.Max.X), TMathUtil<RealType>::Min(Max.Y, Box.Max.Y)));
 		if (Intersection.Height() <= 0 || Intersection.Width() <= 0)
 		{
 			return TAxisAlignedBox2<RealType>::Empty();
@@ -692,7 +693,7 @@ struct TAxisAlignedBox2
 		}
 	}
 
-	RealType DistanceSquared(const FVector2<RealType>& V) const
+	RealType DistanceSquared(const TVector2<RealType>& V) const
 	{
 		RealType dx = (V.X < Min.X) ? Min.X - V.X : (V.X > Max.X ? V.X - Max.X : 0);
 		RealType dy = (V.Y < Min.Y) ? Min.Y - V.Y : (V.Y > Max.Y ? V.Y - Max.Y : 0);
diff --git a/Engine/Source/Runtime/GeometryCore/Public/CircleTypes.h b/Engine/Source/Runtime/GeometryCore/Public/CircleTypes.h
--- a/Engine/Source/Runtime/GeometryCore/Public/CircleTypes.h
+++ b/Engine/Source/Runtime/GeometryCore/Public/CircleTypes.h
@@ -13,10 +13,12 @@ namespace UE
 namespace Geometry 
 {
 
+using namespace UE::Math;
+
 template<typename RealType>
 struct TCircle2
 {
-	FVector2<RealType> Center = FVector2<RealType>(0, 0);
+	TVector2<RealType> Center = TVector2<RealType>(0, 0);
 	RealType Radius = (RealType)1;
 	bool bIsReversed = false;
 
@@ -24,11 +26,11 @@ struct TCircle2
 
 	explicit TCircle2(const RealType& RadiusIn)
 	{
-		Center = FVector2<RealType>(0, 0);
+		Center = TVector2<RealType>(0, 0);
 		Radius = RadiusIn;
 	}
 
-	TCircle2(const FVector2<RealType>& CenterIn, const RealType& RadiusIn)
+	TCircle2(const TVector2<RealType>& CenterIn, const RealType& RadiusIn)
 	{
 		Center = CenterIn;
 		Radius = RadiusIn;
@@ -75,37 +77,37 @@ struct TCircle2
 	}
 
 
-	FVector2<RealType> GetPointFromAngleD(RealType AngleDeg) const
+	TVector2<RealType> GetPointFromAngleD(RealType AngleDeg) const
 	{
 		return GetPointFromAngleR(AngleDeg * TMathUtil<RealType>::DegToRad);
 	}
 
-	FVector2<RealType> GetPointFromAngleR(RealType AngleRad) const
+	TVector2<RealType> GetPointFromAngleR(RealType AngleRad) const
 	{
 		RealType c = TMathUtil<RealType>::Cos(AngleRad), s = TMathUtil<RealType>::Sin(AngleRad);
-		return FVector2<RealType>(Center.X + c*Radius, Center.Y + s*Radius);
+		return TVector2<RealType>(Center.X + c*Radius, Center.Y + s*Radius);
 	}
 
 
-	FVector2<RealType> GetPointFromUnitParameter(RealType UnitParam) const
+	TVector2<RealType> GetPointFromUnitParameter(RealType UnitParam) const
 	{
 		RealType AngleRad = ((bIsReversed) ? (-UnitParam) : (UnitParam)) * TMathUtil<RealType>::TwoPi;
 		return GetPointFromAngleR(AngleRad);
 	}
 
 
-	bool IsInside(const FVector2<RealType>& Point) const
+	bool IsInside(const TVector2<RealType>& Point) const
 	{
 		return Center.DistanceSquared(Point) < Radius*Radius;
 	}
 
 
-	RealType SignedDistance(const FVector2<RealType>& Point) const
+	RealType SignedDistance(const TVector2<RealType>& Point) const
 	{
 		return Center.Distance(Point) - Radius;
 	}
 
-	RealType Distance(const FVector2<RealType>& Point) const
+	RealType Distance(const TVector2<RealType>& Point) const
 	{
 		return TMathUtil<RealType>::Abs(Center.Distance(Point) - Radius);
 	}
@@ -113,10 +115,10 @@ struct TCircle2
 
 	TAxisAlignedBox2<RealType> GetBoundingBox() const
 	{
-		FVector2<RealType>(Center.X + Radius, Center.Y + Radius);
+		TVector2<RealType>(Center.X + Radius, Center.Y + Radius);
 		return TAxisAlignedBox2<RealType>(
-			FVector2<RealType>(Center.X - Radius, Center.Y - Radius),
-			FVector2<RealType>(Center.X + Radius, Center.Y + Radius));
+			TVector2<RealType>(Center.X - Radius, Center.Y - Radius),
+			TVector2<RealType>(Center.X + Radius, Center.Y + Radius));
 	}
 
 
@@ -223,7 +225,7 @@ struct TCircle3
 	FVector3<RealType> GetPointFromAngleR(RealType AngleRad) const
 	{
 		RealType c = TMathUtil<RealType>::Cos(AngleRad), s = TMathUtil<RealType>::Sin(AngleRad);
-		return Frame.FromPlaneUV(FVector2<RealType>(Radius*c, Radius*s), 2);
+		return Frame.FromPlaneUV(TVector2<RealType>(Radius*c, Radius*s), 2);
 	}
 
 	FVector3<RealType> GetPointFromUnitParameter(RealType UnitParam) const
diff --git a/Engine/Source/Runtime/GeometryCore/Public/CompGeom/ConvexHull2.h b/Engine/Source/Runtime/GeometryCore/Public/CompGeom/ConvexHull2.h
--- a/Engine/Source/Runtime/GeometryCore/Public/CompGeom/ConvexHull2.h
+++ b/Engine/Source/Runtime/GeometryCore/Public/CompGeom/ConvexHull2.h
@@ -16,6 +16,7 @@
 namespace UE {
 namespace Geometry {
 
+using namespace UE::Math;
 
 
 template<typename RealType>
@@ -33,7 +34,7 @@ public:
 	 * @param FilterFunc Optional filter to include only a subset of the points in the output hull
 	 * @return true if hull was generated, false if points span < 2 dimensions
 	 */
-	bool Solve(int32 NumPoints, TFunctionRef<FVector2<RealType>(int32)> GetPointFunc, TFunctionRef<bool(int32)> FilterFunc = [](int32 Idx) {return true;});
+	bool Solve(int32 NumPoints, TFunctionRef<TVector2<RealType>(int32)> GetPointFunc, TFunctionRef<bool(int32)> FilterFunc = [](int32 Idx) {return true;});
 
 	/**
 	 * Generate convex hull as long as input is not degenerate
@@ -44,7 +45,7 @@ public:
 	 * @param Filter Optional filter to include only a subset of the points in the output hull
 	 * @return true if hull was generated, false if points span < 2 dimensions
 	 */
-	bool Solve(TArrayView<const FVector2<RealType>> Points, TFunctionRef<bool(int32)> FilterFunc)
+	bool Solve(TArrayView<const TVector2<RealType>> Points, TFunctionRef<bool(int32)> FilterFunc)
 	{
 		return Solve(Points.Num(), [&Points](int32 Idx)
 			{
@@ -61,7 +62,7 @@ public:
 	 * @param Points Array of points to consider
 	 * @return true if hull was generated, false if points span < 2 dimensions
 	 */
-	bool Solve(TArrayView<const FVector2<RealType>> Points)
+	bool Solve(TArrayView<const TVector2<RealType>> Points)
 	{
 		return Solve(Points.Num(), [&Points](int32 Idx)
 			{
@@ -108,9 +109,9 @@ public:
 protected:
 
 	// divide-and-conquer algorithm
-	void GetHull(TFunctionRef<FVector2<RealType>(int32)> GetPointFunc, TArray<int32>& Merged, int32& IdxFirst, int32& IdxLast);
-	void Merge(TFunctionRef<FVector2<RealType>(int32)> GetPointFunc, TArray<int32>& Merged, int32 j0, int32 j1, int32 j2, int32 j3, int32& i0, int32& i1);
-	void GetTangent(TFunctionRef<FVector2<RealType>(int32)> GetPointFunc, TArray<int32>& Merged, int32 j0, int32 j1, int32 j2, int32 j3, int32& i0, int32& i1);
+	void GetHull(TFunctionRef<TVector2<RealType>(int32)> GetPointFunc, TArray<int32>& Merged, int32& IdxFirst, int32& IdxLast);
+	void Merge(TFunctionRef<TVector2<RealType>(int32)> GetPointFunc, TArray<int32>& Merged, int32 j0, int32 j1, int32 j2, int32 j3, int32& i0, int32& i1);
+	void GetTangent(TFunctionRef<TVector2<RealType>(int32)> GetPointFunc, TArray<int32>& Merged, int32 j0, int32 j1, int32 j2, int32 j3, int32& i0, int32& i1);
 
 	int32 Dimension = 0;
 	int32 NumUniquePoints = 0;
diff --git a/Engine/Source/Runtime/GeometryCore/Public/CompGeom/PolygonTriangulation.h b/Engine/Source/Runtime/GeometryCore/Public/CompGeom/PolygonTriangulation.h
--- a/Engine/Source/Runtime/GeometryCore/Public/CompGeom/PolygonTriangulation.h
+++ b/Engine/Source/Runtime/GeometryCore/Public/CompGeom/PolygonTriangulation.h
@@ -8,6 +8,7 @@
 namespace PolygonTriangulation
 {
 	using namespace UE::Geometry;
+	using namespace UE::Math;
 
 	/**
 	 * Compute triangulation of simple 2D polygon using ear-clipping
@@ -15,7 +16,7 @@ namespace PolygonTriangulation
 	 * @param OutTriangles computed triangulation. Each triangle is a tuple of indices into VertexPositions.
 	 */
 	template<typename RealType>
-	void GEOMETRYCORE_API TriangulateSimplePolygon(const TArray<FVector2<RealType>>& VertexPositions, TArray<FIndex3i>& OutTriangles);
+	void GEOMETRYCORE_API TriangulateSimplePolygon(const TArray<TVector2<RealType>>& VertexPositions, TArray<FIndex3i>& OutTriangles);
 
 
 	template<typename RealType>
diff --git a/Engine/Source/Runtime/GeometryCore/Public/Curve/DynamicGraph2.h b/Engine/Source/Runtime/GeometryCore/Public/Curve/DynamicGraph2.h
--- a/Engine/Source/Runtime/GeometryCore/Public/Curve/DynamicGraph2.h
+++ b/Engine/Source/Runtime/GeometryCore/Public/Curve/DynamicGraph2.h
@@ -20,6 +20,8 @@ namespace UE
 namespace Geometry
 {
 
+using namespace UE::Math;
+
 template <typename T>
 class FDynamicGraph2 : public FDynamicGraph
 {
@@ -29,17 +31,17 @@ class FDynamicGraph2 : public FDynamicGraph
 	TDynamicVectorN<T, 2> vertices;
 
 public:
-	static FVector2<T> InvalidVertex()
+	static TVector2<T> InvalidVertex()
 	{
-		return FVector2<T>(TNumericLimits<T>::Max(), 0);
+		return TVector2<T>(TNumericLimits<T>::Max(), 0);
 	}
 
-	FVector2<T> GetVertex(int VID) const
+	TVector2<T> GetVertex(int VID) const
 	{
 		return vertices_refcount.IsValid(VID) ? vertices.AsVector2(VID) : InvalidVertex();
 	}
 
-	void SetVertex(int VID, FVector2<T> VNewPos)
+	void SetVertex(int VID, TVector2<T> VNewPos)
 	{
 		check(VectorUtil::IsFinite(VNewPos)); // this will really catch a lot of bugs...
 		if (vertices_refcount.IsValid(VID))
@@ -50,7 +52,7 @@ public:
 	}
 
 	using FDynamicGraph::GetEdgeV;
-	bool GetEdgeV(int EID, FVector2<T>& A, FVector2<T>& B) const
+	bool GetEdgeV(int EID, UE::Math::TVector2<T>& A, UE::Math::TVector2<T>& B) const
 	{
 		if (edges_refcount.IsValid(EID))
 		{
@@ -70,14 +72,14 @@ public:
 			vertices.AsVector2(e.B));
 	}
 
-	FVector2<T> GetEdgeCenter(int EID) const
+	TVector2<T> GetEdgeCenter(int EID) const
 	{
 		checkf(edges_refcount.IsValid(EID), TEXT("FDynamicGraph2.GetEdgeCenter: invalid segment with id %d"), EID);
 		const FEdge& e = edges[EID];
 		return 0.5 * (vertices.AsVector2(e.A) + vertices.AsVector2(e.B));
 	}
 
-	int AppendVertex(FVector2<T> V)
+	int AppendVertex(TVector2<T> V)
 	{
 		int vid = append_vertex_internal();
 		vertices.InsertAt({{V.X, V.Y}}, vid);
@@ -134,9 +136,9 @@ public:
 */
 
 	/** Enumerate positions of all vertices in graph */
-	value_iteration<FVector2<T>> Vertices() const
+	value_iteration<TVector2<T>> Vertices() const
 	{
-		return vertices_refcount.MappedIndices<FVector2<T>>(
+		return vertices_refcount.MappedIndices<TVector2<T>>(
 			[=](int vid) {
 				return vertices.template AsVector2<T>(vid);
 			});
@@ -154,10 +156,10 @@ public:
 		{
 			Sorted.Add(EID);
 		}
-		FVector2<T> V = vertices.AsVector2(VID);
+		TVector2<T> V = vertices.AsVector2(VID);
 		Algo::SortBy(Sorted, [&](int EID) {
 			int NbrVID = edge_other_v(EID, VID);
-			FVector2<T> D = vertices.AsVector2(NbrVID) - V;
+			TVector2<T> D = vertices.AsVector2(NbrVID) - V;
 			return TMathUtil<T>::Atan2Positive(D.Y, D.X);
 		});
 
@@ -168,7 +170,7 @@ public:
 	FAxisAlignedBox2d GetBounds() const
 	{
 		TAxisAlignedBox2<T> AABB;
-		for (const FVector2<T>& V : Vertices())
+		for (const TVector2<T>& V : Vertices())
 		{
 			AABB.Contain(V);
 		}
@@ -206,9 +208,9 @@ public:
 		int nbra = edge_other_v(*ValueIterate, VID);
 		int nbrb = edge_other_v(*++ValueIterate, VID);
 
-		FVector2<T> V = vertices.AsVector2(VID);
-		FVector2<T> A = vertices.AsVector2(nbra);
-		FVector2<T> B = vertices.AsVector2(nbrb);
+		TVector2<T> V = vertices.AsVector2(VID);
+		TVector2<T> A = vertices.AsVector2(nbra);
+		TVector2<T> B = vertices.AsVector2(nbrb);
 		A -= V;
 		if (Normalize(A) == 0)
 		{
@@ -226,7 +228,7 @@ protected:
 	// internal used in SplitEdge
 	virtual int append_new_split_vertex(int A, int B) override
 	{
-		FVector2<T> vNew = 0.5 * (GetVertex(A) + GetVertex(B));
+		TVector2<T> vNew = 0.5 * (GetVertex(A) + GetVertex(B));
 		int f = AppendVertex(vNew);
 		return f;
 	}
@@ -235,7 +237,7 @@ protected:
 	{
 		for (int VID : VertexIndices())
 		{
-			FVector2<T> V = GetVertex(VID);
+			TVector2<T> V = GetVertex(VID);
 			CheckOrFailF(VectorUtil::IsFinite(V));
 		}
 	}
diff --git a/Engine/Source/Runtime/GeometryCore/Public/Curve/GeneralPolygon2.h b/Engine/Source/Runtime/GeometryCore/Public/Curve/GeneralPolygon2.h
--- a/Engine/Source/Runtime/GeometryCore/Public/Curve/GeneralPolygon2.h
+++ b/Engine/Source/Runtime/GeometryCore/Public/Curve/GeneralPolygon2.h
@@ -18,6 +18,8 @@ namespace UE
 namespace Geometry
 {
 
+using namespace UE::Math;
+
 /**
  * TGeneralPolygon2 is a 2D polygon with holes
  */
@@ -168,7 +170,7 @@ public:
     }
 
 
-	void Translate(FVector2<T> translate) {
+	void Translate(TVector2<T> translate) {
 		Outer.Translate(translate);
 		for (TPolygon2<T>& Hole : Holes)
 		{
@@ -177,7 +179,7 @@ public:
 	}
 
 
-    void Scale(FVector2<T> scale, FVector2<T> origin) {
+    void Scale(TVector2<T> scale, TVector2<T> origin) {
 		Outer.Scale(scale, origin);
 		for (TPolygon2<T>& Hole : Holes)
 		{
@@ -185,7 +187,7 @@ public:
 		}
 	}
 
-    void Transform(const TFunction<FVector2<T> (const FVector2<T>&)>& TransformFunc)
+    void Transform(const TFunction<TVector2<T> (const TVector2<T>&)>& TransformFunc)
     {
         Outer.Transform(TransformFunc);
 		for (TPolygon2<T>& Hole : Holes)
@@ -205,7 +207,7 @@ public:
     }
 
 
-	bool Contains(FVector2<T> vTest)
+	bool Contains(TVector2<T> vTest)
 	{
 		if (Outer.Contains(vTest) == false)
 		{
@@ -254,7 +256,7 @@ public:
     }
 
 
-    FVector2<T> GetSegmentPoint(int iSegment, double fSegT, int iHoleIndex = -1)
+    TVector2<T> GetSegmentPoint(int iSegment, double fSegT, int iHoleIndex = -1)
 	{
 		if (iHoleIndex == -1)
 		{
@@ -272,7 +274,7 @@ public:
 		return Holes[iHoleIndex].Segment(iSegment);			
 	}
 
-	FVector2<T> GetNormal(int iSegment, double segT, int iHoleIndex = -1)
+	TVector2<T> GetNormal(int iSegment, double segT, int iHoleIndex = -1)
 	{
 		if (iHoleIndex == -1)
 		{
@@ -282,7 +284,7 @@ public:
 	}
 
 	// this should be more efficient when there are Holes...
-	double DistanceSquared(FVector2<T> p, int &iHoleIndex, int &iNearSeg, double &fNearSegT)
+	double DistanceSquared(TVector2<T> p, int &iHoleIndex, int &iNearSeg, double &fNearSegT)
 	{
 		iNearSeg = iHoleIndex = -1;
 		fNearSegT = TMathUtil<T>::MaxReal;
diff --git a/Engine/Source/Runtime/GeometryCore/Public/Distance/DistPoint3Triangle3.h b/Engine/Source/Runtime/GeometryCore/Public/Distance/DistPoint3Triangle3.h
--- a/Engine/Source/Runtime/GeometryCore/Public/Distance/DistPoint3Triangle3.h
+++ b/Engine/Source/Runtime/GeometryCore/Public/Distance/DistPoint3Triangle3.h
@@ -13,6 +13,8 @@ namespace UE
 namespace Geometry
 {
 
+using namespace UE::Math;
+
 /**
  * Compute unsigned distance between 3D Point and 3D Triangle
  */
@@ -59,7 +61,7 @@ public:
 		Real f10 = b0 + a00;
 		Real f01 = b0 + a01;
 
-		FVector2<Real> p0, p1, p;
+		TVector2<Real> p0, p1, p;
 		Real dt1, h0, h1;
 
 		// Compute the endpoints p0 and p1 of the segment.  The segment is
@@ -194,7 +196,7 @@ public:
 	}
 
 private:
-	void GetMinEdge02(Real const& a11, Real const& b1, FVector2<Real>& p) const
+	void GetMinEdge02(Real const& a11, Real const& b1, TVector2<Real>& p) const
 	{
 		p[0] = (Real)0;
 		if (b1 >= (Real)0)
@@ -213,7 +215,7 @@ private:
 
 	void GetMinEdge12(
 			Real const& a01, Real const& a11, Real const& b1, Real const& f10,
-			Real const& f01, FVector2<Real>& p) const
+			Real const& f01, TVector2<Real>& p) const
 	{
 		Real h0 = a01 + b1 - f10;
 		if (h0 >= (Real)0)
@@ -236,8 +238,8 @@ private:
 	}
 
 	void GetMinInterior(
-			FVector2<Real> const& p0, Real const& h0, FVector2<Real> const& p1,
-			Real const& h1, FVector2<Real>& p) const
+			TVector2<Real> const& p0, Real const& h0, TVector2<Real> const& p1,
+			Real const& h1, TVector2<Real>& p) const
 	{
 		Real z = h0 / (h0 - h1);
 		p = ((Real)1 - z) * p0 + z * p1;
diff --git a/Engine/Source/Runtime/GeometryCore/Public/FrameTypes.h b/Engine/Source/Runtime/GeometryCore/Public/FrameTypes.h
--- a/Engine/Source/Runtime/GeometryCore/Public/FrameTypes.h
+++ b/Engine/Source/Runtime/GeometryCore/Public/FrameTypes.h
@@ -14,6 +14,8 @@ namespace UE
 namespace Geometry
 {
 
+using namespace UE::Math;
+
 /**
  * TFrame3 is an object that represents an oriented 3D coordinate frame, ie orthogonal X/Y/Z axes at a point in space.
  * One can think of this Frame as a local coordinate space measured along these axes.
@@ -261,7 +263,7 @@ struct TFrame3
 	 * @param PlaneNormalAxis which plane to project onto, identified by perpendicular normal. Default is 2, ie normal is Z, plane is (X,Y)
 	 * @return 2D coordinates in UV plane, relative to origin
 	 */
-	FVector2<RealType> ToPlaneUV(const FVector3<RealType>& Pos, int PlaneNormalAxis = 2) const
+	TVector2<RealType> ToPlaneUV(const FVector3<RealType>& Pos, int PlaneNormalAxis = 2) const
 	{
 		int Axis0 = 0, Axis1 = 1;
 		if (PlaneNormalAxis == 0)
@@ -275,7 +277,7 @@ struct TFrame3
 		FVector3<RealType> LocalPos = Pos - Origin;
 		RealType U = LocalPos.Dot(GetAxis(Axis0));
 		RealType V = LocalPos.Dot(GetAxis(Axis1));
-		return FVector2<RealType>(U, V);
+		return TVector2<RealType>(U, V);
 	}
 
 
@@ -286,7 +288,7 @@ struct TFrame3
 	 * @param PlaneNormalAxis which plane to map to, identified by perpendicular normal. Default is 2, ie normal is Z, plane is (X,Y)
 	 * @return 3D coordinates in frame's plane (including Origin translation)
 	 */
-	FVector3<RealType> FromPlaneUV(const FVector2<RealType>& PosUV, int PlaneNormalAxis = 2) const
+	FVector3<RealType> FromPlaneUV(const TVector2<RealType>& PosUV, int PlaneNormalAxis = 2) const
 	{
 		FVector3<RealType> PlanePos(PosUV[0], PosUV[1], 0);
 		if (PlaneNormalAxis == 0)
@@ -412,7 +414,7 @@ struct TFrame3
 			FallbackAxis : UpAxis;
 
 		// figure out which PerpAxis is closer to target, and align that one to +/- TargetAxis (whichever is smaller rotation)
-		FVector2<RealType> Dots(GetAxis(PerpAxis1).Dot(TargetAxis), GetAxis(PerpAxis2).Dot(TargetAxis));
+		TVector2<RealType> Dots(GetAxis(PerpAxis1).Dot(TargetAxis), GetAxis(PerpAxis2).Dot(TargetAxis));
 		int UseAxis = (TMathUtil<RealType>::Abs(Dots.X) > TMathUtil<RealType>::Abs(Dots.Y)) ? 0 : 1;
 		RealType UseSign = Dots[UseAxis] < 0 ? -1 : 1;
 		ConstrainedAlignAxis(UseAxis, UseSign*TargetAxis, NormalVec);
diff --git a/Engine/Source/Runtime/GeometryCore/Public/Generators/SweepGenerator.h b/Engine/Source/Runtime/GeometryCore/Public/Generators/SweepGenerator.h
--- a/Engine/Source/Runtime/GeometryCore/Public/Generators/SweepGenerator.h
+++ b/Engine/Source/Runtime/GeometryCore/Public/Generators/SweepGenerator.h
@@ -75,7 +75,7 @@ protected:
 			PathPercentages.Add(0);
 			for (int XIdx = 0, XNum = CrossSection.VertexCount(); XIdx < XNum; XIdx++)
 			{
-				double SegLen = CrossSection[XIdx].Distance(CrossSection[(XIdx + 1) % XNum]);
+				double SegLen = Distance(CrossSection[XIdx], CrossSection[(XIdx + 1) % XNum]);
 				TotalPerimeter += SegLen;
 				CrossSectionPercentages.Add(TotalPerimeter);
 			}
@@ -380,7 +380,7 @@ public:
 		AlongPercents[0] = 0;
 		for (int XIdx = 0; XIdx+1 < NumX; XIdx++)
 		{
-			float Dist = FVector2d(Radii[XIdx], Heights[XIdx]).Distance(FVector2d(Radii[XIdx + 1], Heights[XIdx + 1]));
+			float Dist = Distance( FVector2d(Radii[XIdx], Heights[XIdx]), FVector2d(Radii[XIdx + 1], Heights[XIdx + 1]) );
 			LenAlong += Dist;
 			AlongPercents[XIdx + 1] = LenAlong;
 		}
diff --git a/Engine/Source/Runtime/GeometryCore/Public/Intersection/IntrLine2Line2.h b/Engine/Source/Runtime/GeometryCore/Public/Intersection/IntrLine2Line2.h
--- a/Engine/Source/Runtime/GeometryCore/Public/Intersection/IntrLine2Line2.h
+++ b/Engine/Source/Runtime/GeometryCore/Public/Intersection/IntrLine2Line2.h
@@ -13,6 +13,8 @@ namespace UE
 namespace Geometry
 {
 
+using namespace UE::Math;
+
 /**
  * Compute intersection between two 2D lines
  */
@@ -28,7 +30,7 @@ protected:
 	
 public:
 	// result data
-	FVector2<RealType> Point;
+	TVector2<RealType> Point;
 	RealType Segment1Parameter;
 	RealType Segment2Parameter;
 	int Quantity = 0;
@@ -118,7 +120,7 @@ public:
 			return false;
 		}
 
-		FVector2<RealType> s = FVector2<RealType>::Zero();
+		TVector2<RealType> s = TVector2<RealType>::Zero();
 		Type = Classify(Line1.Origin, Line1.Direction,
 			Line2.Origin, Line2.Direction, dotThresh, DistThresh, s);
 
@@ -146,9 +148,9 @@ public:
 
 
 	static EIntersectionType Classify(
-		const FVector2<RealType>& P0, const FVector2<RealType>& D0, 
-		const FVector2<RealType>& P1, const FVector2<RealType>& D1, 
-		RealType DotThreshold, RealType DistThreshold, FVector2<RealType>& s)
+		const TVector2<RealType>& P0, const TVector2<RealType>& D0, 
+		const TVector2<RealType>& P1, const TVector2<RealType>& D1, 
+		RealType DotThreshold, RealType DistThreshold, TVector2<RealType>& s)
 	{
 		// Ensure DotThreshold is nonnegative.
 		DotThreshold = FMath::Max(DotThreshold, (RealType)0);
@@ -161,7 +163,7 @@ public:
 		// produces the point of intersection.  Also,
 		//   s1 = Q.Dot(Perp(D0))/D0.Dot(Perp(D1))
 
-		FVector2<RealType> diff = P1 - P0;
+		TVector2<RealType> diff = P1 - P0;
 		RealType D0DotPerpD1 = DotPerp(D0, D1);
 		if (FMath::Abs(D0DotPerpD1) > DotThreshold) 
 		{
diff --git a/Engine/Source/Runtime/GeometryCore/Public/Intersection/IntrLine2Triangle2.h b/Engine/Source/Runtime/GeometryCore/Public/Intersection/IntrLine2Triangle2.h
--- a/Engine/Source/Runtime/GeometryCore/Public/Intersection/IntrLine2Triangle2.h
+++ b/Engine/Source/Runtime/GeometryCore/Public/Intersection/IntrLine2Triangle2.h
@@ -17,6 +17,8 @@ namespace UE
 namespace Geometry
 {
 
+using namespace UE::Math;
+
 /**
  * Compute intersection between 2D Line and 2D Triangle
  */
@@ -34,8 +36,8 @@ public:
 	EIntersectionResult Result = EIntersectionResult::NotComputed;
 	EIntersectionType Type = EIntersectionType::Empty;
 
-	FVector2<Real> Point0;
-	FVector2<Real> Point1;
+	TVector2<Real> Point0;
+	TVector2<Real> Point1;
 	double Param0;
 	double Param1;
 
@@ -106,7 +108,7 @@ public:
 		}
 		else
 		{
-			FVector2<Real> param;
+			TVector2<Real> param;
 			GetInterval(Line.Origin, Line.Direction, Triangle, Dist, Sign, param);
 
 			TIntersector1<Real> intr(param[0], param[1], -TMathUtil<Real>::MaxReal, +TMathUtil<Real>::MaxReal);
@@ -144,7 +146,7 @@ public:
 
 
 	static void TriangleLineRelations(
-		const FVector2<Real>& Origin, const FVector2<Real>& Direction, const TTriangle2<Real>& Tri,
+		const TVector2<Real>& Origin, const TVector2<Real>& Direction, const TTriangle2<Real>& Tri,
 		FVector3<Real>& Dist, FVector3i& Sign, int& Positive, int& Negative, int& Zero, Real Tolerance = TMathUtil<Real>::ZeroTolerance)
 	{
 		Positive = 0;
@@ -152,7 +154,7 @@ public:
 		Zero = 0;
 		for (int i = 0; i < 3; ++i)
 		{
-			FVector2<Real> diff = Tri.V[i] - Origin;
+			TVector2<Real> diff = Tri.V[i] - Origin;
 			Dist[i] = DotPerp(diff, Direction);
 			if (Dist[i] > Tolerance)
 			{
@@ -174,15 +176,15 @@ public:
 	}
 
 
-	static bool GetInterval(const FVector2<Real>& Origin, const FVector2<Real>& Direction, const TTriangle2<Real>& Tri,
-							const FVector3<Real>& Dist, const FVector3i& Sign, FVector2<Real>& param)
+	static bool GetInterval(const TVector2<Real>& Origin, const TVector2<Real>& Direction, const TTriangle2<Real>& Tri,
+							const FVector3<Real>& Dist, const FVector3i& Sign, TVector2<Real>& param)
 	{
 		// Project Triangle onto Line.
 		FVector3<Real> proj;
 		int i;
 		for (i = 0; i < 3; ++i)
 		{
-			FVector2<Real> diff = Tri.V[i] - Origin;
+			TVector2<Real> diff = Tri.V[i] - Origin;
 			proj[i] = Direction.Dot(diff);
 		}
 
diff --git a/Engine/Source/Runtime/GeometryCore/Public/Intersection/IntrSegment2Segment2.h b/Engine/Source/Runtime/GeometryCore/Public/Intersection/IntrSegment2Segment2.h
--- a/Engine/Source/Runtime/GeometryCore/Public/Intersection/IntrSegment2Segment2.h
+++ b/Engine/Source/Runtime/GeometryCore/Public/Intersection/IntrSegment2Segment2.h
@@ -17,6 +17,8 @@ namespace UE
 namespace Geometry
 {
 
+using namespace UE::Math;
+
 // ported from WildMagic5
 //
 //
@@ -66,8 +68,8 @@ public:
 	EIntersectionResult Result = EIntersectionResult::NotComputed;
 	EIntersectionType Type = EIntersectionType::Empty;
 	// these values are all on segment 1, unlike many other tests!!
-	FVector2<RealType> Point0;
-	FVector2<RealType> Point1;     // only set if Quantity == 2, ie segment overlap
+	TVector2<RealType> Point0;
+	TVector2<RealType> Point1;     // only set if Quantity == 2, ie segment overlap
 
 	RealType Parameter0;
 	RealType Parameter1;     // only set if Quantity == 2, ie segment overlap
@@ -154,7 +156,7 @@ public:
 		}
 
 
-		FVector2<RealType> s = FVector2<RealType>::Zero();
+		TVector2<RealType> s = TVector2<RealType>::Zero();
 		Type = TIntrLine2Line2<RealType>::Classify(Segment1.Center, Segment1.Direction,
 			Segment2.Center, Segment2.Direction,
 			DotThreshold, IntervalThreshold, s);
@@ -178,7 +180,7 @@ public:
 		else if (Type == EIntersectionType::Line)
 		{
 			// Compute the location of Segment1 endpoints relative to segment0.
-			FVector2<RealType> diff = Segment2.Center - Segment1.Center;
+			TVector2<RealType> diff = Segment2.Center - Segment1.Center;
 			RealType t1 = Segment1.Direction.Dot(diff);
 			RealType tmin = t1 - Segment2.Extent;
 			RealType tmax = t1 + Segment2.Extent;
diff --git a/Engine/Source/Runtime/GeometryCore/Public/Intersection/IntrSegment2Triangle2.h b/Engine/Source/Runtime/GeometryCore/Public/Intersection/IntrSegment2Triangle2.h
--- a/Engine/Source/Runtime/GeometryCore/Public/Intersection/IntrSegment2Triangle2.h
+++ b/Engine/Source/Runtime/GeometryCore/Public/Intersection/IntrSegment2Triangle2.h
@@ -17,6 +17,8 @@ namespace UE
 namespace Geometry
 {
 
+using namespace UE::Math;
+
 /**
  * Compute intersection between 2D segment and 2D triangle.
  * Note that if Segment.Extent is zero, this will test if the Segment Center is in the triangle.
@@ -41,8 +43,8 @@ public:
 	}
 
 
-	FVector2<Real> Point0;
-	FVector2<Real> Point1;
+	TVector2<Real> Point0;
+	TVector2<Real> Point1;
 	double Param0;
 	double Param1;
 
@@ -93,17 +95,17 @@ public:
 			int pos = 0, neg = 0;
 			for (int TriPrev = 2, TriIdx = 0; TriIdx < 3; TriPrev = TriIdx++)
 			{
-				FVector2<Real> ToPt = Segment.Center - Triangle.V[TriIdx];
-				FVector2<Real> EdgePerp = PerpCW(Triangle.V[TriIdx] - Triangle.V[TriPrev]);
+				TVector2<Real> ToPt = Segment.Center - Triangle.V[TriIdx];
+				TVector2<Real> EdgePerp = PerpCW(Triangle.V[TriIdx] - Triangle.V[TriPrev]);
 				Real EdgeLen = Normalize(EdgePerp);
 				if (EdgeLen == 0) // triangle is just a line segment; try edgeperp = one of the other edges
 				{
-					FVector2<Real> OtherV = Triangle.V[(TriIdx + 1) % 3];
+					TVector2<Real> OtherV = Triangle.V[(TriIdx + 1) % 3];
 					EdgePerp = OtherV - Triangle.V[TriIdx];
 					EdgeLen = Normalize(EdgePerp);
 					if (EdgeLen == 0) // triangle is just a point; go by distance between vertex and center
 					{
-						if (Triangle.V[0].DistanceSquared(Segment.Center) <= Tolerance*Tolerance)
+						if ( DistanceSquared(Triangle.V[0], Segment.Center) <= Tolerance*Tolerance)
 						{
 							pos = neg = 0; // closer than tolerance; accept collision
 						}
@@ -115,8 +117,8 @@ public:
 					}
 					else // valid edge, test from the other side of the degenerate tri as well
 					{
-						FVector2<Real> ToPtFromOther = Segment.Center - OtherV;
-						FVector2<Real> BackwardsEdgePerp = -EdgePerp;
+						TVector2<Real> ToPtFromOther = Segment.Center - OtherV;
+						TVector2<Real> BackwardsEdgePerp = -EdgePerp;
 						Real OtherSideSign = BackwardsEdgePerp.Dot(ToPtFromOther);
 						if (OtherSideSign < -Tolerance)
 						{
@@ -177,7 +179,7 @@ public:
 		}
 		else 
 		{
-			FVector2<Real> param;
+			TVector2<Real> param;
 			TIntrLine2Triangle2<Real>::GetInterval(Segment.Center, Segment.Direction, Triangle, dist, sign, param);
 
 			TIntersector1<Real> intr(param[0], param[1], -Segment.Extent, +Segment.Extent);
diff --git a/Engine/Source/Runtime/GeometryCore/Public/Intersection/IntrTriangle2Triangle2.h b/Engine/Source/Runtime/GeometryCore/Public/Intersection/IntrTriangle2Triangle2.h
--- a/Engine/Source/Runtime/GeometryCore/Public/Intersection/IntrTriangle2Triangle2.h
+++ b/Engine/Source/Runtime/GeometryCore/Public/Intersection/IntrTriangle2Triangle2.h
@@ -13,6 +13,8 @@ namespace UE
 namespace Geometry
 {
 
+using namespace UE::Math;
+
 /**
  * Compute intersection between 2D triangles
  */
@@ -36,7 +38,7 @@ public:
 
 	// intersection polygon - this array will always be 6 elements long,
 	// however only the first Quantity vertices will be valid
-	FVector2<Real> Points[6];
+	TVector2<Real> Points[6];
 
 	TIntrTriangle2Triangle2()
 	{}
@@ -66,7 +68,7 @@ public:
 	bool Test()
 	{
 		int i0, i1;
-		FVector2<Real> dir;
+		TVector2<Real> dir;
 
 		// Test edges of Triangle0 for separation.
 		for (i0 = 0, i1 = 2; i0 < 3; i1 = i0++)
@@ -125,7 +127,7 @@ public:
 		for (int i1 = 2, i0 = 0; i0 < 3; i1 = i0++)
 		{
 			// Clip against edge <V0[i1],V0[i0]>.
-			FVector2<Real> N(
+			TVector2<Real> N(
 				Triangle0.V[i1].Y - Triangle0.V[i0].Y,
 				Triangle0.V[i0].X - Triangle0.V[i1].X);
 			double c = N.Dot(Triangle0.V[i1]);
@@ -157,7 +159,7 @@ public:
 
 
 
-	static int WhichSide(const TTriangle2<Real>& V, const FVector2<Real>& P, const FVector2<Real>& D)
+	static int WhichSide(const TTriangle2<Real>& V, const TVector2<Real>& P, const TVector2<Real>& D)
 	{
 		// Vertices are projected to the form P+t*D.  Return value is +1 if all
 		// t > 0, -1 if all t < 0, 0 otherwise, in which case the line splits the
@@ -192,7 +194,7 @@ public:
 private:
 	// quantity, V initially are input polygon vertex count and vertices;
 	// on return, they are the clipped polygon vertex count and vertices
-	static void ClipConvexPolygonAgainstLine(const FVector2<Real>& N, double c, int& quantity, FVector2<Real> V[6])
+	static void ClipConvexPolygonAgainstLine(const TVector2<Real>& N, double c, int& quantity, TVector2<Real> V[6])
 	{
 		// The input vertices are assumed to be in counterclockwise order.  The
 		// ordering is an invariant of this function.
@@ -223,7 +225,7 @@ private:
 			if (negative > 0)
 			{
 				// Line transversely intersects polygon.
-				FVector2<Real> CV[6];
+				TVector2<Real> CV[6];
 				int cQuantity = 0, cur, prv;
 				double t;
 
diff --git a/Engine/Source/Runtime/GeometryCore/Public/Intersection/IntrTriangle3Triangle3.h b/Engine/Source/Runtime/GeometryCore/Public/Intersection/IntrTriangle3Triangle3.h
--- a/Engine/Source/Runtime/GeometryCore/Public/Intersection/IntrTriangle3Triangle3.h
+++ b/Engine/Source/Runtime/GeometryCore/Public/Intersection/IntrTriangle3Triangle3.h
@@ -19,6 +19,8 @@ namespace UE
 namespace Geometry
 {
 
+using namespace UE::Math;
+
 /**
  * Compute intersection between 3D triangles
  * use Test() for fast boolean query, does not compute intersection info
@@ -564,7 +566,7 @@ public:
 		}
 
 		TTriangle2<Real> projTri;
-		FVector2<Real> projEnd0, projEnd1;
+		TVector2<Real> projEnd0, projEnd1;
 		int i;
 
 		if (maxNormal == 0)
@@ -613,7 +615,7 @@ public:
 
 		int Quantity = 0;
 
-		FVector2<Real> intr[2];
+		TVector2<Real> intr[2];
 		if (calc.Type == EIntersectionType::Segment)
 		{
 			Quantity = 2;
@@ -692,10 +694,10 @@ protected:
 		FVector3<Real> V2mV0 = triangle.V[2] - triangle.V[0];
 
 		// The planar representation of P-V0.
-		FVector2<Real> ProjP(U0.Dot(PmV0), U1.Dot(PmV0));
+		TVector2<Real> ProjP(U0.Dot(PmV0), U1.Dot(PmV0));
 
 		// The planar representation of the triangle <V0-V0,V1-V0,V2-V0>.
-		TTriangle2<Real> ProjT(FVector2<Real>::Zero(), FVector2<Real>(U0.Dot(V1mV0), U1.Dot(V1mV0)), FVector2<Real>(U0.Dot(V2mV0), U1.Dot(V2mV0)));
+		TTriangle2<Real> ProjT(TVector2<Real>::Zero(), TVector2<Real>(U0.Dot(V1mV0), U1.Dot(V1mV0)), TVector2<Real>(U0.Dot(V2mV0), U1.Dot(V2mV0)));
 
 		// Test whether P-V0 is in the triangle <0,V1-V0,V2-V0>.
 		if (ProjT.IsInsideOrOn_Oriented(ProjP) <= 0)
@@ -781,9 +783,9 @@ protected:
 		}
 
 		// 2D triangle intersection routines require counterclockwise ordering.
-		FVector2<Real> save;
-		FVector2<Real> edge0 = projTri0.V[1] - projTri0.V[0];
-		FVector2<Real> edge1 = projTri0.V[2] - projTri0.V[0];
+		TVector2<Real> save;
+		TVector2<Real> edge0 = projTri0.V[1] - projTri0.V[0];
+		TVector2<Real> edge1 = projTri0.V[2] - projTri0.V[0];
 		if (DotPerp(edge0, edge1) < (Real)0)
 		{
 			// Triangle is clockwise, reorder it.
diff --git a/Engine/Source/Runtime/GeometryCore/Public/LineTypes.h b/Engine/Source/Runtime/GeometryCore/Public/LineTypes.h
--- a/Engine/Source/Runtime/GeometryCore/Public/LineTypes.h
+++ b/Engine/Source/Runtime/GeometryCore/Public/LineTypes.h
@@ -12,6 +12,8 @@ namespace UE
 namespace Geometry
 {
 
+using namespace UE::Math;
+
 /**
  * TLine2 is a two-dimensional infinite line.
  * The line is stored in (Center,Direction) form.
@@ -20,24 +22,24 @@ template<typename T>
 struct TLine2
 {
 	/** Origin / Center Point of Line */
-	FVector2<T> Origin;
+	TVector2<T> Origin;
 
 	/** Direction of Line, Normalized */
-	FVector2<T> Direction;
+	TVector2<T> Direction;
 
 	/**
 	 * Construct default line along X axis
 	 */
 	TLine2()
 	{
-		Origin = FVector2<T>::Zero();
-		Direction = FVector2<T>::UnitX();
+		Origin = TVector2<T>::Zero();
+		Direction = TVector2<T>::UnitX();
 	}
 
 	/**
 	 * Construct line with given Origin and Direction
 	 */
-	TLine2(const FVector2<T>& OriginIn, const FVector2<T>& DirectionIn)
+	TLine2(const TVector2<T>& OriginIn, const TVector2<T>& DirectionIn)
 		: Origin(OriginIn), Direction(DirectionIn)
 	{
 	}
@@ -46,7 +48,7 @@ struct TLine2
 	/**
 	 * @return line between two points
 	 */
-	static TLine2<T> FromPoints(const FVector2<T>& Point0, const FVector2<T>& Point1)
+	static TLine2<T> FromPoints(const TVector2<T>& Point0, const TVector2<T>& Point1)
 	{
 		return TLine2<T>(Point0, Normalized(Point1 - Point0) );
 	}
@@ -55,7 +57,7 @@ struct TLine2
 	/**
 	 * @return point on line at given line parameter value (distance along line from origin)
 	 */
-	inline FVector2<T> PointAt(T LineParameter) const
+	inline TVector2<T> PointAt(T LineParameter) const
 	{
 		return Origin + LineParameter * Direction;
 	}
@@ -64,7 +66,7 @@ struct TLine2
 	/**
 	 * @return line parameter (ie distance from Origin) at nearest point on line to QueryPoint
 	 */
-	inline T Project(const FVector2<T>& QueryPoint) const
+	inline T Project(const TVector2<T>& QueryPoint) const
 	{
 		return (QueryPoint - Origin).Dot(Direction);
 	}
@@ -72,17 +74,17 @@ struct TLine2
 	/**
 	 * @return smallest squared distance from line to QueryPoint
 	 */
-	inline T DistanceSquared(const FVector2<T>& QueryPoint) const
+	inline T DistanceSquared(const TVector2<T>& QueryPoint) const
 	{
 		T ParameterT = (QueryPoint - Origin).Dot(Direction);
-		FVector2<T> proj = Origin + ParameterT * Direction;
+		TVector2<T> proj = Origin + ParameterT * Direction;
 		return (proj - QueryPoint).SquaredLength();
 	}
 
 	/**
 	 * @return nearest point on line to QueryPoint
 	 */
-	inline FVector2<T> NearestPoint(const FVector2<T>& QueryPoint) const
+	inline TVector2<T> NearestPoint(const TVector2<T>& QueryPoint) const
 	{
 		T ParameterT = (QueryPoint - Origin).Dot(Direction);
 		return Origin + ParameterT * Direction;
@@ -92,7 +94,7 @@ struct TLine2
 	/**
 	 * @return +1 if QueryPoint is "right" of line, -1 if "left" or 0 if "on" line (up to given tolerance)
 	 */
-	inline int WhichSide(const FVector2<T>& QueryPoint, T OnLineTolerance = 0) const
+	inline int WhichSide(const TVector2<T>& QueryPoint, T OnLineTolerance = 0) const
 	{
 		T x0 = QueryPoint.X - Origin.X;
 		T y0 = QueryPoint.Y - Origin.Y;
@@ -110,10 +112,10 @@ struct TLine2
 	 * @param ParallelDotTolerance tolerance used to determine if lines are parallel (and hence cannot intersect)
 	 * @return true if lines intersect and IntersectionPointOut was computed
 	 */
-	bool IntersectionPoint(const TLine2<T>& OtherLine, FVector2<T>& IntersectionPointOut, T ParallelDotTolerance = TMathUtil<T>::ZeroTolerance) const
+	bool IntersectionPoint(const TLine2<T>& OtherLine, TVector2<T>& IntersectionPointOut, T ParallelDotTolerance = TMathUtil<T>::ZeroTolerance) const
 	{
 		// see IntrTLine2TLine2 for more detailed explanation
-		FVector2<T> diff = OtherLine.Origin - Origin;
+		TVector2<T> diff = OtherLine.Origin - Origin;
 		T D0DotPerpD1 = DotPerp(Direction, OtherLine.Direction);
 		if (TMathUtil<T>::Abs(D0DotPerpD1) > ParallelDotTolerance)                     // TLines intersect in a single point.
 		{
diff --git a/Engine/Source/Runtime/GeometryCore/Public/MatrixTypes.h b/Engine/Source/Runtime/GeometryCore/Public/MatrixTypes.h
--- a/Engine/Source/Runtime/GeometryCore/Public/MatrixTypes.h
+++ b/Engine/Source/Runtime/GeometryCore/Public/MatrixTypes.h
@@ -10,6 +10,8 @@ namespace UE
 namespace Geometry
 {
 
+using namespace UE::Math;
+
 template <typename RealType>
 struct TMatrix3
 {
@@ -271,8 +273,8 @@ struct TMatrix3
 template <typename RealType>
 struct TMatrix2
 {
-	FVector2<RealType> Row0;
-	FVector2<RealType> Row1;
+	TVector2<RealType> Row0;
+	TVector2<RealType> Row1;
 
 	TMatrix2()
 	{
@@ -280,20 +282,20 @@ struct TMatrix2
 
 	TMatrix2(RealType ConstantValue)
 	{
-		Row0 = Row1 = FVector2<RealType>(ConstantValue, ConstantValue);
+		Row0 = Row1 = TVector2<RealType>(ConstantValue, ConstantValue);
 	}
 
 	TMatrix2(RealType Diag0, RealType Diag1)
 	{
-		Row0 = FVector2<RealType>(Diag0, 0);
-		Row1 = FVector2<RealType>(0, Diag1);
+		Row0 = TVector2<RealType>(Diag0, 0);
+		Row1 = TVector2<RealType>(0, Diag1);
 	}
 
 	/**
 	 * Construct outer-product of U*transpose(V) of U and V
 	 * result is that Mij = u_i * v_j
 	 */
-	TMatrix2<RealType>(const FVector2<RealType>& U, const FVector2<RealType>& V)
+	TMatrix2<RealType>(const TVector2<RealType>& U, const TVector2<RealType>& V)
 		: Row0(U.X * V.X, U.X * V.Y),
 		  Row1(U.Y * V.X, U.Y * V.Y)
 	{
@@ -305,7 +307,7 @@ struct TMatrix2
 	{
 	}
 
-	TMatrix2(const FVector2<RealType>& V1, const FVector2<RealType>& V2, bool bRows)
+	TMatrix2(const TVector2<RealType>& V1, const TVector2<RealType>& V2, bool bRows)
 	{
 		if (bRows)
 		{
@@ -314,8 +316,8 @@ struct TMatrix2
 		}
 		else
 		{
-			Row0 = FVector2<RealType>(V1.X, V2.X);
-			Row1 = FVector2<RealType>(V1.Y, V2.Y);
+			Row0 = TVector2<RealType>(V1.X, V2.X);
+			Row1 = TVector2<RealType>(V1.Y, V2.Y);
 		}
 	}
 
@@ -348,9 +350,9 @@ struct TMatrix2
 			Row1.X * Scale, Row1.Y * Scale);
 	}
 
-	FVector2<RealType> operator*(const FVector2<RealType>& V) const
+	TVector2<RealType> operator*(const TVector2<RealType>& V) const
 	{
-		return FVector2<RealType>(
+		return TVector2<RealType>(
 			Row0.X * V.X + Row0.Y * V.Y,
 			Row1.X * V.X + Row1.Y * V.Y);
 	}
diff --git a/Engine/Source/Runtime/GeometryCore/Public/Polygon2.h b/Engine/Source/Runtime/GeometryCore/Public/Polygon2.h
--- a/Engine/Source/Runtime/GeometryCore/Public/Polygon2.h
+++ b/Engine/Source/Runtime/GeometryCore/Public/Polygon2.h
@@ -19,6 +19,8 @@ namespace UE
 namespace Geometry
 {
 
+using namespace UE::Math;
+
 /**
  * TPolygon2 is a 2D polygon represented as a list of Vertices.
  * 
@@ -29,7 +31,7 @@ class TPolygon2
 {
 protected:
 	/** The list of vertices/corners of the polygon */
-	TArray<FVector2<T>> Vertices;
+	TArray<TVector2<T>> Vertices;
 
 	/** A counter that is incremented every time the polygon vertices are modified */
 	int Timestamp;
@@ -50,14 +52,25 @@ public:
 	/**
 	 * Construct polygon with given list of vertices
 	 */
-	TPolygon2(const TArray<FVector2<T>>& VertexList) : Vertices(VertexList), Timestamp(0)
+	TPolygon2(const TArray<TVector2<T>>& VertexList) : Vertices(VertexList), Timestamp(0)
 	{
 	}
 
+	template<typename OtherVertexType>
+	TPolygon2(const TArray<OtherVertexType>& VertexList) : Timestamp(0)
+	{
+		Vertices.Reserve(VertexList.Num());
+		for (const OtherVertexType& OtherVtx : VertexList)
+		{
+			Vertices.Add( TVector2<T>(OtherVtx.X, OtherVtx.Y) );
+		}
+	}
+
+
 	/**
 	 * Construct polygon with given indices into a vertex array
 	 */
-	TPolygon2(TArrayView<const FVector2<T>> VertexArray, TArrayView<const int32> VertexIndices) : Timestamp(0)
+	TPolygon2(TArrayView<const TVector2<T>> VertexArray, TArrayView<const int32> VertexIndices) : Timestamp(0)
 	{
 		Vertices.SetNum(VertexIndices.Num());
 		for (int32 Idx = 0; Idx < VertexIndices.Num(); Idx++)
@@ -75,7 +88,7 @@ public:
 	/**
 	 * Get the vertex at a given index
 	 */
-	const FVector2<T>& operator[](int Index) const
+	const TVector2<T>& operator[](int Index) const
 	{
 		return Vertices[Index];
 	}
@@ -84,7 +97,7 @@ public:
 	 * Get the vertex at a given index
 	 * @warning changing the vertex via this operator does not update Timestamp!
 	 */
-	FVector2<T>& operator[](int Index)
+	TVector2<T>& operator[](int Index)
 	{
 		return Vertices[Index];
 	}
@@ -93,7 +106,7 @@ public:
 	/**
 	 * @return first vertex of Polygon
 	 */
-	const FVector2<T>& Start() const
+	const TVector2<T>& Start() const
 	{
 		return Vertices[0];
 	}
@@ -101,7 +114,7 @@ public:
 	/**
 	 * @return list of Vertices of Polygon
 	 */
-	const TArray<FVector2<T>>& GetVertices() const
+	const TArray<TVector2<T>>& GetVertices() const
 	{
 		return Vertices;
 	}
@@ -117,7 +130,7 @@ public:
 	/**
 	 * Add a vertex to the Polygon
 	 */
-	void AppendVertex(const FVector2<T>& Position)
+	void AppendVertex(const TVector2<T>& Position)
 	{
 		Vertices.Add(Position);
 		Timestamp++;
@@ -126,7 +139,7 @@ public:
 	/**
 	 * Add a list of Vertices to the Polygon
 	 */
-	void AppendVertices(const TArray<FVector2<T>>& NewVertices)
+	void AppendVertices(const TArray<TVector2<T>>& NewVertices)
 	{
 		Vertices.Append(NewVertices);
 		Timestamp++;
@@ -135,7 +148,7 @@ public:
 	/**
 	 * Set vertex at given index to a new Position
 	 */
-	void Set(int VertexIndex, const FVector2<T>& Position)
+	void Set(int VertexIndex, const TVector2<T>& Position)
 	{
 		Vertices[VertexIndex] = Position;
 		Timestamp++;
@@ -153,7 +166,7 @@ public:
 	/**
 	 * Replace the list of Vertices with a new list
 	 */
-	void SetVertices(const TArray<FVector2<T>>& NewVertices)
+	void SetVertices(const TArray<TVector2<T>>& NewVertices)
 	{
 		Vertices = NewVertices;
 		Timestamp++;
@@ -178,10 +191,10 @@ public:
 	 * Get the tangent vector at a vertex of the polygon, which is the normalized
 	 * vector from the previous vertex to the next vertex
 	 */
-	FVector2<T> GetTangent(int VertexIndex) const
+	TVector2<T> GetTangent(int VertexIndex) const
 	{
-		FVector2<T> next = Vertices[(VertexIndex + 1) % Vertices.Num()];
-		FVector2<T> prev = Vertices[VertexIndex == 0 ? Vertices.Num() - 1 : VertexIndex - 1];
+		TVector2<T> next = Vertices[(VertexIndex + 1) % Vertices.Num()];
+		TVector2<T> prev = Vertices[VertexIndex == 0 ? Vertices.Num() - 1 : VertexIndex - 1];
 		return Normalized(next - prev);
 	}
 
@@ -190,7 +203,7 @@ public:
 	 * Get the normal vector at a vertex of the polygon, which is perpendicular to GetTangent()
 	 * Points "inward" for a Clockwise Polygon, and outward for CounterClockwise
 	 */
-	FVector2<T> GetNormal(int VertexIndex) const
+	TVector2<T> GetNormal(int VertexIndex) const
 	{
 		return PerpCW(GetTangent(VertexIndex));
 	}
@@ -201,14 +214,14 @@ public:
 	 * This vector is independent of the lengths of the adjacent segments.
 	 * Points "inward" for a Clockwise Polygon, and outward for CounterClockwise
 	 */
-	FVector2<T> GetNormal_FaceAvg(int VertexIndex) const
+	TVector2<T> GetNormal_FaceAvg(int VertexIndex) const
 	{
-		FVector2<T> next = Vertices[(VertexIndex + 1) % Vertices.Num()];
-		FVector2<T> prev = Vertices[VertexIndex == 0 ? Vertices.Num() - 1 : VertexIndex - 1];
+		TVector2<T> next = Vertices[(VertexIndex + 1) % Vertices.Num()];
+		TVector2<T> prev = Vertices[VertexIndex == 0 ? Vertices.Num() - 1 : VertexIndex - 1];
 		next -= Vertices[VertexIndex]; Normalize(next);
 		prev -= Vertices[VertexIndex]; Normalize(prev);
 
-		FVector2<T> n = (PerpCW(next) - PerpCW(prev));
+		TVector2<T> n = (PerpCW(next) - PerpCW(prev));
 		T len = Normalize(n);
 		if (len == 0) 
 		{
@@ -317,8 +330,8 @@ public:
 		}
 		for (int i = 0; i < N; ++i) 
 		{
-			const FVector2<T>& v1 = Vertices[i];
-			const FVector2<T>& v2 = Vertices[(i + 1) % N];
+			const TVector2<T>& v1 = Vertices[i];
+			const TVector2<T>& v2 = Vertices[(i + 1) % N];
 			fArea += v1.X * v2.Y - v1.Y * v2.X;
 		}
 		return fArea * 0.5;
@@ -341,7 +354,7 @@ public:
 		int N = Vertices.Num();
 		for (int i = 0; i < N; ++i)
 		{
-			fPerim += Vertices[i].Distance(Vertices[(i + 1) % N]);
+			fPerim += Distance(Vertices[i], Vertices[(i + 1) % N]);
 		}
 		return fPerim;
 	}
@@ -350,7 +363,7 @@ public:
 	/**
 	 * Get the previous and next vertex positions for a given vertex of the Polygon
 	 */
-	void NeighbourPoints(int iVertex, FVector2<T> &PrevNbrOut, FVector2<T> &NextNbrOut) const
+	void NeighbourPoints(int iVertex, TVector2<T> &PrevNbrOut, TVector2<T> &NextNbrOut) const
 	{
 		int N = Vertices.Num();
 		PrevNbrOut = Vertices[(iVertex == 0) ? N - 1 : iVertex - 1];
@@ -361,7 +374,7 @@ public:
 	/**
 	 * Get the vectors from a given vertex to the previous and next Vertices, optionally normalized
 	 */
-	void NeighbourVectors(int iVertex, FVector2<T> &ToPrevOut, FVector2<T> &ToNextOut, bool bNormalize = false) const
+	void NeighbourVectors(int iVertex, TVector2<T> &ToPrevOut, TVector2<T> &ToNextOut, bool bNormalize = false) const
 	{
 		int N = Vertices.Num();
 		ToPrevOut = Vertices[(iVertex == 0) ? N - 1 : iVertex - 1] - Vertices[iVertex];
@@ -379,7 +392,7 @@ public:
 	 */
 	T OpeningAngleDeg(int iVertex) const
 	{
-		FVector2<T> e0, e1;
+		TVector2<T> e0, e1;
 		NeighbourVectors(iVertex, e0, e1, true);
 		return AngleD(e0, e1);
 	}
@@ -388,11 +401,11 @@ public:
 	/**
 	 * @return analytic winding integral for this Polygon at an arbitrary point
 	 */
-	T WindingIntegral(const FVector2<T>& QueryPoint) const
+	T WindingIntegral(const TVector2<T>& QueryPoint) const
 	{
 		T sum = 0;
 		int N = Vertices.Num();
-		FVector2<T> a = Vertices[0] - QueryPoint, b = FVector2<T>::Zero();
+		TVector2<T> a = Vertices[0] - QueryPoint, b = TVector2<T>::Zero();
 		for (int i = 0; i < N; ++i) 
 		{
 			b = Vertices[(i + 1) % N] - QueryPoint;
@@ -406,7 +419,7 @@ public:
 	/**
 	 * @return true if the given query point is inside the Polygon, based on the winding integral
 	 */
-	bool Contains(const FVector2<T>& QueryPoint) const
+	bool Contains(const TVector2<T>& QueryPoint) const
 	{
 		int nWindingNumber = 0;
 
@@ -415,7 +428,7 @@ public:
 		{
 			return false;
 		}
-		FVector2<T> a = Vertices[0], b = FVector2<T>::Zero();
+		TVector2<T> a = Vertices[0], b = TVector2<T>::Zero();
 		for (int i = 0; i < N; ++i) 
 		{
 			b = Vertices[(i + 1) % N];
@@ -584,7 +597,7 @@ public:
 	 * @param OutArray intersection points are stored here
 	 * @return true if any intersections were found
 	 */
-	bool FindIntersections(const TPolygon2<T>& OtherPoly, TArray<FVector2<T>>& OutArray) const
+	bool FindIntersections(const TPolygon2<T>& OtherPoly, TArray<TVector2<T>>& OutArray) const
 	{
 		if (!Bounds().Intersects(OtherPoly.Bounds()))
 		{
@@ -632,7 +645,7 @@ public:
 	 * @param SegmentParam parameter in range [-Extent,Extent] along segment
 	 * @return point on the segment at the given parameter value
 	 */
-	FVector2<T> GetSegmentPoint(int SegmentIndex, T SegmentParam) const
+	TVector2<T> GetSegmentPoint(int SegmentIndex, T SegmentParam) const
 	{
 		TSegment2<T> seg(Vertices[SegmentIndex], Vertices[(SegmentIndex + 1) % Vertices.Num()]);
 		return seg.PointAt(SegmentParam);
@@ -643,7 +656,7 @@ public:
 	 * @param SegmentParam parameter in range [0,1] along segment
 	 * @return point on the segment at the given parameter value
 	 */
-	FVector2<T> GetSegmentPointUnitParam(int SegmentIndex, T SegmentParam) const
+	TVector2<T> GetSegmentPointUnitParam(int SegmentIndex, T SegmentParam) const
 	{
 		TSegment2<T> seg(Vertices[SegmentIndex], Vertices[(SegmentIndex + 1) % Vertices.Num()]);
 		return seg.PointBetween(SegmentParam);
@@ -655,13 +668,13 @@ public:
 	 * @param SegmentParam parameter in range [0,1] along segment
 	 * @return interpolated normal to the segment at the given parameter value
 	 */
-	FVector2<T> GetNormal(int iSeg, T SegmentParam) const
+	TVector2<T> GetNormal(int iSeg, T SegmentParam) const
 	{
 		TSegment2<T> seg(Vertices[iSeg], Vertices[(iSeg + 1) % Vertices.Num()]);
 		T t = ((SegmentParam / seg.Extent) + 1.0) / 2.0;
 
-		FVector2<T> n0 = GetNormal(iSeg);
-		FVector2<T> n1 = GetNormal((iSeg + 1) % Vertices.Num());
+		TVector2<T> n0 = GetNormal(iSeg);
+		TVector2<T> n1 = GetNormal((iSeg + 1) % Vertices.Num());
 		return Normalized((T(1) - t) * n0 + t * n1);
 	}
 
@@ -674,7 +687,7 @@ public:
 	 * @param NearestSegParamOut the parameter value of the nearest point on the segment
 	 * @return squared distance to the polygon
 	 */
-	T DistanceSquared(const FVector2<T>& QueryPoint, int& NearestSegIndexOut, T& NearestSegParamOut) const
+	T DistanceSquared(const TVector2<T>& QueryPoint, int& NearestSegIndexOut, T& NearestSegParamOut) const
 	{
 		NearestSegIndexOut = -1;
 		NearestSegParamOut = TNumericLimits<T>::Max();
@@ -688,11 +701,11 @@ public:
 			T d = TNumericLimits<T>::Max();
 			if (t >= seg.Extent)
 			{
-				d = seg.EndPoint().DistanceSquared(QueryPoint);
+				d = UE::Geometry::DistanceSquared(seg.EndPoint(), QueryPoint);
 			}
 			else if (t <= -seg.Extent)
 			{
-				d = seg.StartPoint().DistanceSquared(QueryPoint);
+				d = UE::Geometry::DistanceSquared(seg.StartPoint(), QueryPoint);
 			}
 			else
 			{
@@ -714,7 +727,7 @@ public:
 	 * @param QueryPoint the query point
 	 * @return squared distance to the polygon
 	 */
-	T DistanceSquared(const FVector2<T>& QueryPoint) const
+	T DistanceSquared(const TVector2<T>& QueryPoint) const
 	{
 		int seg; T segt;
 		return DistanceSquared(QueryPoint, seg, segt);
@@ -728,9 +741,9 @@ public:
 	{
 		T avg = 0; int N = Vertices.Num();
 		for (int i = 1; i < N; ++i) {
-			avg += Vertices[i].Distance(Vertices[i - 1]);
+			avg += Distance(Vertices[i], Vertices[i - 1]);
 		}
-		avg += Vertices[N - 1].Distance(Vertices[0]);
+		avg += Distance(Vertices[N - 1], Vertices[0]);
 		return avg / N;
 	}
 
@@ -739,7 +752,7 @@ public:
 	 * Translate the polygon
 	 * @returns the Polygon, so that you can chain calls like Translate().Scale()
 	 */
-	TPolygon2<T>& Translate(const FVector2<T>& Translate) 
+	TPolygon2<T>& Translate(const TVector2<T>& Translate) 
 	{
 		int N = Vertices.Num();
 		for (int i = 0; i < N; ++i)
@@ -754,7 +767,7 @@ public:
 	 * Scale the Polygon relative to a given point
 	 * @returns the Polygon, so that you can chain calls like Translate().Scale()
 	 */
-	TPolygon2<T>& Scale(const FVector2<T>& Scale, const FVector2<T>& Origin)
+	TPolygon2<T>& Scale(const TVector2<T>& Scale, const TVector2<T>& Origin)
 	{
 		int N = Vertices.Num();
 		for (int i = 0; i < N; ++i)
@@ -770,7 +783,7 @@ public:
 	 * Apply an arbitrary transformation to the Polygon
 	 * @returns the Polygon, so that you can chain calls like Translate().Scale()
 	 */
-	TPolygon2<T>& Transform(const TFunction<FVector2<T> (const FVector2<T>&)>& TransformFunc)
+	TPolygon2<T>& Transform(const TFunction<TVector2<T> (const TVector2<T>&)>& TransformFunc)
 	{
 		int N = Vertices.Num();
 		for (int i = 0; i < N; ++i) 
@@ -791,7 +804,7 @@ public:
 	 */
 	void VtxNormalOffset(T OffsetDistance, bool bUseFaceAvg = false)
 	{
-		TArray<FVector2<T>> NewVertices;
+		TArray<TVector2<T>> NewVertices;
 		NewVertices.SetNumUninitialized(Vertices.Num());
 		if (bUseFaceAvg) 
 		{
@@ -824,15 +837,15 @@ public:
 	{
 		// [TODO] possibly can do with half as many normalizes if we do w/ sequential edges,
 		//  rather than centering on each v?
-		TArray<FVector2<T>> NewVertices;
+		TArray<TVector2<T>> NewVertices;
 		NewVertices.SetNumUninitialized(Vertices.Num());
 		for (int k = 0; k < Vertices.Num(); ++k) 
 		{
-			FVector2<T> v = Vertices[k];
-			FVector2<T> next = Vertices[(k + 1) % Vertices.Num()];
-			FVector2<T> prev = Vertices[k == 0 ? Vertices.Num() - 1 : k - 1];
-			FVector2<T> dn = Normalized(next - v);
-			FVector2<T> dp = Normalized(prev - v);
+			TVector2<T> v = Vertices[k];
+			TVector2<T> next = Vertices[(k + 1) % Vertices.Num()];
+			TVector2<T> prev = Vertices[k == 0 ? Vertices.Num() - 1 : k - 1];
+			TVector2<T> dn = Normalized(next - v);
+			TVector2<T> dp = Normalized(prev - v);
 			TLine2<T> ln(v + OffsetDistance * PerpCW(dn), dn);
 			TLine2<T> lp(v - OffsetDistance * PerpCW(dp), dp);
 
@@ -862,7 +875,7 @@ private:
 	//            v[] = polyline array of vertex points
 	//            j,k = indices for the subchain v[j] to v[k]
 	//    Output: mk[] = array of markers matching vertex array v[]
-	static void SimplifyDouglasPeucker(T Tolerance, const TArray<FVector2<T>>& Vertices, int j, int k, TArray<bool>& Marked)
+	static void SimplifyDouglasPeucker(T Tolerance, const TArray<TVector2<T>>& Vertices, int j, int k, TArray<bool>& Marked)
 	{
 		Marked.SetNum(Vertices.Num());
 		if (k <= j + 1) // there is nothing to simplify
@@ -914,7 +927,7 @@ public:
 		}
 
 		int i, k, pv;            // misc counters
-		TArray<FVector2<T>> NewVertices;
+		TArray<TVector2<T>> NewVertices;
 		NewVertices.SetNumUninitialized(n + 1);  // vertex buffer
 		TArray<bool> Marked;
 		Marked.SetNumUninitialized(n + 1);
@@ -928,7 +941,7 @@ public:
 		NewVertices[0] = Vertices[0];              // start at the beginning
 		for (i = 1, k = 1, pv = 0; i < n; i++) 
 		{
-			if (Vertices[i].DistanceSquared(Vertices[pv]) < clusterTol2)
+			if ( UE::Geometry::DistanceSquared(Vertices[i], Vertices[pv]) < clusterTol2)
 			{
 				continue;
 			}
@@ -1021,22 +1034,22 @@ public:
 	{
 		check(IsClockwise());
 
-		TArray<FVector2<T>> OldV = Vertices;
+		TArray<TVector2<T>> OldV = Vertices;
 		int N = OldV.Num();
 		Vertices.Reset();
 
 		int iCur = 0;
 		do {
-			FVector2<T> center = OldV[iCur];
+			TVector2<T> center = OldV[iCur];
 
 			int iPrev = (iCur == 0) ? N - 1 : iCur - 1;
-			FVector2<T> prev = OldV[iPrev];
+			TVector2<T> prev = OldV[iPrev];
 			int iNext = (iCur + 1) % N;
-			FVector2<T> next = OldV[iNext];
+			TVector2<T> next = OldV[iNext];
 
-			FVector2<T> cp = prev - center;
+			TVector2<T> cp = prev - center;
 			T cpdist = Normalize(cp);
-			FVector2<T> cn = next - center;
+			TVector2<T> cn = next - center;
 			T cndist = Normalize(cn);
 
 			// if degenerate, skip this vert
@@ -1064,9 +1077,9 @@ public:
 
 
 			T prev_cut_dist = TMathUtil<T>::Min(ChamferDist, cpdist*0.5);
-			FVector2<T> prev_cut = center + cp * prev_cut_dist;
+			TVector2<T> prev_cut = center + cp * prev_cut_dist;
 			T next_cut_dist = TMathUtil<T>::Min(ChamferDist, cndist * 0.5);
-			FVector2<T> next_cut = center + cn * next_cut_dist;
+			TVector2<T> next_cut = center + cn * next_cut_dist;
 
 			Vertices.Add(prev_cut);
 			Vertices.Add(next_cut);
@@ -1082,33 +1095,33 @@ public:
 	/**
 	 * Construct a four-vertex rectangle polygon
 	 */
-	static TPolygon2<T> MakeRectangle(const FVector2<T>& Center, T Width, T Height)
+	static TPolygon2<T> MakeRectangle(const TVector2<T>& Center, T Width, T Height)
 	{
 		TPolygon2<T> Rectangle;
 		Rectangle.Vertices.SetNumUninitialized(4);
-		Rectangle.Set(0, FVector2<T>(Center.X - Width / 2, Center.Y - Height / 2));
-		Rectangle.Set(1, FVector2<T>(Center.X + Width / 2, Center.Y - Height / 2));
-		Rectangle.Set(2, FVector2<T>(Center.X + Width / 2, Center.Y + Height / 2));
-		Rectangle.Set(3, FVector2<T>(Center.X - Width / 2, Center.Y + Height / 2));
+		Rectangle.Set(0, TVector2<T>(Center.X - Width / 2, Center.Y - Height / 2));
+		Rectangle.Set(1, TVector2<T>(Center.X + Width / 2, Center.Y - Height / 2));
+		Rectangle.Set(2, TVector2<T>(Center.X + Width / 2, Center.Y + Height / 2));
+		Rectangle.Set(3, TVector2<T>(Center.X - Width / 2, Center.Y + Height / 2));
 		return Rectangle;
 	}
 
 	/**
 	* Construct a rounded rectangle polygon
 	*/
-	static TPolygon2<T> MakeRoundedRectangle(const FVector2<T>& Center, T Width, T Height, T Corner, int CornerSteps)
+	static TPolygon2<T> MakeRoundedRectangle(const TVector2<T>& Center, T Width, T Height, T Corner, int CornerSteps)
 	{
 		TPolygon2<T> RRectangle;
 		CornerSteps = FMath::Max(2, CornerSteps);
 		RRectangle.Vertices.SetNumUninitialized(CornerSteps * 4);
 		
-		FVector2<T> InnerExtents(Width*.5 - Corner, Height*.5 - Corner);
-		FVector2<T> InnerCorners[4]
+		TVector2<T> InnerExtents(Width*.5 - Corner, Height*.5 - Corner);
+		TVector2<T> InnerCorners[4]
 		{
 			Center - InnerExtents,
-			FVector2<T>(Center.X + InnerExtents.X, Center.Y - InnerExtents.Y),
+			TVector2<T>(Center.X + InnerExtents.X, Center.Y - InnerExtents.Y),
 			Center + InnerExtents,
-			FVector2<T>(Center.X - InnerExtents.X, Center.Y + InnerExtents.Y),
+			TVector2<T>(Center.X - InnerExtents.X, Center.Y + InnerExtents.Y),
 		};
 
 		T IdxToAng = TMathUtil<T>::HalfPi / T(CornerSteps - 1);
@@ -1117,10 +1130,10 @@ public:
 			T Angle = StepIdx * IdxToAng;
 			T OffC = TMathUtil<T>::Cos(Angle) * Corner;
 			T OffS = TMathUtil<T>::Sin(Angle) * Corner;
-			RRectangle.Set(StepIdx + CornerSteps * 0, InnerCorners[0] + FVector2<T>(-OffC, -OffS));
-			RRectangle.Set(StepIdx + CornerSteps * 1, InnerCorners[1] + FVector2<T>(OffS, -OffC));
-			RRectangle.Set(StepIdx + CornerSteps * 2, InnerCorners[2] + FVector2<T>(OffC, OffS));
-			RRectangle.Set(StepIdx + CornerSteps * 3, InnerCorners[3] + FVector2<T>(-OffS, OffC));
+			RRectangle.Set(StepIdx + CornerSteps * 0, InnerCorners[0] + TVector2<T>(-OffC, -OffS));
+			RRectangle.Set(StepIdx + CornerSteps * 1, InnerCorners[1] + TVector2<T>(OffS, -OffC));
+			RRectangle.Set(StepIdx + CornerSteps * 2, InnerCorners[2] + TVector2<T>(OffC, OffS));
+			RRectangle.Set(StepIdx + CornerSteps * 3, InnerCorners[3] + TVector2<T>(-OffS, OffC));
 		}
 		return RRectangle;
 	}
@@ -1138,7 +1151,7 @@ public:
 		{
 			T t = (T)i / (T)Steps;
 			T a = TMathUtil<T>::TwoPi * t + AngleShiftRadians;
-			Circle.Set(i, FVector2<T>(Radius * TMathUtil<T>::Cos(a), Radius * TMathUtil<T>::Sin(a)));
+			Circle.Set(i, TVector2<T>(Radius * TMathUtil<T>::Cos(a), Radius * TMathUtil<T>::Sin(a)));
 		}
 
 		return Circle;
diff --git a/Engine/Source/Runtime/GeometryCore/Public/Sampling/Gaussians.h b/Engine/Source/Runtime/GeometryCore/Public/Sampling/Gaussians.h
--- a/Engine/Source/Runtime/GeometryCore/Public/Sampling/Gaussians.h
+++ b/Engine/Source/Runtime/GeometryCore/Public/Sampling/Gaussians.h
@@ -87,13 +87,13 @@ public:
 
 	TGaussian2(RealType SigmaIn = 1.0) : Sigma(SigmaIn) {}
 
-	RealType Evaluate(const FVector2<RealType>& XY) const
+	RealType Evaluate(const UE::Math::TVector2<RealType>& XY) const
 	{
 		RealType InvTwoSigmaSqr = 1.0 / (2.0 * Sigma * Sigma);
 		return TMathUtil<RealType>::InvPi * InvTwoSigmaSqr * TMathUtil<RealType>::Exp( -XY.SquaredLength() * InvTwoSigmaSqr );
 	}
 
-	RealType EvaluateSqrUnscaled(const FVector2<RealType>& XY) const
+	RealType EvaluateSqrUnscaled(const UE::Math::TVector2<RealType>& XY) const
 	{
 		RealType InvTwoSigmaSqr = 1.0 / (2.0 * Sigma * Sigma);
 		return TMathUtil<RealType>::Exp(-XY.SquaredLength() * InvTwoSigmaSqr);
diff --git a/Engine/Source/Runtime/GeometryCore/Public/Sampling/GridSampler.h b/Engine/Source/Runtime/GeometryCore/Public/Sampling/GridSampler.h
--- a/Engine/Source/Runtime/GeometryCore/Public/Sampling/GridSampler.h
+++ b/Engine/Source/Runtime/GeometryCore/Public/Sampling/GridSampler.h
@@ -38,12 +38,12 @@ public:
 	 * @param Index linear point index in range [0,Num()-1]
 	 * @return grid sample for given linear index
 	 */
-	FVector2<RealType> Sample(int32 Index) const
+	UE::Math::TVector2<RealType> Sample(int32 Index) const
 	{
 		checkSlow(Index >= 0 && Index < Num());
 		RealType X = (static_cast<RealType>(Index % Size) + RealType{ 1 } / RealType{ 2 }) / static_cast<RealType>(Size);
 		RealType Y = (static_cast<RealType>(Index / Size) + RealType{ 1 } / RealType{ 2 }) / static_cast<RealType>(Size);
-		return FVector2<RealType>(X, Y);
+		return UE::Math::TVector2<RealType>(X, Y);
 	}
 };
 
diff --git a/Engine/Source/Runtime/GeometryCore/Public/Sampling/SphericalFibonacci.h b/Engine/Source/Runtime/GeometryCore/Public/Sampling/SphericalFibonacci.h
--- a/Engine/Source/Runtime/GeometryCore/Public/Sampling/SphericalFibonacci.h
+++ b/Engine/Source/Runtime/GeometryCore/Public/Sampling/SphericalFibonacci.h
@@ -11,6 +11,8 @@ namespace UE
 namespace Geometry
 {
 
+using namespace UE::Math;
+
 /**
  * A Spherical Fibonacci (SF) Point Set is a set of points that are roughly evenly distributed on a sphere.
  * The points lie on a spiral, see https://dl.acm.org/doi/10.1145/2816795.2818131 for more information.
@@ -89,14 +91,14 @@ public:
 		TMatrix2<RealType> invB = B.Inverse();
 
 		//Vector2d c = floor(mul(invB, RealType2(phi, cosTheta - (1 - 1.0/N))));
-		FVector2<RealType> c(phi, cosTheta - (1.0 - 1.0/N));
+		TVector2<RealType> c(phi, cosTheta - (1.0 - 1.0/N));
 		c = invB * c;
 		c.X = TMathUtil<RealType>::Floor(c.X); c.Y = TMathUtil<RealType>::Floor(c.Y);
 
 		RealType d = TMathUtil<RealType>::MaxReal, j = 0;
 		for (int32 s = 0; s < 4; ++s) 
 		{
-			FVector2<RealType> cosTheta_second((RealType)(s % 2) + c.X, (RealType)(s / 2) + c.Y);
+			TVector2<RealType> cosTheta_second((RealType)(s % 2) + c.X, (RealType)(s / 2) + c.Y);
 			cosTheta = B.Row1.Dot(cosTheta_second) + (1.0 - 1.0 / N);
 			cosTheta = TMathUtil<RealType>::Clamp(cosTheta, -1.0, +1.0) * 2.0 - cosTheta;
 			RealType i = TMathUtil<RealType>::Floor(N * 0.5 - cosTheta * N * 0.5);
@@ -167,7 +169,7 @@ public:
 	 * @param Index point index in range [0,Num()-1]
 	 * @return sphere point for given Index
 	 */
-	FVector2<RealType> Point(int32 Index) const
+	TVector2<RealType> Point(int32 Index) const
 	{
 		static const RealType PHI = (TMathUtil<RealType>::Sqrt(5.0) + 1.0) / 2.0;
 		checkSlow(Index >= 0 && Index < N);
@@ -193,14 +195,14 @@ public:
 			X = R * TMathUtil<RealType>::Cos(Theta);
 			Y = R * TMathUtil<RealType>::Sin(Theta);
 		}
-		return FVector2<RealType>(X,Y);
+		return TVector2<RealType>(X,Y);
 	}
 
 	/**
 	 * @param Index point index in range [0,Num()-1]
 	 * @return sphere point for given Index
 	 */
-	FVector2<RealType> operator[](int32 Index) const
+	TVector2<RealType> operator[](int32 Index) const
 	{
 		return Point(Index);
 	}
@@ -257,7 +259,7 @@ public:
 		case EDistribution::Cosine:
 		{
 			TFibonacciLattice<RealType> Points(N, TFibonacciLattice<RealType>::EType::Disc);
-			const FVector2<RealType> Pt = Points[Index];
+			const TVector2<RealType> Pt = Points[Index];
             
 			// Planar projection of Fibonacci spiral (unit disc [xd, yd]) to
 			// hemisphere to achieve cosine weighted ray distribution.
diff --git a/Engine/Source/Runtime/GeometryCore/Public/SegmentTypes.h b/Engine/Source/Runtime/GeometryCore/Public/SegmentTypes.h
--- a/Engine/Source/Runtime/GeometryCore/Public/SegmentTypes.h
+++ b/Engine/Source/Runtime/GeometryCore/Public/SegmentTypes.h
@@ -13,6 +13,8 @@ namespace UE
 namespace Geometry
 {
 
+using namespace UE::Math;
+
 /*
  * 2D Line Segment stored as Center point, normalized Direction vector, and scalar Extent
  */
@@ -21,9 +23,9 @@ struct TSegment2
 {
 public:
 	/** Center point of segment */
-	FVector2<T> Center = FVector2<T>::Zero();
+	TVector2<T> Center = TVector2<T>::Zero();
 	/** normalized Direction vector of segment */
-	FVector2<T> Direction = FVector2<T>::UnitX();
+	TVector2<T> Direction = TVector2<T>::UnitX();
 	/** Extent of segment, which is half the total length */
 	T Extent = (T)0;
 
@@ -32,7 +34,7 @@ public:
 	/**
 	 * Construct a Segment from two Points
 	 */
-	TSegment2(const FVector2<T>& Point0, const FVector2<T>& Point1)
+	TSegment2(const TVector2<T>& Point0, const TVector2<T>& Point1)
 	{
 		// set from endpoints 
 		Center = T(.5) * (Point0 + Point1);
@@ -43,7 +45,7 @@ public:
 	/**
 	 * Construct a segment from a Center Point, normalized Direction, and scalar Extent
 	 */
-	TSegment2(const FVector2<T>& CenterIn, const FVector2<T>& DirectionIn, T ExtentIn)
+	TSegment2(const TVector2<T>& CenterIn, const TVector2<T>& DirectionIn, T ExtentIn)
 	{
 		Center = CenterIn;
 		Direction = DirectionIn; 
@@ -53,13 +55,13 @@ public:
 
 
 	/** Update the Segment with a new start point */
-	inline void SetStartPoint(const FVector2<T>& Point)
+	inline void SetStartPoint(const TVector2<T>& Point)
 	{
 		update_from_endpoints(Point, EndPoint());
 	}
 
 	/** Update the Segment with a new end point */
-	inline void SetEndPoint(const FVector2<T>& Point)
+	inline void SetEndPoint(const TVector2<T>& Point)
 	{
 		update_from_endpoints(StartPoint(), Point);
 	}
@@ -74,13 +76,13 @@ public:
 
 
 	/** @return start point of segment */
-	inline FVector2<T> StartPoint() const
+	inline TVector2<T> StartPoint() const
 	{
 		return Center - Extent * Direction;
 	}
 
 	/** @return end point of segment */
-	inline FVector2<T> EndPoint() const
+	inline TVector2<T> EndPoint() const
 	{
 		return Center + Extent * Direction;
 	}
@@ -93,7 +95,7 @@ public:
 	}
 
 	/** @return first (i == 0) or second (i == 1) endpoint of the Segment  */
-	inline FVector2<T> GetPointFromIndex(int i) const
+	inline TVector2<T> GetPointFromIndex(int i) const
 	{
 		return (i == 0) ? (Center - Extent * Direction) : (Center + Extent * Direction);
 	}
@@ -101,7 +103,7 @@ public:
 	/**
 	 * @return point on segment at given (signed) Distance from the segment Origin
 	 */
-	inline FVector2<T> PointAt(T DistanceParameter) const
+	inline TVector2<T> PointAt(T DistanceParameter) const
 	{
 		return Center + DistanceParameter * Direction;
 	}
@@ -110,7 +112,7 @@ public:
 	 * @param UnitParameter value in range [0,1]
 	 * @return point on segment at that linearly interpolates between start and end based on Unit Parameter
 	 */
-	inline FVector2<T> PointBetween(T UnitParameter) const
+	inline TVector2<T> PointBetween(T UnitParameter) const
 	{
 		return Center + ((T)2 * UnitParameter - (T)1) * Extent * Direction;
 	}
@@ -118,7 +120,7 @@ public:
 	/**
 	 * @return minimum squared distance from Point to segment
 	 */
-	inline T DistanceSquared(const FVector2<T>& Point) const
+	inline T DistanceSquared(const TVector2<T>& Point) const
 	{
 		T DistParameter;
 		return DistanceSquared(Point, DistParameter);
@@ -131,28 +133,28 @@ public:
 	 * @param DistParameterOut calculated distance parameter in range [-Extent,Extent]
 	 * @return minimum squared distance from Point to Segment
 	 */	
-	T DistanceSquared(const FVector2<T>& Point, T& DistParameterOut) const
+	T DistanceSquared(const TVector2<T>& Point, T& DistParameterOut) const
 	{
 		DistParameterOut = (Point - Center).Dot(Direction);
 		if (DistParameterOut >= Extent)
 		{
 			DistParameterOut = Extent;
-			return Point.DistanceSquared(EndPoint());
+			return UE::Geometry::DistanceSquared(Point, EndPoint());
 		}
 		else if (DistParameterOut <= -Extent)
 		{
 			DistParameterOut = -Extent;
-			return Point.DistanceSquared(StartPoint());
+			return UE::Geometry::DistanceSquared(Point, StartPoint());
 		}
-		FVector2<T> ProjectedPt = Center + DistParameterOut * Direction;
-		return ProjectedPt.DistanceSquared(Point);
+		TVector2<T> ProjectedPt = Center + DistParameterOut * Direction;
+		return UE::Geometry::DistanceSquared(ProjectedPt, Point);
 	}
 
 
 	/**
 	 * @return nearest point on segment to QueryPoint
 	 */
-	inline FVector2<T> NearestPoint(const FVector2<T>& QueryPoint) const
+	inline TVector2<T> NearestPoint(const TVector2<T>& QueryPoint) const
 	{
 		T t = (QueryPoint - Center).Dot(Direction);
 		if (t >= Extent)
@@ -170,7 +172,7 @@ public:
 	/**
 	 * @return scalar projection of QueryPoint onto line of Segment (not clamped to Extents)
 	 */
-	inline T Project(const FVector2<T>& QueryPoint) const
+	inline T Project(const TVector2<T>& QueryPoint) const
 	{
 		return (QueryPoint - Center).Dot(Direction);
 	}
@@ -179,7 +181,7 @@ public:
 	/**
 	 * @return scalar projection of QueryPoint onto line of Segment, mapped to [0,1] range along segment
 	 */
-	inline T ProjectUnitRange(const FVector2<T>& QueryPoint) const
+	inline T ProjectUnitRange(const TVector2<T>& QueryPoint) const
 	{
 		T ProjT = (QueryPoint - Center).Dot(Direction);
 		T Alpha = ((ProjT / Extent) + (T)1) * (T)0.5;
@@ -194,11 +196,11 @@ public:
 	 * @param Tolerance tolerance band in which we return 0
 	 * @return +1 if point is to right of line, -1 if left, and 0 if on line or within tolerance band
 	 */
-	int WhichSide(const FVector2<T>& QueryPoint, T Tolerance = 0)
+	int WhichSide(const TVector2<T>& QueryPoint, T Tolerance = 0)
 	{
 		// [TODO] subtract Center from test?
-		FVector2<T> EndPt = Center + Extent * Direction;
-		FVector2<T> StartPt = Center - Extent * Direction;
+		TVector2<T> EndPt = Center + Extent * Direction;
+		TVector2<T> StartPt = Center - Extent * Direction;
 		T det = -Orient(EndPt, StartPt, QueryPoint);
 		return (det > Tolerance ? +1 : (det < -Tolerance ? -1 : 0));
 	}
@@ -215,7 +217,7 @@ public:
 	{
 		// see IntrLine2Line2 and IntrSegment2Segment2 for details on this code
 
-		FVector2<T> diff = OtherSegment.Center - Center;
+		TVector2<T> diff = OtherSegment.Center - Center;
 		T D0DotPerpD1 = DotPerp(Direction, OtherSegment.Direction);
 		if (TMathUtil<T>::Abs(D0DotPerpD1) > DotThresh)      // Lines intersect in a single point.
 		{
@@ -260,7 +262,7 @@ public:
 	/**
 	 * Calculate distance from QueryPoint to segment (StartPt,EndPt)
 	 */
-	static T FastDistanceSquared(const FVector2<T>& StartPt, const FVector2<T>& EndPt, const FVector2<T>& QueryPt, T Tolerance = TMathUtil<T>::Epsilon)
+	static T FastDistanceSquared(const TVector2<T>& StartPt, const TVector2<T>& EndPt, const TVector2<T>& QueryPt, T Tolerance = TMathUtil<T>::Epsilon)
 	{
 		T vx = EndPt.X - StartPt.X, vy = EndPt.Y - StartPt.Y;
 		T len2 = vx * vx + vy * vy;
@@ -295,7 +297,7 @@ public:
 	 * @param Tolerance tolerance band in which we return 0
 	 * @return +1 if point is to right of line, -1 if left, and 0 if on line or within tolerance band
 	 */
-	static int WhichSide(const FVector2<T>& StartPt, const FVector2<T>& EndPt, const FVector2<T>& QueryPt, T Tolerance = (T)0)
+	static int WhichSide(const TVector2<T>& StartPt, const TVector2<T>& EndPt, const TVector2<T>& QueryPt, T Tolerance = (T)0)
 	{
 		T det = -Orient(StartPt, EndPt, QueryPt);
 		return (det > Tolerance ? +1 : (det < -Tolerance ? -1 : 0));
@@ -308,7 +310,7 @@ public:
 protected:
 
 	// update segment based on new endpoints
-	inline void update_from_endpoints(const FVector2<T>& p0, const FVector2<T>& p1)
+	inline void update_from_endpoints(const TVector2<T>& p0, const TVector2<T>& p1)
 	{
 		Center = 0.5 * (p0 + p1);
 		Direction = p1 - p0;
@@ -450,15 +452,15 @@ public:
 		if (DistParameterOut >= Extent)
 		{
 			DistParameterOut = Extent;
-			return Point.DistanceSquared(EndPoint());
+			return UE::Geometry::DistanceSquared(Point, EndPoint());
 		}
 		else if (DistParameterOut <= -Extent)
 		{
 			DistParameterOut = -Extent;
-			return Point.DistanceSquared(StartPoint());
+			return UE::Geometry::DistanceSquared(Point, StartPoint());
 		}
 		FVector3<T> ProjectedPt = Center + DistParameterOut * Direction;
-		return ProjectedPt.DistanceSquared(Point);
+		return UE::Geometry::DistanceSquared(ProjectedPt, Point);
 	}
 
 
diff --git a/Engine/Source/Runtime/GeometryCore/Public/Spatial/PointHashGrid2.h b/Engine/Source/Runtime/GeometryCore/Public/Spatial/PointHashGrid2.h
--- a/Engine/Source/Runtime/GeometryCore/Public/Spatial/PointHashGrid2.h
+++ b/Engine/Source/Runtime/GeometryCore/Public/Spatial/PointHashGrid2.h
@@ -13,6 +13,8 @@ namespace UE
 namespace Geometry
 {
 
+using namespace UE::Math;
+
 /**
  * Hash Grid for values associated with 2D points.
  *
@@ -56,7 +58,7 @@ public:
 	 * @param Value the point/value to insert
 	 * @param Position the position associated with this value
 	 */
-	void InsertPoint(const PointDataType& Value, const FVector2<RealType>& Position)
+	void InsertPoint(const PointDataType& Value, const TVector2<RealType>& Position)
 	{
 		FVector2i idx = Indexer.ToGrid(Position);
 		{
@@ -70,7 +72,7 @@ public:
 	 * @param Value the point/value to insert
 	 * @param Position the position associated with this value
 	 */
-	void InsertPointUnsafe(const PointDataType& Value, const FVector2<RealType>& Position)
+	void InsertPointUnsafe(const PointDataType& Value, const TVector2<RealType>& Position)
 	{
 		FVector2i idx = Indexer.ToGrid(Position);
 		Hash.Add(idx, Value);
@@ -83,7 +85,7 @@ public:
 	 * @param Position the position associated with this value
 	 * @return true if the value existed at this position
 	 */
-	bool RemovePoint(const PointDataType& Value, const FVector2<RealType>& Position)
+	bool RemovePoint(const PointDataType& Value, const TVector2<RealType>& Position)
 	{
 		FVector2i idx = Indexer.ToGrid(Position);
 		{
@@ -98,7 +100,7 @@ public:
 	 * @param Position the position associated with this value
 	 * @return true if the value existed at this position
 	 */
-	bool RemovePointUnsafe(const PointDataType& Value, const FVector2<RealType>& Position)
+	bool RemovePointUnsafe(const PointDataType& Value, const TVector2<RealType>& Position)
 	{
 		FVector2i idx = Indexer.ToGrid(Position);
 		return Hash.RemoveSingle(idx, Value) > 0;
@@ -111,7 +113,7 @@ public:
 	 * @param OldPosition the current position associated with this value
 	 * @param NewPosition the new position for this value
 	 */
-	void UpdatePoint(const PointDataType& Value, const FVector2<RealType>& OldPosition, const FVector2<RealType>& NewPosition)
+	void UpdatePoint(const PointDataType& Value, const TVector2<RealType>& OldPosition, const TVector2<RealType>& NewPosition)
 	{
 		FVector2i old_idx = Indexer.ToGrid(OldPosition);
 		FVector2i new_idx = Indexer.ToGrid(NewPosition);
@@ -139,7 +141,7 @@ public:
 	 * @param OldPosition the current position associated with this value
 	 * @param NewPosition the new position for this value
 	 */
-	void UpdatePointUnsafe(const PointDataType& Value, const FVector2<RealType>& OldPosition, const FVector2<RealType>& NewPosition)
+	void UpdatePointUnsafe(const PointDataType& Value, const TVector2<RealType>& OldPosition, const TVector2<RealType>& NewPosition)
 	{
 		FVector2i old_idx = Indexer.ToGrid(OldPosition);
 		FVector2i new_idx = Indexer.ToGrid(NewPosition);
@@ -163,7 +165,7 @@ public:
 	 * @return the found pair (Value,DistanceFunc(Value)), or (InvalidValue,MaxDouble) if not found
 	 */
 	TPair<PointDataType, RealType> FindNearestInRadius(
-		const FVector2<RealType>& QueryPoint, RealType Radius, 
+		const TVector2<RealType>& QueryPoint, RealType Radius, 
 		TFunctionRef<RealType(const PointDataType&)> DistanceSqFunc,
 		TFunctionRef<bool(const PointDataType&)> IgnoreFunc = [](const PointDataType& data) { return false; }) const
 	{
@@ -172,8 +174,8 @@ public:
 			return TPair<PointDataType, RealType>(GetInvalidValue(), TNumericLimits<RealType>::Max());
 		}
 
-		FVector2i min_idx = Indexer.ToGrid(QueryPoint - Radius * FVector2<RealType>::One());
-		FVector2i max_idx = Indexer.ToGrid(QueryPoint + Radius * FVector2<RealType>::One());
+		FVector2i min_idx = Indexer.ToGrid(QueryPoint - Radius * TVector2<RealType>::One());
+		FVector2i max_idx = Indexer.ToGrid(QueryPoint + Radius * TVector2<RealType>::One());
 
 		RealType min_distsq = TNumericLimits<RealType>::Max();
 		PointDataType nearest = GetInvalidValue();
diff --git a/Engine/Source/Runtime/GeometryCore/Public/Spatial/SampledScalarField2.h b/Engine/Source/Runtime/GeometryCore/Public/Spatial/SampledScalarField2.h
--- a/Engine/Source/Runtime/GeometryCore/Public/Spatial/SampledScalarField2.h
+++ b/Engine/Source/Runtime/GeometryCore/Public/Spatial/SampledScalarField2.h
@@ -14,6 +14,8 @@ namespace UE
 namespace Geometry
 {
 
+using namespace UE::Math;
+
 /**
  * TSampledScalarField2 implements a generic 2D grid of values that can be interpolated in various ways.
  * The grid is treated as a set of sample points in 2D space, IE a grid origin and x/y point-spacing is part of this class.
@@ -28,8 +30,8 @@ class TSampledScalarField2
 public:
 	TDenseGrid2<ValueType> GridValues;
 
-	FVector2<RealType> GridOrigin;
-	FVector2<RealType> CellDimensions;
+	TVector2<RealType> GridOrigin;
+	TVector2<RealType> CellDimensions;
 
 	/**
 	 * Create empty grid, defaults to 2x2 grid of whatever default value of ValueType is
@@ -37,8 +39,8 @@ public:
 	TSampledScalarField2()
 	{
 		GridValues.Resize(2, 2);
-		GridOrigin = FVector2<RealType>::Zero();
-		CellDimensions = FVector2<RealType>::One();
+		GridOrigin = TVector2<RealType>::Zero();
+		CellDimensions = TVector2<RealType>::One();
 	}
 
 	void CopyConfiguration(const TSampledScalarField2<RealType,ValueType>& OtherField)
@@ -75,7 +77,7 @@ public:
 	/**
 	 * Set the 2D origin of the grid
 	 */
-	void SetPosition(const FVector2<RealType>& Origin)
+	void SetPosition(const TVector2<RealType>& Origin)
 	{
 		GridOrigin = Origin;
 	}
@@ -93,10 +95,10 @@ public:
 	 * @param Position sample point relative to grid origin/dimensions
 	 * @return interpolated value at this position
 	 */
-	ValueType BilinearSampleClamped(const FVector2<RealType>& Position) const
+	ValueType BilinearSampleClamped(const TVector2<RealType>& Position) const
 	{
 		// transform Position into grid coordinates
-		FVector2<RealType> GridPoint(
+		TVector2<RealType> GridPoint(
 			((Position.X - GridOrigin.X) / CellDimensions.X),
 			((Position.Y - GridOrigin.Y) / CellDimensions.Y));
 
@@ -140,10 +142,10 @@ public:
 	 * @param Position sample point relative to grid origin/dimensions
 	 * @return interpolated value at this position
 	 */
-	void BilinearSampleGradientClamped(const FVector2<RealType>& Position, ValueType& GradXOut, ValueType& GradYOut) const
+	void BilinearSampleGradientClamped(const TVector2<RealType>& Position, ValueType& GradXOut, ValueType& GradYOut) const
 	{
 		// transform Position into grid coordinates
-		FVector2<RealType> GridPoint(
+		TVector2<RealType> GridPoint(
 			((Position.X - GridOrigin.X) / CellDimensions.X),
 			((Position.Y - GridOrigin.Y) / CellDimensions.Y));
 
diff --git a/Engine/Source/Runtime/GeometryCore/Public/TriangleTypes.h b/Engine/Source/Runtime/GeometryCore/Public/TriangleTypes.h
--- a/Engine/Source/Runtime/GeometryCore/Public/TriangleTypes.h
+++ b/Engine/Source/Runtime/GeometryCore/Public/TriangleTypes.h
@@ -10,6 +10,7 @@ namespace UE
 namespace Geometry
 {
 
+using namespace UE::Math;
 
 /**
  * Triangle utility functions
@@ -33,35 +34,35 @@ namespace TriangleUtil
 template<typename RealType>
 struct TTriangle2
 {
-	FVector2<RealType> V[3];
+	TVector2<RealType> V[3];
 
 	TTriangle2() {}
 
-	TTriangle2(const FVector2<RealType>& V0, const FVector2<RealType>& V1, const FVector2<RealType>& V2)
+	TTriangle2(const TVector2<RealType>& V0, const TVector2<RealType>& V1, const TVector2<RealType>& V2)
 	{
 		V[0] = V0;
 		V[1] = V1;
 		V[2] = V2;
 	}
 
-	TTriangle2(const FVector2<RealType> VIn[3])
+	TTriangle2(const TVector2<RealType> VIn[3])
 	{
 		V[0] = VIn[0];
 		V[1] = VIn[1];
 		V[2] = VIn[2];
 	}
 
-	FVector2<RealType> BarycentricPoint(RealType Bary0, RealType Bary1, RealType Bary2) const
+	TVector2<RealType> BarycentricPoint(RealType Bary0, RealType Bary1, RealType Bary2) const
 	{
 		return Bary0 * V[0] + Bary1 * V[1] + Bary2 * V[2];
 	}
 
-	FVector2<RealType> BarycentricPoint(const FVector3<RealType>& BaryCoords) const
+	TVector2<RealType> BarycentricPoint(const FVector3<RealType>& BaryCoords) const
 	{
 		return BaryCoords[0] * V[0] + BaryCoords[1] * V[1] + BaryCoords[2] * V[2];
 	}
 
-	FVector3<RealType> GetBarycentricCoords(const FVector2<RealType>& Point) const
+	FVector3<RealType> GetBarycentricCoords(const TVector2<RealType>& Point) const
 	{
 		return VectorUtil::BarycentricCoords(Point, V[0], V[1], V[2]);
 	}
@@ -72,7 +73,7 @@ struct TTriangle2
 	 * @param C third vertex of triangle
 	 * @return signed area of triangle
 	 */
-	static RealType SignedArea(const FVector2<RealType>& A, const FVector2<RealType>& B, const FVector2<RealType>& C)
+	static RealType SignedArea(const TVector2<RealType>& A, const TVector2<RealType>& B, const TVector2<RealType>& C)
 	{
 		return ((RealType)0.5) * ((A.X*B.Y - A.Y*B.X) + (B.X*C.Y - B.Y*C.X) + (C.X*A.Y - C.Y*A.X));
 	}
@@ -97,7 +98,7 @@ struct TTriangle2
 	 * @param QueryPoint test point
 	 * @return true if QueryPoint is inside triangle
 	 */
-	static bool IsInside(const FVector2<RealType>& A, const FVector2<RealType>& B, const FVector2<RealType>& C, const FVector2<RealType>& QueryPoint)
+	static bool IsInside(const TVector2<RealType>& A, const TVector2<RealType>& B, const TVector2<RealType>& C, const TVector2<RealType>& QueryPoint)
 	{
 		RealType Sign1 = Orient(A, B, QueryPoint);
 		RealType Sign2 = Orient(B, C, QueryPoint);
@@ -106,14 +107,14 @@ struct TTriangle2
 	}
 
 	/** @return true if QueryPoint is inside triangle */
-	bool IsInside(const FVector2<RealType>& QueryPoint) const
+	bool IsInside(const TVector2<RealType>& QueryPoint) const
 	{
 		return IsInside(V[0], V[1], V[2], QueryPoint);
 	}
 
 
 	/** @return true if QueryPoint is inside triangle or on edge */
-	static bool IsInsideOrOn(const FVector2<RealType>& A, const FVector2<RealType>& B, const FVector2<RealType>& C, const FVector2<RealType>& QueryPoint)
+	static bool IsInsideOrOn(const TVector2<RealType>& A, const TVector2<RealType>& B, const TVector2<RealType>& C, const TVector2<RealType>& QueryPoint)
 	{
 		RealType Sign1 = Orient(A, B, QueryPoint);
 		RealType Sign2 = Orient(B, C, QueryPoint);
@@ -122,7 +123,7 @@ struct TTriangle2
 	}
 
 	/** @return true if QueryPoint is inside triangle or on edge */
-	bool IsInsideOrOn(const FVector2<RealType>& QueryPoint, RealType Epsilon = 0) const
+	bool IsInsideOrOn(const TVector2<RealType>& QueryPoint, RealType Epsilon = 0) const
 	{
 		return IsInsideOrOn(V[0], V[1], V[2], QueryPoint);
 	}
@@ -133,7 +134,7 @@ struct TTriangle2
 	 *  (the code early-outs at the first 'outside' edge, which only works if the triangle is oriented as expected)
 	 * @return 1 if outside, -1 if inside, 0 if on boundary
 	 */
-	int IsInsideOrOn_Oriented(const FVector2<RealType>& QueryPoint) const
+	int IsInsideOrOn_Oriented(const TVector2<RealType>& QueryPoint) const
 	{
 		return IsInsideOrOn_Oriented(V[0], V[1], V[2], QueryPoint);
 	}
@@ -143,7 +144,7 @@ struct TTriangle2
 	 *  (the code early-outs at the first 'outside' edge, which only works if the triangle is oriented as expected)
 	 * @return 1 if outside, -1 if inside, 0 if on boundary
 	 */
-	static int IsInsideOrOn_Oriented(const FVector2<RealType>& A, const FVector2<RealType>& B, const FVector2<RealType>& C, const FVector2<RealType>& QueryPoint)
+	static int IsInsideOrOn_Oriented(const TVector2<RealType>& A, const TVector2<RealType>& B, const TVector2<RealType>& C, const TVector2<RealType>& QueryPoint)
 	{
 		checkSlow(Orient(A, B, C) <= 0); // TODO: remove this checkSlow; it's just to make sure the orientation is as expected
 
@@ -173,7 +174,6 @@ struct TTriangle2
 
 typedef TTriangle2<float> FTriangle2f;
 typedef TTriangle2<double> FTriangle2d;
-typedef TTriangle2<int> FTriangle2i;
 
 
 
diff --git a/Engine/Source/Runtime/GeometryCore/Public/Util/DynamicVector.h b/Engine/Source/Runtime/GeometryCore/Public/Util/DynamicVector.h
--- a/Engine/Source/Runtime/GeometryCore/Public/Util/DynamicVector.h
+++ b/Engine/Source/Runtime/GeometryCore/Public/Util/DynamicVector.h
@@ -494,7 +494,7 @@ public:
 	{
 		return Data[TopIndex * N + SubIndex];
 	}
-	inline void SetVector2(unsigned int TopIndex, const FVector2<Type>& V)
+	inline void SetVector2(unsigned int TopIndex, const UE::Math::TVector2<Type>& V)
 	{
 		check(N >= 2);
 		unsigned int i = TopIndex * N;
@@ -509,10 +509,10 @@ public:
 		Data[i + 1] = V.Y;
 		Data[i + 2] = V.Z;
 	}
-	inline FVector2<Type> AsVector2(unsigned int TopIndex) const
+	inline UE::Math::TVector2<Type> AsVector2(unsigned int TopIndex) const
 	{
 		check(N >= 2);
-		return FVector2<Type>(
+		return UE::Math::TVector2<Type>(
 			Data[TopIndex * N + 0],
 			Data[TopIndex * N + 1]);
 	}
diff --git a/Engine/Source/Runtime/GeometryCore/Public/Util/ElementLinearization.h b/Engine/Source/Runtime/GeometryCore/Public/Util/ElementLinearization.h
--- a/Engine/Source/Runtime/GeometryCore/Public/Util/ElementLinearization.h
+++ b/Engine/Source/Runtime/GeometryCore/Public/Util/ElementLinearization.h
@@ -266,7 +266,7 @@ public:
 		YVector[i] = Value;
 	}
 
-	void SetXY(int32 i, const FVector2<RealType>& Value)
+	void SetXY(int32 i, const UE::Math::TVector2<RealType>& Value)
 	{
 		XVector[i] = Value.X;
 		YVector[i] = Value.Y;
diff --git a/Engine/Source/Runtime/GeometryCore/Public/Util/GridIndexing2.h b/Engine/Source/Runtime/GeometryCore/Public/Util/GridIndexing2.h
--- a/Engine/Source/Runtime/GeometryCore/Public/Util/GridIndexing2.h
+++ b/Engine/Source/Runtime/GeometryCore/Public/Util/GridIndexing2.h
@@ -12,6 +12,8 @@ namespace UE
 namespace Geometry
 {
 
+using namespace UE::Math;
+
 /**
  * Convert between integer grid coordinates and scaled real-valued coordinates (ie assumes integer grid origin == real origin)
  */
@@ -27,7 +29,7 @@ struct TScaleGridIndexer2
 	}
 
 	/** Convert real-valued point to integer grid coordinates */
-	inline FVector2i ToGrid(const FVector2<RealType>& P) const
+	inline FVector2i ToGrid(const TVector2<RealType>& P) const
 	{
 		return FVector2i(
 			(int)TMathUtil<RealType>::Floor(P.X / CellSize),
@@ -35,9 +37,9 @@ struct TScaleGridIndexer2
 	}
 
 	/** Convert integer grid coordinates to real-valued point */
-	inline FVector2<RealType> FromGrid(const FVector2i& GridPoint) const
+	inline TVector2<RealType> FromGrid(const FVector2i& GridPoint) const
 	{
-		return FVector2<RealType>(GridPoint.X*CellSize, GridPoint.Y*CellSize);
+		return TVector2<RealType>(GridPoint.X*CellSize, GridPoint.Y*CellSize);
 	}
 };
 typedef TScaleGridIndexer2<float> FScaleGridIndexer2f;
@@ -53,9 +55,9 @@ struct TShiftGridIndexer2
 	/** Real-valued size of an integer grid cell */
 	RealType CellSize;
 	/** Real-valued origin of grid, position of integer grid origin */
-	FVector2<RealType> Origin;
+	TVector2<RealType> Origin;
 
-	TShiftGridIndexer2(const FVector2<RealType>& origin, RealType cellSize)
+	TShiftGridIndexer2(const TVector2<RealType>& origin, RealType cellSize)
 	{
 		Origin = origin;
 		CellSize = cellSize;
@@ -63,7 +65,7 @@ struct TShiftGridIndexer2
 	}
 
 	/** Convert real-valued point to integer grid coordinates */
-	inline FVector2i ToGrid(const FVector2<RealType>& Point) const
+	inline FVector2i ToGrid(const TVector2<RealType>& Point) const
 	{
 		return FVector2i(
 			(int)TMathUtil<RealType>::Floor((Point.X - Origin.X) / CellSize),
@@ -71,17 +73,17 @@ struct TShiftGridIndexer2
 	}
 
 	/** Convert integer grid coordinates to real-valued point */
-	inline FVector2<RealType> FromGrid(const FVector2i& GridPoint) const
+	inline TVector2<RealType> FromGrid(const FVector2i& GridPoint) const
 	{
-		return FVector2<RealType>(
+		return TVector2<RealType>(
 			((RealType)GridPoint.X * CellSize) + Origin.X,
 			((RealType)GridPoint.Y * CellSize) + Origin.Y);
 	}
 
 	/** Convert real-valued grid coordinates to real-valued point */
-	inline FVector2<RealType> FromGrid(const FVector2<RealType>& RealGridPoint) const
+	inline TVector2<RealType> FromGrid(const TVector2<RealType>& RealGridPoint) const
 	{
-		return FVector2<RealType>(
+		return TVector2<RealType>(
 			((RealType)RealGridPoint.X * CellSize) + Origin.X,
 			((RealType)RealGridPoint.Y * CellSize) + Origin.Y);
 	}
diff --git a/Engine/Source/Runtime/GeometryCore/Public/VectorTypes.h b/Engine/Source/Runtime/GeometryCore/Public/VectorTypes.h
--- a/Engine/Source/Runtime/GeometryCore/Public/VectorTypes.h
+++ b/Engine/Source/Runtime/GeometryCore/Public/VectorTypes.h
@@ -9,253 +9,119 @@
 #include <sstream>
 
 
-namespace UE {
-namespace Geometry {
-
-
 /**
-* Templated 2D Vector. Ported from g3Sharp library, with the intention of
-* maintaining compatibility with existing g3Sharp code. Has an API
-* similar to WildMagic, GTEngine, Eigen, etc.
-*
-* Convenience typedefs for FVector2f/FVector2d are defined, and
-* should be preferentially used over the base template type
-*
-* @todo Possibly can be replaced/merged with Chaos TVector<T,N>
-*/
-template <typename T>
-struct FVector2
-{
-	T X{}, Y{};
-
-	constexpr FVector2()
-		: X(0), Y(0)
-	{
-	}
-
-	constexpr FVector2(T ValX, T ValY)
-		: X(ValX), Y(ValY)
-	{
-	}
-
-	constexpr FVector2(const FVector2& Vec) = default;
-
-	constexpr FVector2(const FVector2D& Vec)
-		: X((T)Vec.X), Y((T)Vec.Y)
-	{
-	}
-
-	template<typename RealType2>
-	explicit constexpr FVector2(const FVector2<RealType2>& Vec)
-		: X((T)Vec.X), Y((T)Vec.Y)
-	{
-	}
-
-	explicit operator FVector2D() const
-	{
-		return FVector2D((float)X, (float)Y);
-	}
+ * Temporary local definition of UE::Math::TVector2 to support LWC conversion.
+ */
+namespace UE {
+namespace Math {
 
-	constexpr static FVector2<T> Zero()
-	{
-		return FVector2<T>((T)0, (T)0);
-	}
-	constexpr static FVector2<T> One()
-	{
-		return FVector2<T>((T)1, (T)1);
-	}
-	constexpr static FVector2<T> UnitX()
-	{
-		return FVector2<T>((T)1, (T)0);
-	}
-	constexpr static FVector2<T> UnitY()
+	template <typename T>
+	struct TVector2
 	{
-		return FVector2<T>((T)0, (T)1);
-	}
+		using FReal = T;
 
-	constexpr FVector2<T>& operator=(const FVector2<T>& V2)
-	{
-		X = V2.X;
-		Y = V2.Y;
-		return *this;
-	}
+		T X{}, Y{};
 
-	constexpr T& operator[](int Idx)
-	{
-		return (&X)[Idx];
-	}
-	constexpr const T& operator[](int Idx) const
-	{
-		return (&X)[Idx];
-	}
+		TVector2() {}
+		TVector2(T InF) : X(InF), Y(InF) {}
+		TVector2(T InX, T InY) : X(InX), Y(InY) {}
+		TVector2(const FVector2D& Vec) : X((T)Vec.X), Y((T)Vec.Y) { }
 
-	constexpr T Length() const
-	{
-		return TMathUtil<T>::Sqrt(X * X + Y * Y);
-	}
-	T SquaredLength() const
-	{
-		return X * X + Y * Y;
-	}
+		template<typename FArg, TEMPLATE_REQUIRES(!TIsSame<T, FArg>::Value)>
+		explicit TVector2(const TVector2<FArg>& From) : TVector2<T>((T)From.X, (T)From.Y) {}
 
-	constexpr T Distance(const FVector2<T>& V2) const
-	{
-		T dx = V2.X - X;
-		T dy = V2.Y - Y;
-		return TMathUtil<T>::Sqrt(dx * dx + dy * dy);
-	}
-	constexpr T DistanceSquared(const FVector2<T>& V2) const
-	{
-		T dx = V2.X - X;
-		T dy = V2.Y - Y;
-		return dx * dx + dy * dy;
-	}
+		explicit operator FVector2D() const { return FVector2D((float)X, (float)Y); }
 
-	constexpr T Dot(const FVector2<T>& V2) const
-	{
-		return X * V2.X + Y * V2.Y;
-	}
+		static inline TVector2<T> Zero() { return TVector2<T>((T)0, (T)0); }
+		static inline TVector2<T> One() { return TVector2<T>((T)1, (T)1); }
+		static inline TVector2<T> UnitX() { return TVector2<T>((T)1, (T)0); }
+		static inline TVector2<T> UnitY() { return TVector2<T>((T)0, (T)1); }
 
-	constexpr  FVector2 operator-() const
-	{
-		return FVector2(-X, -Y);
-	}
+		TVector2<T>& operator=(const TVector2<T>& V2) { X = V2.X; Y = V2.Y;	return *this; }
 
-	constexpr FVector2 operator+(const FVector2& V2) const
-	{
-		return FVector2(X + V2.X, Y + V2.Y);
-	}
+		T& operator[](int Idx) { return (&X)[Idx]; }
+		T operator[](int Idx) const { return (&X)[Idx]; }
 
-	constexpr FVector2 operator-(const FVector2& V2) const
-	{
-		return FVector2(X - V2.X, Y - V2.Y);
-	}
+		T Length() const { return TMathUtil<T>::Sqrt(X * X + Y * Y); }
+		T SquaredLength() const { return X * X + Y * Y; }
+		T Dot(const TVector2<T>& V2) const { return X * V2.X + Y * V2.Y; }
 
-	constexpr FVector2<T> operator+(const T& Scalar) const
-	{
-		return FVector2<T>(X + Scalar, Y + Scalar);
-	}
+		TVector2 operator-() const { return TVector2(-X, -Y); }
+		TVector2 operator+(const TVector2& V2) const { return TVector2(X + V2.X, Y + V2.Y); }
+		TVector2 operator-(const TVector2& V2) const { return TVector2(X - V2.X, Y - V2.Y); }
+		TVector2<T> operator+(const T& Scalar) const { return TVector2<T>(X + Scalar, Y + Scalar); }
+		TVector2<T> operator-(const T& Scalar) const { return TVector2<T>(X - Scalar, Y - Scalar); }
+		TVector2<T> operator*(const T& Scalar) const { return TVector2<T>(X * Scalar, Y * Scalar); }
+		TVector2<T> operator*(const TVector2<T>& V2) const { return TVector2<T>(X * V2.X, Y * V2.Y); }
+		TVector2<T> operator/(const T& Scalar) const { return TVector2<T>(X / Scalar, Y / Scalar); }
+		TVector2<T> operator/(const TVector2<T>& V2) const { return TVector2<T>(X / V2.X, Y / V2.Y); }
+		TVector2<T>& operator+=(const TVector2<T>& V2) { X += V2.X; Y += V2.Y; return *this; }
+		TVector2<T>& operator-=(const TVector2<T>& V2) { X -= V2.X; Y -= V2.Y; return *this; }
+		TVector2<T>& operator*=(const T& Scalar) { X *= Scalar; Y *= Scalar; return *this; }
+		TVector2<T>& operator/=(const T& Scalar) { X /= Scalar; Y /= Scalar; return *this; }
+		bool operator==(const TVector2<T>& Other) const { return X == Other.X && Y == Other.Y; }
+		bool operator!=(const TVector2<T>& Other) const	{ return X != Other.X || Y != Other.Y; }
 
-	constexpr FVector2<T> operator-(const T& Scalar) const
-	{
-		return FVector2<T>(X - Scalar, Y - Scalar);
-	}
+		template<typename RealType2>
+		TVector2<T> operator*(const RealType2& Scalar) const { return TVector2<T>(X * (T)Scalar, Y * (T)Scalar); }
 
-	constexpr FVector2<T> operator*(const T& Scalar) const
-	{
-		return FVector2<T>(X * Scalar, Y * Scalar);
-	}
+		friend FArchive& operator<<(FArchive& Ar, TVector2& Vec) { Vec.Serialize(Ar); return Ar; }
+		void Serialize(FArchive& Ar) { Ar << X; Ar << Y; }
+	};
 
-	template<typename RealType2>
-	constexpr FVector2<T> operator*(const RealType2& Scalar) const
-	{
-		return FVector2<T>(X * (T)Scalar, Y * (T)Scalar);
-	}
 
-	constexpr FVector2<T> operator*(const FVector2<T>& V2) const // component-wise
-	{
-		return FVector2<T>(X * V2.X, Y * V2.Y);
-	}
+	template <typename RealType>
+	inline TVector2<RealType> operator*(RealType Scalar, const TVector2<RealType>& V) {	return TVector2<RealType>(Scalar * V.X, Scalar * V.Y); }
 
-	constexpr FVector2<T> operator/(const T& Scalar) const
-	{
-		return FVector2<T>(X / Scalar, Y / Scalar);
-	}
+	template <typename T>
+	FORCEINLINE uint32 GetTypeHash(const TVector2<T>& Vector) { return FCrc::MemCrc_DEPRECATED(&Vector, sizeof(TVector2<T>)); }
 
-	constexpr FVector2<T> operator/(const FVector2<T>& V2) const // component-wise
-	{
-		return FVector2<T>(X / V2.X, Y / V2.Y);
-	}
+	template <typename RealType>
+	std::ostream& operator<<(std::ostream& os, const TVector2<RealType>& Vec) { os << Vec.X << " " << Vec.Y; return os;	}
 
-	constexpr FVector2<T>& operator+=(const FVector2<T>& V2)
-	{
-		X += V2.X;
-		Y += V2.Y;
-		return *this;
-	}
-
-	constexpr FVector2<T>& operator-=(const FVector2<T>& V2)
-	{
-		X -= V2.X;
-		Y -= V2.Y;
-		return *this;
-	}
-
-	constexpr FVector2<T>& operator*=(const T& Scalar)
-	{
-		X *= Scalar;
-		Y *= Scalar;
-		return *this;
-	}
+}
+}
 
-	constexpr FVector2<T>& operator/=(const T& Scalar)
-	{
-		X /= Scalar;
-		Y /= Scalar;
-		return *this;
-	}
+typedef UE::Math::TVector2<float> FVector2f;
+typedef UE::Math::TVector2<double> FVector2d;
+template<> struct TCanBulkSerialize<FVector2f> { enum { Value = true }; };
+template<> struct TCanBulkSerialize<FVector2d> { enum { Value = true }; };
 
-	constexpr bool operator==(const FVector2<T>& Other) const
-	{
-		return X == Other.X && Y == Other.Y;
-	}
+namespace UE {
+namespace Geometry {
 
-	constexpr bool operator!=(const FVector2<T>& Other) const
-	{
-		return X != Other.X || Y != Other.Y;
-	}
 
-	/**
-	 * Serialization operator for FVector2.
-	 *
-	 * @param Ar Archive to serialize with.
-	 * @param Vec Vector to serialize.
-	 * @returns Passing down serializing archive.
-	 */
-	friend FArchive& operator<<(FArchive& Ar, FVector2& Vec)
-	{
-		Vec.Serialize(Ar);
-		return Ar;
-	}
-
-	/** Serialize FVector2 to an archive. */
-	void Serialize(FArchive& Ar)
-	{
-		Ar << X;
-		Ar << Y;
-	}
-};
 
 /** @return dot product of V1 with PerpCW(V2), ie V2 rotated 90 degrees clockwise */
 template <typename T>
-constexpr T DotPerp(const FVector2<T>& V1, const FVector2<T>& V2)
+constexpr T DotPerp(const UE::Math::TVector2<T>& V1, const UE::Math::TVector2<T>& V2)
 {
 	return V1.X * V2.Y - V1.Y * V2.X;
 }
 
 /** @return right-Perpendicular vector to V, ie V rotated 90 degrees clockwise */
 template <typename T>
-constexpr FVector2<T> PerpCW(const FVector2<T>& V)
+constexpr UE::Math::TVector2<T> PerpCW(const UE::Math::TVector2<T>& V)
 {
-	return FVector2<T>(V.Y, -V.X);
+	return UE::Math::TVector2<T>(V.Y, -V.X);
 }
 
 /** @return > 0 if C is to the left of the line from A to B, < 0 if to the right, 0 if on the line */
 template<typename T>
-T Orient(const FVector2<T>& A, const FVector2<T>& B, const FVector2<T>& C)
+T Orient(const UE::Math::TVector2<T>& A, const UE::Math::TVector2<T>& B, const UE::Math::TVector2<T>& C)
 {
 	return DotPerp((B - A), (C - A));
 }
 
 
 template <typename T>
-constexpr bool IsNormalized(const FVector2<T>& Vector, const T Tolerance = TMathUtil<T>::ZeroTolerance)
+constexpr bool IsNormalized(const UE::Math::TVector2<T>& Vector, const T Tolerance = TMathUtil<T>::ZeroTolerance)
 {
 	return TMathUtil<T>::Abs((Vector.X*Vector.X + Vector.Y*Vector.Y) - 1) < Tolerance;
 }
 
 template <typename T>
-T Normalize(FVector2<T>& Vector, const T Epsilon = 0)
+T Normalize(UE::Math::TVector2<T>& Vector, const T Epsilon = 0)
 {
 	 T length = Vector.Length();
 	 if (length > Epsilon)
@@ -270,21 +136,38 @@ T Normalize(FVector2<T>& Vector, const T Epsilon = 0)
 }
 
 template <typename T>
-FVector2<T> Normalized(const FVector2<T>& Vector, const T Epsilon = 0)
+UE::Math::TVector2<T> Normalized(const UE::Math::TVector2<T>& Vector, const T Epsilon = 0)
 {
 	T length = Vector.Length();
 	if (length > Epsilon)
 	{
 		T invLength = ((T)1) / length;
-		return FVector2<T>(Vector.X*invLength, Vector.Y*invLength);
+		return UE::Math::TVector2<T>(Vector.X*invLength, Vector.Y*invLength);
 	}
-	return FVector2<T>((T)0, (T)0);
+	return UE::Math::TVector2<T>((T)0, (T)0);
+}
+
+
+template<typename T>
+T Distance(const UE::Math::TVector2<T>& V1, const UE::Math::TVector2<T>& V2)
+{
+	T dx = V2.X - V1.X;
+	T dy = V2.Y - V1.Y;
+	return TMathUtil<T>::Sqrt(dx * dx + dy * dy);
+}
+
+template<typename T>
+T DistanceSquared(const UE::Math::TVector2<T>& V1, const UE::Math::TVector2<T>& V2)
+{
+	T dx = V2.X - V1.X;
+	T dy = V2.Y - V1.Y;
+	return dx * dx + dy * dy;
 }
 
 
 // Angle in Degrees
 template <typename T>
-T AngleD(const FVector2<T>& V1, const FVector2<T>& V2)
+T AngleD(const UE::Math::TVector2<T>& V1, const UE::Math::TVector2<T>& V2)
 {
 	T DotVal = V1.Dot(V2);
 	T ClampedDot = (DotVal < (T)-1) ? (T)-1 : ((DotVal > (T)1) ? (T)1 : DotVal);
@@ -293,7 +176,7 @@ T AngleD(const FVector2<T>& V1, const FVector2<T>& V2)
 
 // Angle in Radians
 template <typename T>
-T AngleR(const FVector2<T>& V1, const FVector2<T>& V2)
+T AngleR(const UE::Math::TVector2<T>& V1, const UE::Math::TVector2<T>& V2)
 {
 	T DotVal = V1.Dot(V2);
 	T ClampedDot = (DotVal < (T)-1) ? (T)-1 : ((DotVal > (T)1) ? (T)1 : DotVal);
@@ -302,7 +185,7 @@ T AngleR(const FVector2<T>& V1, const FVector2<T>& V2)
 
 // Angle in Radians
 template <typename T>
-T SignedAngleR(const FVector2<T>& V1, const FVector2<T>& V2)
+T SignedAngleR(const UE::Math::TVector2<T>& V1, const UE::Math::TVector2<T>& V2)
 {
 	T DotVal = V1.Dot(V2);
 	T ClampedDot = (DotVal < (T)-1) ? (T)-1 : ((DotVal > (T)1) ? (T)1 : DotVal);
@@ -319,10 +202,10 @@ T SignedAngleR(const FVector2<T>& V1, const FVector2<T>& V2)
 }
 
 template <typename T>
-FVector2<T> Lerp(const FVector2<T>& A, const FVector2<T>& B, T Alpha)
+UE::Math::TVector2<T> Lerp(const UE::Math::TVector2<T>& A, const UE::Math::TVector2<T>& B, T Alpha)
 {
 	T OneMinusAlpha = (T)1 - Alpha;
-	return FVector2<T>(OneMinusAlpha * A.X + Alpha * B.X,
+	return UE::Math::TVector2<T>(OneMinusAlpha * A.X + Alpha * B.X,
 		OneMinusAlpha * A.Y + Alpha * B.Y);
 }
 
@@ -688,21 +571,21 @@ T AngleR(const UE::Math::TVector<T>& V1, const UE::Math::TVector<T>& V2)
 }
 
 template <typename T>
-constexpr FVector2<T> GetXY(const UE::Math::TVector<T>& V)
+constexpr UE::Math::TVector2<T> GetXY(const UE::Math::TVector<T>& V)
 {
-	return FVector2<T>(V.X, V.Y);
+	return UE::Math::TVector2<T>(V.X, V.Y);
 }
 
 template <typename T>
-constexpr FVector2<T> GetXZ(const UE::Math::TVector<T>& V)
+constexpr UE::Math::TVector2<T> GetXZ(const UE::Math::TVector<T>& V)
 {
-	return FVector2<T>(V.X, V.Z);
+	return UE::Math::TVector2<T>(V.X, V.Z);
 }
 
 template <typename T>
-constexpr FVector2<T> GetYZ(const UE::Math::TVector<T>& V)
+constexpr UE::Math::TVector2<T> GetYZ(const UE::Math::TVector<T>& V)
 {
-	return FVector2<T>(V.Y, V.Z);
+	return UE::Math::TVector2<T>(V.Y, V.Z);
 }
 
 
@@ -840,34 +723,6 @@ FORCEINLINE uint32 GetTypeHash(const FVector3<T>& Vector)
 
 
 
-template <typename T>
-FORCEINLINE uint32 GetTypeHash(const FVector2<T>& Vector)
-{
-	// (this is how FIntVector and all the other FVectors do their hash functions)
-	// Note: this assumes there's no padding that could contain uncompared data.
-	return FCrc::MemCrc_DEPRECATED(&Vector, sizeof(FVector2<T>));
-}
-
-template <typename RealType>
-inline FVector2<RealType> operator*(RealType Scalar, const FVector2<RealType>& V)
-{
-	return FVector2<RealType>(Scalar * V.X, Scalar * V.Y);
-}
-
-// allow float*Vector2<double> and double*Vector2<float>
-template <typename RealType, typename RealType2>
-inline FVector2<RealType> operator*(RealType2 Scalar, const FVector2<RealType>& V)
-{
-	return FVector2<RealType>((RealType)Scalar * V.X, (RealType)Scalar * V.Y);
-}
-
-template <typename RealType>
-std::ostream& operator<<(std::ostream& os, const FVector2<RealType>& Vec)
-{
-	os << Vec.X << " " << Vec.Y;
-	return os;
-}
-
 
 
 
@@ -1159,16 +1014,16 @@ FORCEINLINE uint32 GetTypeHash(const TVector4<T>& Vector)
 } // end namespace UE::Geometry
 } // end namespace UE
 
-typedef UE::Geometry::FVector2<float> FVector2f;
-typedef UE::Geometry::FVector2<double> FVector2d;
+//typedef UE::Geometry::FVector2<float> FVector2f;
+//typedef UE::Geometry::FVector2<double> FVector2d;
 
 //typedef UE::Geometry::FVector3<float> FVector3f;
 //typedef UE::Geometry::FVector3<double> FVector3d;
 
 typedef UE::Geometry::TVector4<float> FVector4f;
 typedef UE::Geometry::TVector4<double> FVector4d;
-template<> struct TCanBulkSerialize<FVector2f> { enum { Value = true }; };
-template<> struct TCanBulkSerialize<FVector2d> { enum { Value = true }; };
+//template<> struct TCanBulkSerialize<FVector2f> { enum { Value = true }; };
+//template<> struct TCanBulkSerialize<FVector2d> { enum { Value = true }; };
 
 template<> struct TCanBulkSerialize<FVector4f> { enum { Value = true }; };
 template<> struct TCanBulkSerialize<FVector4d> { enum { Value = true }; };
diff --git a/Engine/Source/Runtime/GeometryCore/Public/VectorUtil.h b/Engine/Source/Runtime/GeometryCore/Public/VectorUtil.h
--- a/Engine/Source/Runtime/GeometryCore/Public/VectorUtil.h
+++ b/Engine/Source/Runtime/GeometryCore/Public/VectorUtil.h
@@ -29,12 +29,13 @@ enum class EIntersectionType
 namespace VectorUtil
 {
 	using namespace UE::Geometry;
+	using namespace UE::Math;
 
 	/**
 	 * @return true if all components of V are finite
 	 */
 	template <typename RealType>
-	inline bool IsFinite(const FVector2<RealType>& V)
+	inline bool IsFinite(const TVector2<RealType>& V)
 	{
 		return TMathUtil<RealType>::IsFinite(V.X) && TMathUtil<RealType>::IsFinite(V.Y);
 	}
@@ -101,10 +102,10 @@ namespace VectorUtil
 	 * @return area of 2D triangle V0,V1,V2
 	 */
 	template <typename RealType>
-	inline RealType Area(const FVector2<RealType>& V0, const FVector2<RealType>& V1, const FVector2<RealType>& V2)
+	inline RealType Area(const TVector2<RealType>& V0, const TVector2<RealType>& V1, const TVector2<RealType>& V2)
 	{
-		FVector2<RealType> Edge1(V1 - V0);
-		FVector2<RealType> Edge2(V2 - V0);
+		TVector2<RealType> Edge1(V1 - V0);
+		TVector2<RealType> Edge2(V2 - V0);
 		RealType CrossZ = DotPerp(Edge1, Edge2);
 		return (RealType)0.5 * TMathUtil<RealType>::Abs(CrossZ);
 	}
@@ -113,7 +114,7 @@ namespace VectorUtil
 	 * @return signed area of 2D triangle V0,V1,V2
 	 */
 	template <typename RealType>
-	inline RealType SignedArea(const FVector2<RealType>& V0, const FVector2<RealType>& V1, const FVector2<RealType>& V2)
+	inline RealType SignedArea(const TVector2<RealType>& V0, const TVector2<RealType>& V1, const TVector2<RealType>& V2)
 	{
 		return ((RealType)0.5) * ((V0.X * V1.Y - V0.Y * V1.X) + (V1.X * V2.Y - V1.Y * V2.X) + (V2.X * V0.Y - V2.Y * V0.X));
 	}
@@ -155,7 +156,7 @@ namespace VectorUtil
 
 	/** @return true if all coordinates of V0 and V1 are within Epsilon of eachother */
 	template <typename RealType>
-	inline bool EpsilonEqual(const FVector2<RealType>& V0, const FVector2<RealType>& V1, RealType Epsilon)
+	inline bool EpsilonEqual(const TVector2<RealType>& V0, const TVector2<RealType>& V1, RealType Epsilon)
 	{
 		return EpsilonEqual(V0.X, V1.X, Epsilon) && EpsilonEqual(V0.Y, V1.Y, Epsilon);
 	}
@@ -335,11 +336,11 @@ namespace VectorUtil
 	* TODO: make robust to degenerate triangles?
 	*/
 	template <typename RealType>
-	FVector3<RealType> BarycentricCoords(const FVector2<RealType>& Point, const FVector2<RealType>& V0, const FVector2<RealType>& V1, const FVector2<RealType>& V2)
+	FVector3<RealType> BarycentricCoords(const TVector2<RealType>& Point, const TVector2<RealType>& V0, const TVector2<RealType>& V1, const TVector2<RealType>& V2)
 	{
-		FVector2<RealType> kV02 = V0 - V2;
-		FVector2<RealType> kV12 = V1 - V2;
-		FVector2<RealType> kPV2 = Point - V2;
+		TVector2<RealType> kV02 = V0 - V2;
+		TVector2<RealType> kV12 = V1 - V2;
+		TVector2<RealType> kPV2 = Point - V2;
 		RealType fM00 = kV02.Dot(kV02);
 		RealType fM01 = kV02.Dot(kV12);
 		RealType fM11 = kV12.Dot(kV12);
@@ -387,11 +388,11 @@ namespace VectorUtil
 		FVector3<RealType> Perp0, Perp1;
 		VectorUtil::MakePerpVectors<RealType>(Normal, Perp0, Perp1);
 		// project points to triangle plane coordinates
-		FVector2<RealType> vi(Vi.Dot(Perp0), Vi.Dot(Perp1));
-		FVector2<RealType> vj(Vj.Dot(Perp0), Vj.Dot(Perp1));
-		FVector2<RealType> vk(Vk.Dot(Perp0), Vk.Dot(Perp1));
+		TVector2<RealType> vi(Vi.Dot(Perp0), Vi.Dot(Perp1));
+		TVector2<RealType> vj(Vj.Dot(Perp0), Vj.Dot(Perp1));
+		TVector2<RealType> vk(Vk.Dot(Perp0), Vk.Dot(Perp1));
 		// calculate gradient
-		FVector2<RealType> GradX = (fj-fi)*PerpCW(vi-vk) + (fk-fi)*PerpCW(vj-vi);
+		TVector2<RealType> GradX = (fj-fi)*PerpCW(vi-vk) + (fk-fi)*PerpCW(vj-vi);
 		// map back to 3D vector in triangle plane
 		RealType AreaScale = (RealType)1 / ((RealType)2 * VectorUtil::Area<RealType>(Vi, Vj, Vk));
 		return AreaScale * (GradX.X * Perp0 + GradX.Y * Perp1);
