GeometryCore: minor DynamicMesh fixes. Change ensure to a checkSlow in DynamicMesh3::AppendTriangle. Add FDynamicMeshOverlay::GetTriangleIfValid(). Handle unset triangles in CountVertexElements().
#rnx

diff --git a/Engine/Source/Runtime/GeometryCore/Private/DynamicMesh/DynamicMesh3_Edits.cpp b/Engine/Source/Runtime/GeometryCore/Private/DynamicMesh/DynamicMesh3_Edits.cpp
--- a/Engine/Source/Runtime/GeometryCore/Private/DynamicMesh/DynamicMesh3_Edits.cpp
+++ b/Engine/Source/Runtime/GeometryCore/Private/DynamicMesh/DynamicMesh3_Edits.cpp
@@ -155,8 +155,9 @@ int FDynamicMesh3::AppendTriangle(const FIndex3i& tv, int gid)
 		checkSlow(false);
 		return InvalidID;
 	}
-	if (!ensure(tv[0] != tv[1] && tv[0] != tv[2] && tv[1] != tv[2]))
+	if (tv[0] == tv[1] || tv[0] == tv[2] || tv[1] == tv[2])
 	{
+		checkSlow(false);
 		return InvalidID;
 	}
 
diff --git a/Engine/Source/Runtime/GeometryCore/Private/DynamicMesh/DynamicMeshOverlay.cpp b/Engine/Source/Runtime/GeometryCore/Private/DynamicMesh/DynamicMeshOverlay.cpp
--- a/Engine/Source/Runtime/GeometryCore/Private/DynamicMesh/DynamicMeshOverlay.cpp
+++ b/Engine/Source/Runtime/GeometryCore/Private/DynamicMesh/DynamicMeshOverlay.cpp
@@ -785,16 +785,19 @@ template<typename RealType, int ElementSize>
 int TDynamicMeshOverlay<RealType, ElementSize>::CountVertexElements(int vid, bool bBruteForce) const
 {
 	TArray<int> VertexElements;
+	FIndex3i Triangle;
 	if (bBruteForce) 
 	{
 		for (int tid : ParentMesh->TriangleIndicesItr())
 		{
-			FIndex3i Triangle = GetTriangle(tid);
-			for (int j = 0; j < 3; ++j)
+			if (GetTriangleIfValid(tid, Triangle))
 			{
-				if (ParentVertices[Triangle[j]] == vid)
+				for (int j = 0; j < 3; ++j)
 				{
-					VertexElements.AddUnique(Triangle[j]);
+					if (ParentVertices[Triangle[j]] == vid)
+					{
+						VertexElements.AddUnique(Triangle[j]);
+					}
 				}
 			}
 		}
@@ -803,12 +806,14 @@ int TDynamicMeshOverlay<RealType, ElementSize>::CountVertexElements(int vid, boo
 	{
 		for (int tid : ParentMesh->VtxTrianglesItr(vid))
 		{
-			FIndex3i Triangle = GetTriangle(tid);
-			for (int j = 0; j < 3; ++j)
+			if (GetTriangleIfValid(tid, Triangle))
 			{
-				if (ParentVertices[Triangle[j]] == vid)
+				for (int j = 0; j < 3; ++j)
 				{
-					VertexElements.AddUnique(Triangle[j]);
+					if (ParentVertices[Triangle[j]] == vid)
+					{
+						VertexElements.AddUnique(Triangle[j]);
+					}
 				}
 			}
 		}
diff --git a/Engine/Source/Runtime/GeometryCore/Public/DynamicMesh/DynamicMeshOverlay.h b/Engine/Source/Runtime/GeometryCore/Public/DynamicMesh/DynamicMeshOverlay.h
--- a/Engine/Source/Runtime/GeometryCore/Public/DynamicMesh/DynamicMeshOverlay.h
+++ b/Engine/Source/Runtime/GeometryCore/Public/DynamicMesh/DynamicMeshOverlay.h
@@ -399,6 +399,19 @@ public:
 		return FIndex3i(ElementTriangles[i], ElementTriangles[i + 1], ElementTriangles[i + 2]);
 	}
 
+	/** If the triangle is set to valid element indices, return the indices in TriangleOut and return true, otherwise return false */
+	inline bool GetTriangleIfValid(int TriangleID, FIndex3i& TriangleOut) const 
+	{
+		int i = 3 * TriangleID;
+		int a = ElementTriangles[i];
+		if ( a >= 0 )
+		{
+			TriangleOut = FIndex3i(a, ElementTriangles[i+1], ElementTriangles[i+2]);
+			checkSlow(TriangleOut.B >= 0 && TriangleOut.C >= 0);
+			return true;
+		}
+		return false;
+	}
 
 	/** Set the element at a given index */
 	inline void SetElement(int ElementID, const RealType* Data)
