DynamicMesh3 - Add direct buffer access that returns const refrence to TriangleEdges.
Also make all other direct buffer access methods that return const references, const methods

diff --git a/Engine/Source/Runtime/GeometryCore/Public/DynamicMesh/DynamicMesh3.h b/Engine/Source/Runtime/GeometryCore/Public/DynamicMesh/DynamicMesh3.h
--- a/Engine/Source/Runtime/GeometryCore/Public/DynamicMesh/DynamicMesh3.h
+++ b/Engine/Source/Runtime/GeometryCore/Public/DynamicMesh/DynamicMesh3.h
@@ -1048,51 +1048,54 @@ public:
 	// direct buffer access
 	//
 public:
-	const TDynamicVector<FVector3d>& GetVerticesBuffer()
+	const TDynamicVector<FVector3d>& GetVerticesBuffer() const 
 	{
 		return Vertices;
 	}
-	const FRefCountVector& GetVerticesRefCounts()
+	const FRefCountVector& GetVerticesRefCounts() const 
 	{
 		return VertexRefCounts;
 	}
-	const TDynamicVector<FVector3f>* GetNormalsBuffer()
+	const TDynamicVector<FVector3f>* GetNormalsBuffer() const 
 	{
 		return HasVertexNormals() ? &VertexNormals.GetValue() : nullptr;
 	}
-	const TDynamicVector<FVector3f>* GetColorsBuffer()
+	const TDynamicVector<FVector3f>* GetColorsBuffer() const 
 	{
 		return HasVertexColors() ? &VertexColors.GetValue() : nullptr;
 	}
-	const TDynamicVector<FVector2f>* GetUVBuffer()
+	const TDynamicVector<FVector2f>* GetUVBuffer() const 
 	{
 		return HasVertexUVs() ? &VertexUVs.GetValue() : nullptr;
 	}
-	const TDynamicVector<FIndex3i>& GetTrianglesBuffer()
+	const TDynamicVector<FIndex3i>& GetTrianglesBuffer() const 
 	{
 		return Triangles;
 	}
-	const FRefCountVector& GetTrianglesRefCounts()
+	const FRefCountVector& GetTrianglesRefCounts() const
 	{
 		return TriangleRefCounts;
 	}
-	const TDynamicVector<int>* GetTriangleGroupsBuffer()
+	const TDynamicVector<int>* GetTriangleGroupsBuffer() const
 	{
 		return HasTriangleGroups() ? &TriangleGroups.GetValue() : nullptr;
 	}
-	const TDynamicVector<FEdge>& GetEdgesBuffer()
+	const TDynamicVector<FEdge>& GetEdgesBuffer() const 
 	{
 		return Edges;
 	}
-	const FRefCountVector& GetEdgesRefCounts()
+	const FRefCountVector& GetEdgesRefCounts() const 
 	{
 		return EdgeRefCounts;
 	}
-	const FSmallListSet& GetVertexEdges()
+	const FSmallListSet& GetVertexEdges() const 
 	{
 		return VertexEdgeLists;
 	}
-
+	const TDynamicVector<FIndex3i>& GetTriangleEdges() const 
+	{
+		return TriangleEdges;
+	}
 	//
 	// Mesh Edit operations
 	//
