GeometryCore: Fix DynamicVector::Resize() over-allocating and assigning invalid internal indices

#rnx

diff --git a/Engine/Source/Runtime/GeometryCore/Public/Util/DynamicVector.h b/Engine/Source/Runtime/GeometryCore/Public/Util/DynamicVector.h
--- a/Engine/Source/Runtime/GeometryCore/Public/Util/DynamicVector.h
+++ b/Engine/Source/Runtime/GeometryCore/Public/Util/DynamicVector.h
@@ -79,7 +79,7 @@ public:
 	inline bool IsEmpty() const { return CurBlock == 0 && CurBlockUsed == 0; }
 	inline size_t GetLength() const { return CurBlock * BlockSize + CurBlockUsed; }
 	inline size_t Num() const { return CurBlock * BlockSize + CurBlockUsed; }
-	inline int GetBlockSize() const { return BlockSize; }
+	static int GetBlockSize() { return BlockSize; }
 	inline size_t GetByteCount() const { return Blocks.Num() * BlockSize * sizeof(Type); }
 
 	inline void Add(const Type& Data);
@@ -90,8 +90,17 @@ public:
 	inline void InsertAt(const Type& Data, unsigned int Index, const Type& InitValue);
 	inline Type& ElementAt(unsigned int Index, Type InitialValue = Type{});
 
-	inline const Type& Front() const { return Blocks[0][0]; }
-	inline const Type& Back() const { return Blocks[CurBlock][CurBlockUsed - 1]; }
+	inline const Type& Front() const
+	{
+		checkSlow(CurBlockUsed > 0);
+		return Blocks[0][0];
+	}
+
+	inline const Type& Back() const
+	{
+		checkSlow(CurBlockUsed > 0);
+		return Blocks[CurBlock][CurBlockUsed - 1];
+	}
 
 	inline Type& operator[](unsigned int Index)
 	{
@@ -242,8 +251,8 @@ private:
 	static constexpr int BlockIndexBitmask = BlockSize - 1; // low 9 bits
 	static_assert( BlockSize && ((BlockSize & (BlockSize - 1)) == 0), "DynamicVector: BlockSize must be a power of two");
 
-	unsigned int CurBlock{0};
-	unsigned int CurBlockUsed{0};
+	unsigned int CurBlock{0};  //< Current block index; always points to the block with the last item in the vector or it is set to zero if the vector is empty. 
+	unsigned int CurBlockUsed{0};  //< Number of used items in the current block.
 
 	using BlockType = TStaticArray<Type, BlockSize>;
 	
@@ -401,6 +410,34 @@ private:
 		}
 	};
 
+	friend bool operator==(const TDynamicVector& Lhs, const TDynamicVector& Rhs)
+	{
+		if (Lhs.Num() != Rhs.Num())
+		{
+			return false;
+		}
+
+		if (Lhs.IsEmpty())
+		{
+			return true;
+		}
+		
+		const uint32 LhsCurBlock = Lhs.CurBlock;
+		for (uint32 BlockIndex = 0; BlockIndex < LhsCurBlock; ++BlockIndex)
+		{
+			if (!CompareItems(Lhs.Blocks[BlockIndex].GetData(), Rhs.Blocks[BlockIndex].GetData(), BlockSize))
+			{
+				return false;
+			}
+		}
+		return CompareItems(Lhs.Blocks[LhsCurBlock].GetData(), Rhs.Blocks[LhsCurBlock].GetData(), Lhs.CurBlockUsed);
+	}
+
+	friend bool operator!=(const TDynamicVector& Lhs, const TDynamicVector& Rhs)
+	{
+		return !(Lhs == Rhs);
+	}
+
 	// helper function - avoids MSVC+ASan internal compiler error in [] operators
 	inline void GetIndices( unsigned int Index, int& ArrayIndex, int& BlockIndex ) const
 	{
@@ -597,29 +634,32 @@ void TDynamicVector<Type>::Resize(size_t Count)
 		return;
 	}
 
-	// figure out how many segments we need
-	int nNumSegs = 1 + (int)Count / BlockSize;
+	// Determine how many blocks we need, but make sure we have at least one block available.
+	const bool bCountIsNotMultipleOfBlockSize = Count % BlockSize != 0;
+	const int32 NumBlocksNeeded = FMath::Max(1, static_cast<int32>(Count) / BlockSize + bCountIsNotMultipleOfBlockSize);
 
-	// figure out how many are currently allocated...
-	int32 nCurCount = Blocks.Num();
+	// Determine how many blocks are currently allocated.
+	int32 NumBlocksCurrent = Blocks.Num();
 
-	// resize to right number of segments
-	if (nNumSegs >= nCurCount)
+	// Allocate needed additional blocks.
+	while (NumBlocksCurrent < NumBlocksNeeded)
 	{
-		// allocate new segments
-		for (int i = (int)nCurCount; i < nNumSegs; ++i)
-		{
-			Blocks.Add(new BlockType());
-		}
+		Blocks.Add(new BlockType());
+		++NumBlocksCurrent;
 	}
-	else
+
+	// Remove unneeded blocks.
+	if (NumBlocksCurrent > NumBlocksNeeded)
 	{
-		Blocks.Truncate(nNumSegs, false);
+		Blocks.Truncate(NumBlocksNeeded, false);
 	}
 
-	// mark last segment
-	CurBlockUsed = (unsigned int)(Count - (nNumSegs - 1) * BlockSize);
-	CurBlock = nNumSegs - 1;
+	// Reset block index for the last item and used item count within the last block.
+	// This is similar to what happens when computing the indices in operator[], but we additionally account for (1) the vector being empty and (2) that the
+	// used item count within the last block needs to be one more than the index of the last item. 
+	const int32 LastItemIndex = Count - 1;
+	CurBlock = Count != 0 ? LastItemIndex >> nShiftBits : 0;
+	CurBlockUsed = Count != 0 ? (LastItemIndex & BlockIndexBitmask) + 1 : 0;
 }
 
 template <class Type>
