GeometryCore: Account for loading TDynamicVector<> with inconsistent internal current block data.

#rnx

diff --git a/Engine/Source/Runtime/GeometryCore/Public/Util/DynamicVector.h b/Engine/Source/Runtime/GeometryCore/Public/Util/DynamicVector.h
--- a/Engine/Source/Runtime/GeometryCore/Public/Util/DynamicVector.h
+++ b/Engine/Source/Runtime/GeometryCore/Public/Util/DynamicVector.h
@@ -140,6 +140,21 @@ public:
 		Ar << CurBlock;
 		Ar << CurBlockUsed;
 		Ar << Blocks;
+
+		// Account for a previously existing issue where one additional empty block was allocated and the resulting values for the current block can cause
+        // crashes, for example when using Back().
+        if (Ar.IsLoading() && CurBlock > 0 && CurBlockUsed == 0)
+        {
+            const SIZE_T Count = Num();
+            checkSlow(Count % BlockSize == 0);
+            checkSlow(Count / BlockSize < Blocks.Num());
+
+        	// Fix CurBlock and CurBlockUsed.
+            SetCurBlock(Count);
+
+        	// Remove empty block.
+            Blocks.Truncate(Count / BlockSize, false);
+        }
 	}
 
 public:
@@ -445,6 +460,16 @@ private:
 		BlockIndex = (Index & BlockIndexBitmask);
 	}
 
+	void SetCurBlock(SIZE_T Count)
+	{
+		// Reset block index for the last item and used item count within the last block.
+		// This is similar to what happens when computing the indices in operator[], but we additionally account for (1) the vector being empty and (2) that the
+		// used item count within the last block needs to be one more than the index of the last item. 
+		const int32 LastItemIndex = Count - 1;
+		CurBlock = Count != 0 ? LastItemIndex >> nShiftBits : 0;
+		CurBlockUsed = Count != 0 ? (LastItemIndex & BlockIndexBitmask) + 1 : 0;
+	}
+
 	// memory chunks for dynamic vector
 	TBlockVector<BlockType> Blocks;
 };
@@ -654,12 +679,8 @@ void TDynamicVector<Type>::Resize(size_t Count)
 		Blocks.Truncate(NumBlocksNeeded, false);
 	}
 
-	// Reset block index for the last item and used item count within the last block.
-	// This is similar to what happens when computing the indices in operator[], but we additionally account for (1) the vector being empty and (2) that the
-	// used item count within the last block needs to be one more than the index of the last item. 
-	const int32 LastItemIndex = Count - 1;
-	CurBlock = Count != 0 ? LastItemIndex >> nShiftBits : 0;
-	CurBlockUsed = Count != 0 ? (LastItemIndex & BlockIndexBitmask) + 1 : 0;
+	// Set current block.
+	SetCurBlock(Count);
 }
 
 template <class Type>
