GeometryCore: Add Embedded Path Tracing across FDynamicMesh surfaces and Intrinsic Triangulations of FDynamicMeshes.  Also update Mesh Laplacian Code to generate Cotan Laplacians from Intrinsic Triangulation.

MeshGeodesicSurfaceTracer.{h,cpp}  -  Add embedded path tracing.
GeodesicSurfaceTracerTests.cpp     -  unit tests for embedded path tracing

IntrinsicTriangulationMesh.{h, cpp}  - Adds two Intrinsic Triangulation mesh types (the more sophisticaed requires embedded path tracing)
IntrinsicMeshTests.cpp               -  unit tests for Intrinsic mesh types.

MeshLaplacian.{h,cpp}                - Adds new enum types for Intrinsic Delaunay Triangulation (IDT)  Laplacians
LaplacianOperators.{h,cpp}           - Adds wrappers and allow enum switches to build IDT Laplacians if requested
LaplacianMatrixAssembly.h            - Actually construction of IDT Laplacian.
Laplacian.cpp                        - Added trivial unit test for construction of IDT Laplacian

diff --git a/Engine/Plugins/Runtime/GeometryProcessing/Source/DynamicMesh/Private/Solvers/Internal/LaplacianOperators.cpp b/Engine/Plugins/Runtime/GeometryProcessing/Source/DynamicMesh/Private/Solvers/Internal/LaplacianOperators.cpp
--- a/Engine/Plugins/Runtime/GeometryProcessing/Source/DynamicMesh/Private/Solvers/Internal/LaplacianOperators.cpp
+++ b/Engine/Plugins/Runtime/GeometryProcessing/Source/DynamicMesh/Private/Solvers/Internal/LaplacianOperators.cpp
@@ -113,7 +113,20 @@ void ConstructCotangentLaplacian(const FDynamicMesh3& DynamicMesh, FVertexLinear
 	Boundary.ExtractResult(LaplacianBoundary);
 }
 
+void ConstructIDTCotangentLaplacian(const FDynamicMesh3& DynamicMesh, FVertexLinearization& VertexMap, FSparseMatrixD& LaplacianInterior, FSparseMatrixD& LaplacianBoundary, const bool bClampWeights)
+{
+	// Sync the mapping between the mesh vertex ids and their offsets in a nominal linear array.
+	VertexMap.Reset(DynamicMesh);
+	const int32 NumVerts = VertexMap.NumVerts();
+	const int32 NumBoundaryVerts = VertexMap.NumBoundaryVerts();
+	const int32 NumInteriorVerts = NumVerts - NumBoundaryVerts;
 
+	FEigenSparseMatrixAssembler Interior(NumInteriorVerts, NumInteriorVerts);
+	FEigenSparseMatrixAssembler Boundary(NumInteriorVerts, NumBoundaryVerts);
+	UE::MeshDeformation::ConstructIDTCotangentLaplacian<double>(DynamicMesh, VertexMap, Interior, Boundary, bClampWeights);
+	Interior.ExtractResult(LaplacianInterior);
+	Boundary.ExtractResult(LaplacianBoundary);
+}
 
 double ConstructScaledCotangentLaplacian(const FDynamicMesh3& DynamicMesh, FVertexLinearization& VertexMap, FSparseMatrixD& LaplacianInterior, FSparseMatrixD& LaplacianBoundary, const bool bClampAreas)
 {
@@ -194,6 +207,12 @@ void ConstructLaplacian(const ELaplacianWeightScheme Scheme, const FDynamicMesh3
 		ConstructScaledCotangentLaplacian(DynamicMesh, VertexMap, LaplacianInterior, LaplacianBoundary, bClampWeights);
 		break;
 	}
+	case ELaplacianWeightScheme::IDTCotanget:
+	{
+		bool bClampWeights = true;
+		ConstructIDTCotangentLaplacian(DynamicMesh, VertexMap, LaplacianInterior, LaplacianBoundary, bClampWeights);
+		break;
+	}
 	case ELaplacianWeightScheme::MeanValue:
 		ConstructMeanValueWeightLaplacian(DynamicMesh, VertexMap, LaplacianInterior, LaplacianBoundary);
 		break;
@@ -289,6 +308,20 @@ TUniquePtr<FSparseMatrixD> ConstructCotangentLaplacian(const FDynamicMesh3& Dyna
 	return LaplacianMatrix;
 }
 
+TUniquePtr<FSparseMatrixD> ConstructIDTCotangentLaplacian(const FDynamicMesh3& DynamicMesh, FVertexLinearization& VertexMap, const bool bClampWeights, TArray<int32>* BoundaryVerts)
+{
+	TUniquePtr<FSparseMatrixD> LaplacianMatrix(new FSparseMatrixD);
+	FSparseMatrixD BoundaryMatrix;
+
+	ConstructIDTCotangentLaplacian(DynamicMesh, VertexMap, *LaplacianMatrix, BoundaryMatrix, bClampWeights);
+
+	if (BoundaryVerts)
+	{
+		ExtractBoundaryVerts(VertexMap, *BoundaryVerts);
+	}
+
+	return LaplacianMatrix;
+}
 
 
 TUniquePtr<FSparseMatrixD> ConstructMeanValueWeightLaplacian(const FDynamicMesh3& DynamicMesh, FVertexLinearization& VertexMap, TArray<int32>* BoundaryVerts)
diff --git a/Engine/Plugins/Runtime/GeometryProcessing/Source/DynamicMesh/Private/Solvers/Internal/LaplacianOperators.h b/Engine/Plugins/Runtime/GeometryProcessing/Source/DynamicMesh/Private/Solvers/Internal/LaplacianOperators.h
--- a/Engine/Plugins/Runtime/GeometryProcessing/Source/DynamicMesh/Private/Solvers/Internal/LaplacianOperators.h
+++ b/Engine/Plugins/Runtime/GeometryProcessing/Source/DynamicMesh/Private/Solvers/Internal/LaplacianOperators.h
@@ -68,6 +68,11 @@ TUniquePtr<FSparseMatrixD> ConstructCotangentLaplacian(const FDynamicMesh3& Dyna
 	const bool bClampWeights = false,
 	TArray<int32>* BoundaryVerts = NULL);
 
+TUniquePtr<FSparseMatrixD> ConstructIDTCotangentLaplacian(const FDynamicMesh3& DynamicMesh,
+	FVertexLinearization& VertexMap,
+	const bool bClampWeights = false,
+	TArray<int32>* BoundaryVerts = NULL);
+
 TUniquePtr<FSparseMatrixD> ConstructCotangentLaplacian(const FDynamicMesh3& DynamicMesh,
 	FVertexLinearization& VertexMap,
 	FSparseMatrixD& AreaMatrix,
diff --git a/Engine/Plugins/Runtime/GeometryProcessing/Source/DynamicMesh/Private/Solvers/MeshLaplacian.cpp b/Engine/Plugins/Runtime/GeometryProcessing/Source/DynamicMesh/Private/Solvers/MeshLaplacian.cpp
--- a/Engine/Plugins/Runtime/GeometryProcessing/Source/DynamicMesh/Private/Solvers/MeshLaplacian.cpp
+++ b/Engine/Plugins/Runtime/GeometryProcessing/Source/DynamicMesh/Private/Solvers/MeshLaplacian.cpp
@@ -27,6 +27,9 @@ FString LaplacianSchemeName(const ELaplacianWeightScheme Scheme)
 	case ELaplacianWeightScheme::Valence:
 		LaplacianName = FString(TEXT("Valence Laplacian"));
 		break;
+	case ELaplacianWeightScheme::IDTCotanget:
+		LaplacianName = FString(TEXT("Intrinsic Delaunay Triangulation Contangent Laplacian"));
+		break;
 	default:
 		check(0 && "Unknown Laplacian Weight Scheme Enum");
 	}
@@ -60,6 +63,9 @@ bool bIsSymmetricLaplacian(const ELaplacianWeightScheme Scheme)
 	case ELaplacianWeightScheme::Valence:
 		bSymmetric = true;
 		break;
+	case ELaplacianWeightScheme::IDTCotanget:
+		bSymmetric = false;
+		break;
 	default:
 		check(0);
 	}
diff --git a/Engine/Plugins/Runtime/GeometryProcessing/Source/DynamicMesh/Public/Solvers/LaplacianMatrixAssembly.h b/Engine/Plugins/Runtime/GeometryProcessing/Source/DynamicMesh/Public/Solvers/LaplacianMatrixAssembly.h
--- a/Engine/Plugins/Runtime/GeometryProcessing/Source/DynamicMesh/Public/Solvers/LaplacianMatrixAssembly.h
+++ b/Engine/Plugins/Runtime/GeometryProcessing/Source/DynamicMesh/Public/Solvers/LaplacianMatrixAssembly.h
@@ -9,6 +9,7 @@
 #include "Solvers/MeshLinearization.h"
 #include "Solvers/MeshLaplacian.h"
 #include "Solvers/PrecomputedMeshWeightData.h"
+#include "Spatial/IntrinsicTriangulationMesh.h"
 
 
 namespace UE
@@ -167,8 +168,31 @@ namespace UE
 			UE::Solvers::TSparseMatrixAssembler<RealType>& LaplacianBoundary,
 			const bool bClampWeights);
 
-
-
+		/**
+		* Construct a sparse matrix representation using a pre-multiplied cotangent-weighted Laplacian, using an intrinsic Delaunay mesh internally 
+		* NB: there is no reason to expect this to be a symmetric matrix.
+		*
+		* This computes the laplacian scaled by the average area A_ave:  ie.  LScaled =   A_ave/(2A_i) ( Cot alpha_ij + Cot beta_ij )
+		*
+		* The mesh itself is assumed to have N interior vertices, and M boundary vertices.
+		*
+		* @param DynamicMesh        The triangle mesh
+		* @param VertexMap          Additional arrays used to map between vertexID and offset in a linear array (i.e. the row).
+		*                           The vertices are ordered so that last M ( = VertexMap.NumBoundaryVerts() )  correspond to those on the boundary.
+		* @param LaplacianInterior  On return, scaled laplacian operator that acts on the interior vertices: sparse  N x N matrix
+		* @param LaplacianBoundary  On return, scaled portion of the operator that acts on the boundary vertices: sparse  N x M matrix
+		* @param bClampAreas        Indicates if (A_ave / A_i) should be clamped to (0.5, 5) range.
+		*                           in practice this is desirable when creating the biharmonic operator, but not the mean curvature flow operator
+		*
+		*   LaplacianInterior * Vector_InteriorVerts + LaplacianBoundary * Vector_BoundaryVerts = Full Laplacian applied to interior vertices.
+		*/
+		template<typename RealType>
+		void ConstructIDTCotangentLaplacian(const FDynamicMesh3& DynamicMesh, const FVertexLinearization& VertexMap,
+			UE::Solvers::TSparseMatrixAssembler<RealType>& LaplacianInterior,
+			UE::Solvers::TSparseMatrixAssembler<RealType>& LaplacianBoundary,
+			const bool bClampWeights);
+		
+		
 		enum class ECotangentWeightMode
 		{
 			/** Standard cotangent weights */
@@ -196,6 +220,16 @@ namespace UE
 			UE::Solvers::TSparseMatrixAssembler<RealType>& LaplacianMatrix,
 			ECotangentWeightMode WeightMode = ECotangentWeightMode::ClampedMagnitude,
 			ECotangentAreaMode AreaMode = ECotangentAreaMode::NoArea );
+		/**
+		* Use intrinsic Delaunay mesh to construct sparse Cotangent Laplacian matrix.
+		* This variant combines the N interior and M boundary vertices into a single (N+M) matrix and does not do any special treatment of boundaries,
+		* they just get standard Cotan weights.  
+		*/
+		template<typename RealType>
+		void ConstructFullIDTCotangentLaplacian(const FDynamicMesh3& Mesh, const FVertexLinearization& VertexMap,
+			UE::Solvers::TSparseMatrixAssembler<RealType>& LaplacianMatrix,
+			ECotangentWeightMode WeightMode = ECotangentWeightMode::ClampedMagnitude,
+			ECotangentAreaMode AreaMode = ECotangentAreaMode::NoArea);
 
 	}
 }
@@ -659,7 +693,144 @@ void UE::MeshDeformation::ConstructCotangentLaplacian(const FDynamicMesh3& Dynam
 	}
 }
 
+template<typename RealType>
+void UE::MeshDeformation::ConstructIDTCotangentLaplacian(const FDynamicMesh3& DynamicMesh, const FVertexLinearization& VertexMap,
+	UE::Solvers::TSparseMatrixAssembler<RealType>& LaplacianInterior,
+	UE::Solvers::TSparseMatrixAssembler<RealType>& LaplacianBoundary,
+	const bool bClampWeights)
+{
+	const TArray<int32>& ToMeshV = VertexMap.ToId();
+	const TArray<int32>& ToIndex = VertexMap.ToIndex();
+	const int32 NumVerts = VertexMap.NumVerts();
+	const int32 NumBoundaryVerts = VertexMap.NumBoundaryVerts();
+	const int32 NumInteriorVerts = NumVerts - NumBoundaryVerts;
+
+	// create intrinsic mesh with delaunay triangulation
+	UE::Geometry::FIntrinsicEdgeFlipMesh IntrinsicMesh(DynamicMesh);
+
+	TSet<int32> Uncorrected; // some edges can't be flipped.
+	const int32 NumFlips = UE::Geometry::FlipToDelaunay(IntrinsicMesh, Uncorrected);
+
+	// pre-allocate space when possible
+	int32 NumMatrixEntries = ComputeNumMatrixElements(IntrinsicMesh, ToMeshV);
+	LaplacianInterior.ReserveEntriesFunc(NumMatrixEntries);
+
+	// cache the cotan weight for each edge (i.e. 1/2 ( cot(a_ij) + cot(b_ij)  for edge(i,j) ) 
+	TArray<double> EdgeCotanWeights;
+	{
+		const int32 MaxEID = IntrinsicMesh.MaxEdgeID();
+		EdgeCotanWeights.AddZeroed(MaxEID);
+		for (int32 EdgeId = 0; EdgeId < MaxEID; ++EdgeId)
+		{
+			if (!IntrinsicMesh.IsEdge(EdgeId))
+			{
+				continue;
+			}
+			EdgeCotanWeights[EdgeId] = IntrinsicMesh.EdgeCotanWeight(EdgeId);
+		}
+	}
+
+	// Construct Laplacian Matrix: loop over verts constructing the corresponding matrix row.
+	//                             skipping the boundary verts for later use.
+	for (int32 i = 0; i < NumInteriorVerts; ++i)
+	{
+		const int32 IVertId = ToMeshV[i]; // I - the row
 
+		// Compute the Voronoi area for this vertex.
+		// 1/8 Sum ( cot(a) + cot(b) )*length^2
+		// see for example http://www.geometry.caltech.edu/pubs/DMSB_III.pdf
+		double WeightArea = 0.0; 
+		bool bUseUniform = false;
+		{
+			for (int32 EdgeId : IntrinsicMesh.VtxEdgesItr(IVertId))
+			{
+				const double EdgeLength = IntrinsicMesh.GetEdgeLength(EdgeId);
+				const double CTWeight = EdgeCotanWeights[EdgeId];
+				const double CTWeightLSqr = CTWeight * EdgeLength * EdgeLength;
+				
+				if (Uncorrected.Contains(EdgeId) || CTWeightLSqr < 1.e-1)
+				{
+					bUseUniform = true; 
+				}
+
+				WeightArea += CTWeightLSqr;
+			}
+			WeightArea *= 0.25;
+			WeightArea = FMath::Max(WeightArea, 1.e-5);
+		}		
+
+		double WeightII = 0.; // accumulate to equal and opposite the sum of the neighbor weights
+
+		if (bUseUniform)
+		{ 
+			// use uniform stencil for this vertex
+			for (int32 EdgeId : IntrinsicMesh.VtxEdgesItr(IVertId))
+			{
+				// note: both sides of the edge exist since this isn't a boundary vert
+				const FIndex2i EdgeV = IntrinsicMesh.GetEdgeV(EdgeId);
+				// the other vert in the edge - identifies the matrix column
+				const int32 JVertId = (EdgeV[0] == IVertId) ? EdgeV[1] : EdgeV[0];  // J - the column
+				const int32 j = ToIndex[JVertId];
+
+				double  WeightIJ = 1.;
+				WeightII += WeightIJ;
+
+				if (j < NumInteriorVerts)
+				{
+					LaplacianInterior.AddEntryFunc(i, j, (RealType)(WeightIJ));
+				}
+				else
+				{
+					int32 jBoundary = j - NumInteriorVerts;
+					LaplacianBoundary.AddEntryFunc(i, jBoundary, (RealType)(WeightIJ));
+				}
+			}
+			LaplacianInterior.AddEntryFunc(i, i, (RealType)(-WeightII));
+		}
+		else
+		{ 
+
+			// for each connecting edge
+			for (int32 EdgeId : IntrinsicMesh.VtxEdgesItr(IVertId))
+			{
+				// note: both sides of the edge exist since this isn't a boundary vert
+				const FIndex2i EdgeV = IntrinsicMesh.GetEdgeV(EdgeId);
+
+				if (EdgeV.A == EdgeV.B)
+				{
+					continue; // the weights cancle. 
+				}
+
+				double WeightIJ = EdgeCotanWeights[EdgeId];
+				if (bClampWeights)
+				{
+					WeightIJ = FMathd::Clamp(WeightIJ, -1.e5 * WeightArea, 1.e5 * WeightArea);
+				}
+
+				double Weight = (WeightIJ / WeightArea);
+				WeightII += Weight;
+
+				// the other vert in the edge - identifies the matrix column
+				const int32 JVertId = (EdgeV[0] == IVertId) ? EdgeV[1] : EdgeV[0];  // J - the column
+				const int32 j = ToIndex[JVertId];
+				
+				if (j < NumInteriorVerts)
+				{
+					LaplacianInterior.AddEntryFunc(i, j, (RealType)(Weight));
+				}
+				else
+				{
+					int32 jBoundary = j - NumInteriorVerts;
+					LaplacianBoundary.AddEntryFunc(i, jBoundary, (RealType)(Weight));
+				}
+
+			}
+
+			LaplacianInterior.AddEntryFunc(i, i, (RealType)(-WeightII));
+		}
+	}
+
+}
 
 
 
@@ -741,3 +912,100 @@ void UE::MeshDeformation::ConstructFullCotangentLaplacian(const FDynamicMesh3& M
 		LaplacianMatrix.AddEntryFunc(i, i, (RealType)(-WeightII / WeightArea));
 	}
 }
+
+
+template<typename RealType>
+void UE::MeshDeformation::ConstructFullIDTCotangentLaplacian(const FDynamicMesh3& Mesh, const FVertexLinearization& VertexMap,
+	UE::Solvers::TSparseMatrixAssembler<RealType>& LaplacianMatrix,
+	ECotangentWeightMode WeightMode,
+	ECotangentAreaMode AreaMode)
+{
+	const TArray<int32>& ToMeshV = VertexMap.ToId();
+	const TArray<int32>& ToIndex = VertexMap.ToIndex();
+	const int32 NumVerts = VertexMap.NumVerts();
+
+	// create intrinsic mesh with delaunay triangulation
+	UE::Geometry::FIntrinsicTriangulation IntrinsicMesh(Mesh);
+	TSet<int32> Uncorrected;
+	const int32 NumFlips = UE::Geometry::FlipToDelaunay(IntrinsicMesh, Uncorrected);
+
+	// pre-allocate space when possible
+	int32 NumMatrixEntries = ComputeNumMatrixElements(IntrinsicMesh, ToMeshV);
+	LaplacianMatrix.ReserveEntriesFunc(NumMatrixEntries);
+
+	// cache the cotan weight for each edge (i.e. 1/2 ( cot(a_ij) + cot(b_ij)  for edge(i,j) ) 
+	TArray<double> EdgeCotanWeights;
+	{
+		const int32 MaxEID = IntrinsicMesh.MaxEdgeID();
+		EdgeCotanWeights.AddUninitialized(MaxEID);
+		
+		for (int32 e = 0; e < MaxEID; ++e)
+		{
+			if (!IntrinsicMesh.IsEdge(e))
+			{
+				continue;
+			}
+			EdgeCotanWeights[e] = IntrinsicMesh.EdgeCotanWeight(e);
+		}
+	}
+
+	// Construct Laplacian Matrix: loop over verts constructing the corresponding matrix row.
+	for (int32 i = 0; i < NumVerts; ++i)
+	{
+		const int32 IVertId = ToMeshV[i]; // I - the row
+		
+		bool bUseUniform = false;
+		double WeightArea = 1.0;
+		if (AreaMode == ECotangentAreaMode::VoronoiArea)
+		{
+			// 1/8 Sum ( cot(a) + cot(b) )*length^2
+			// see for example http://www.geometry.caltech.edu/pubs/DMSB_III.pdf
+			WeightArea = 0.0;
+			for (int32 EdgeId : IntrinsicMesh.VtxEdgesItr(IVertId))
+			{
+				const double EdgeLength = IntrinsicMesh.GetEdgeLength(EdgeId);
+				const double CTWeight = EdgeCotanWeights[EdgeId];
+				const double CTWeightLSqr = CTWeight * EdgeLength * EdgeLength;
+
+				if (Uncorrected.Contains(EdgeId) || CTWeightLSqr < 1.e-1)
+				{
+					bUseUniform = true;
+				}
+
+				WeightArea += CTWeightLSqr;
+			}
+			WeightArea *= 0.25;  
+		}
+		if (bUseUniform)
+		{
+			WeightArea = 1.0;
+		}
+
+		double WeightII = 0.; // accumulate to equal and opposite the sum of the neighbor weights
+
+		for (int32 EdgeId : IntrinsicMesh.VtxEdgesItr(IVertId))
+		{
+			const FIndex2i EdgeV = IntrinsicMesh.GetEdgeV(EdgeId);
+			if (EdgeV.A == EdgeV.B)
+			{
+				continue; // weights will cancle.
+			}
+			// the other vert in the edge - identifies the matrix column
+			const int32 JVertId = (EdgeV[0] == IVertId) ? EdgeV[1] : EdgeV[0];  // J - the column
+
+			double WeightIJ = (bUseUniform) ? 1. : EdgeCotanWeights[EdgeId];
+
+			if (WeightMode == ECotangentWeightMode::ClampedMagnitude)
+			{
+				WeightIJ = FMathd::Clamp(WeightIJ, -1.e5 * WeightArea, 1.e5 * WeightArea);
+			}
+
+			double Weight = WeightIJ / WeightArea;
+			WeightII += Weight;
+
+			const int32 j = ToIndex[JVertId];
+			LaplacianMatrix.AddEntryFunc(i, j, (RealType)(Weight));
+		}
+		LaplacianMatrix.AddEntryFunc(i, i, (RealType)(-WeightII));
+	}
+}
\ No newline at end of file
diff --git a/Engine/Plugins/Runtime/GeometryProcessing/Source/DynamicMesh/Public/Solvers/MeshLaplacian.h b/Engine/Plugins/Runtime/GeometryProcessing/Source/DynamicMesh/Public/Solvers/MeshLaplacian.h
--- a/Engine/Plugins/Runtime/GeometryProcessing/Source/DynamicMesh/Public/Solvers/MeshLaplacian.h
+++ b/Engine/Plugins/Runtime/GeometryProcessing/Source/DynamicMesh/Public/Solvers/MeshLaplacian.h
@@ -23,7 +23,9 @@ enum class ELaplacianWeightScheme
 	/**  */
 	Cotangent,
 	/**  */
-	ClampedCotangent
+	ClampedCotangent,
+	/** Cotan for Intrinsic Delaunay triangulation*/
+	IDTCotanget,
 };
 
 
diff --git a/Engine/Source/Runtime/GeometryCore/Private/Spatial/IntrinsicTriangulationMesh.cpp b/Engine/Source/Runtime/GeometryCore/Private/Spatial/IntrinsicTriangulationMesh.cpp
new file mode 100644
--- /dev/null
+++ b/Engine/Source/Runtime/GeometryCore/Private/Spatial/IntrinsicTriangulationMesh.cpp
@@ -0,0 +1,1627 @@
+// Copyright Epic Games, Inc. All Rights Reserved.
+
+#include "Spatial/IntrinsicTriangulationMesh.h"
+#include "Spatial/MeshGeodesicSurfaceTracer.h"
+#include "DynamicMesh/DynamicMeshAttributeSet.h"
+#include "Util/IndexUtil.h"
+#include "Algo/Reverse.h"
+#include "Containers/BitArray.h"
+#include "MathUtil.h"
+#include "Async/ParallelFor.h"
+
+using namespace UE::Geometry;
+
+
+/**
+ *  Helper functions that use triangle edge lengths instead of vertex positions to do some euclidean triangle-based calculations.
+ *  Many of these are based on the "law of cosines" 
+ */
+namespace
+{
+	/**
+	* Given a flat triangle with edges L0, L1, L2 in CCW order, compute the angle between
+	* the L0 and L2 edge in radians [0, Pi] range.
+	* 
+	* used by Sharp, Soliman and Crane [2019, ACM Transactions on Graphics] section 3.2
+	*
+	* but can be understood as trivial application of "the law of cosines"
+	*/
+	double InteriorAngle(const double L0, const double L1, const double L2)
+	{
+		// V1 = V2 - V0.   L1^2 = ||V1||^2 = ||V2||^2 + ||V0||^2 - 2 Dot(V2, V0)  =  L2^2 + L0^2 - 2 L2 L0 CosTheta
+		// CosTheta = (L0^2 + L1^2 - L2^2) / (2 L0 L1)
+		double CosTheta = (L0 * L0 + L2 * L2 - L1 * L1) / (2. * L0 * L2);
+
+		// account for roundoff, theoretically this would already be in the [-1,1] range
+		CosTheta = FMath::Clamp(CosTheta, -1., 1.);
+
+		return FMath::Acos(CosTheta);
+	}
+
+	/**
+	* For a triangle, use the law of cos to compute the third edge length
+	*  - Given interior angle in radians AngleR, and adjacent edge lengths L0 and L1, this computes the opposing edge length. 
+	*/
+	double LawOfCosLength(const double AngleR, const double L0, const double L1)
+	{
+		double L2sqr = L0 * L0 + L1 * L1 - 2.*L0*L1*FMath::Cos(AngleR);
+		// theoretically this is always positive, but protect round-off.   Alternately, could compute (L0-L1)^2 + 4 L0 L1 Sin^2 ( AngleR/2)  
+		L2sqr = FMath::Max(L2sqr, 0.);
+		return FMath::Sqrt(L2sqr);
+	}
+
+	/**
+	* Compute the unsigned distance between two points, described as barycentric coordinates
+	* relative to a triangle with side lengths L0, L1, L2 is CCW order.
+	* @params L0, L1, L2 - triangle side lengths in CCW order
+	* @param BCPoint0  - Barycentric coordinates for the first point
+	* @param BCPoint1  - Barycentric coordinates for the second point
+	*
+	* NB: it is assumed that the points are within the triangle, and the triangle is valid (i.e. lengths satisfy triangle inequality)
+	*/
+	double DistanceBetweenBarycentricPoints(const double L0, const double L1, const double L2, const FVector3d& BCPoint0, const FVector3d& BCPoint1)
+	{
+		// see Sharp, Soliman and Crane [2019, ACM Transactions on Graphics] or  Schindler and Chen [2012, Section 3.2]
+		using Vector3Type = FVector3d;
+		using ScalarType = double;
+		const ScalarType Zero(0);
+
+		const Vector3Type u = BCPoint1 - BCPoint0;
+		const ScalarType dsqr = -(L0 * L0 * u[0] * u[1] + L1 * L1 * u[1] * u[2] + L2 * L2 * u[2] * u[0]);
+		const ScalarType d = (dsqr < Zero) ? Zero : sqrt(dsqr);
+		return d;
+	}
+
+
+	/**
+	* Compute a vector (in polar form) from triangle corner to a point within the triangle.
+	* @params L0, L1, L2 - triangle side lengths in CCW order
+	* @param BarycetricPoint - interior point in triangle
+	*
+	* @return, a polar vector from the vertex shared by L0, L2 to the BarycentricPoint in the form
+	*          Vector[0] = radial distance
+	*          Vector[1] = angle measured from the L0 side of the triangle
+	*/
+	FVector2d VectorToPoint(const double L0, const double L1, const double L2, const FVector3d& BarycentricPoint)
+	{
+		FVector3d BCpi(1., 0., 0.);
+		FVector3d BCpj(0., 1., 0.);
+		double rpi = DistanceBetweenBarycentricPoints(L0, L1, L2, BCpi, BarycentricPoint);
+		double rpj = DistanceBetweenBarycentricPoints(L0, L1, L2, BCpj, BarycentricPoint);
+		double angle = InteriorAngle(L0, rpj, rpi);
+		return FVector2d(rpi, angle);
+	}
+
+	/**
+	* Heron's formula for computing the area of a triangle given the lengths of all three sides.  
+	*/
+	double TriangleArea(const double L0, const double L1, const double L2)
+	{
+		// note: since this is a triangle, the lengths should satisfy triangle inequality, meaning this should be positive
+		double AreaSqrTimes16 = FMath::Max(0., (L0 + L1 + L2) * (-L0 + L1 + L2 ) * (L0 -L1 + L2) * (L0 + L1 -L2));
+		return FMath::Sqrt(AreaSqrTimes16) * 0.25;
+	}
+	double TriangleArea(const FVector3d& Ls)
+	{
+		return TriangleArea(Ls[0], Ls[1], Ls[2]);
+	}
+	/**
+	* Compute the cotangent of the interior angle opposite the edge L0, give a triangle with edges {L0, L1, L2} in CCW order
+	*/
+	double ComputeCotangent(const double L0, const double L1, const double L2)
+	{
+		const double Area = TriangleArea(L0, L1, L2); 
+		const double CT = 0.25 * (-L0 * L0 + L1 * L1 + L2 * L2) / Area;
+		return CT;
+	}
+
+	/**
+	* Return the location of third vertex in a triangle in R2 with prescribed lengths
+	* such that the resulting triangle could be constructed as { (0,0),  (L1, 0), Result}.
+	* 
+	* NB: this assumes, but does not check, that the lengths satisfy the triangle inequality
+	*/
+	FVector2d ComputeOpposingVert2d(const double L0, const double L1, const double L2)
+	{
+		const double Area = TriangleArea(L0, L1, L2);
+		const double Y = 2. * Area / L0;
+		double X = FMath::Sqrt(FMath::Max(0., L2 * L2 - Y * Y));
+		FVector2d ResultPos(X, Y);
+		// angle between L0 and L2 edges can be computed form 2 * L0 * L2 * Cos(theta) = L0^2 + L2^2 - L1^2
+		if (L0 * L0 + L2 * L2 < L1 * L1)
+		{
+			ResultPos.X = -X; // angle > 90-degrees.
+		}
+		return ResultPos;
+	}
+
+	/**
+	* Given three side lengths that satisfy the triangle inequality,
+	* this updates the 2d positions to be the vertices of a triangle, 
+	* the edge lengths of which ( when in CCW order)  match the prescribed lengths.  
+	* with vertices { (0,0),  (L1, 0), (X, Y)} where Y > 0
+	*/
+	void TriangleFromLengths(const double L0, const double L1, const double L2, FVector2d& p0, FVector2d& p1, FVector2d& p3)
+	{
+		p0 = FVector2d(0., 0.);
+		p1 = FVector2d(L1, 0.);
+		p3 = ComputeOpposingVert2d(L0, L1, L2);
+	}
+
+	
+};
+
+
+/**
+ *  Utilities and generic implementation for flipping a mesh to Delaunay  
+ */
+namespace
+{
+	/**
+	* LIFO for unique DynamicMesh IDs, can be constructed from any FRefCountVector::IndexEnumerable
+	* and supports a filtering function on construction.  If the filtering function is absent, all
+	* unique elements will be added, otherwise only those selected by the filter.
+	*    
+	*	for example: 
+	* 
+	*	 auto AddToQueueFilter = [](int ID)->bool{...};
+	* 
+	*    FIndexLIFO EdgeQueueLIFO(Mesh.EdgeIndicesItr(), AddToQueueFilter);
+	*/
+	class FIndexLIFO
+	{
+	public:
+		// constructor that adds all unique IDs from IndexEnumerable
+		FIndexLIFO(FRefCountVector::IndexEnumerable IDEnumerable)
+			: IsEnqueued(false, (int32)IDEnumerable.Vector->GetMaxIndex())
+		{
+			const int32 NumIDs = (int32)IDEnumerable.Vector->GetCount();
+			IDs.Reserve(NumIDs);
+			for (int32 ID : IDEnumerable)
+			{
+				Enqueue(ID);
+			}
+		}
+		// filtering constructor that only adds IDs for which Filter(ID) is true
+		FIndexLIFO(FRefCountVector::IndexEnumerable IDEnumerable, const TFunctionRef<bool(int32)>& Filter)
+			: IsEnqueued(false, (int32)IDEnumerable.Vector->GetMaxIndex())
+		{
+			const int32 NumIDs = (int32)IDEnumerable.Vector->GetCount();
+			IDs.Reserve(NumIDs);
+			for (int32 ID : IDEnumerable)
+			{
+				if (Filter(ID))
+				{
+					Enqueue(ID);
+				}
+			}
+		}
+		bool Dequeue(int32& IDOut)
+		{
+			if (IDs.Num())
+			{
+				IDOut = IDs.Pop(false);
+				IsEnqueued[IDOut] = false;
+				return true;
+			}
+			IDOut = -1;
+			return false;
+		}
+		void Enqueue(int32 ID)
+		{
+			if (!IsEnqueued[ID])
+			{
+				IDs.Add(ID);
+				IsEnqueued[ID] = true;
+			}
+		}
+	private:
+		FIndexLIFO();
+		TBitArray<FDefaultBitArrayAllocator> IsEnqueued;
+		TArray<int32> IDs;
+	};
+
+	/**
+	* FIFO for unique Dynamic Mesh IDs - can be constructed from any FRefCountVector::IndexEnumerable
+	* and supports a filtering function on construction.  If the filtering function is absent, all
+	* unique elements will be added, otherwise only those selected by the filter.
+	*
+	*	for example:
+	*
+	*	 auto AddToQueueFilter = [](int ID)->bool{...};
+	*
+	*    FIndexFIFO EdgeQueueFIFO(Mesh.EdgeIndicesItr(), AddToQueueFilter);
+	*/
+	class FIndexFIFO
+	{
+	public:
+		// constructor that adds all unique IDs from IndexEnumerable
+		FIndexFIFO(FRefCountVector::IndexEnumerable IDEnumerable)
+			: IsEnqueued(false, (int32)IDEnumerable.Vector->GetMaxIndex())
+		{
+			const int32 NumEdges = (int32)IDEnumerable.Vector->GetMaxIndex();
+			for (int32 ID : IDEnumerable)
+			{
+				Enqueue(ID);
+			}
+		}
+
+		// filtering constructor that only adds IDs for which Filter(ID) is true  Note: Filter is called in parallel.
+		FIndexFIFO(FRefCountVector::IndexEnumerable IDEnumerable, const TFunctionRef<bool(int32)>& Filter)
+			: IsEnqueued(false, (int32)IDEnumerable.Vector->GetMaxIndex())
+		{
+			// parallel evaluation that assumes Filter is expensive
+			
+			const int32 MaxID = (int32)IDEnumerable.Vector->GetMaxIndex();
+			TArray<int32> ToInclude;
+			ToInclude.AddZeroed(MaxID);
+			ParallelFor(MaxID, [&ToInclude, &Filter](int32 ID)
+			{
+				if (Filter(ID))
+				{
+					ToInclude[ID] = 1;
+				}
+			});
+			for (int32 ID = 0; ID < MaxID; ++ID)
+			{
+				if (ToInclude[ID] == 1)
+				{
+					Enqueue(ID);
+				}
+			}
+		}
+
+		bool Dequeue(int32& IDOut)
+		{
+			if (IDs.Dequeue(IDOut))
+			{
+				IsEnqueued[IDOut] = false;
+				return true;
+			}
+			IDOut = -1;
+			return false;
+		}
+		void Enqueue(int32 ID)
+		{
+			if (!IsEnqueued[ID])
+			{
+				IDs.Enqueue(ID);
+				IsEnqueued[ID] = true;
+			}
+		}
+	private:
+		FIndexFIFO();
+		TBitArray<FDefaultBitArrayAllocator> IsEnqueued;
+		TQueue<int32> IDs;
+	};
+
+
+	/**
+	* Carry out edge flips on intrinsic mesh until either the MaxFlipCount is reached or the mesh is fully Delaunay and returns
+	* the number of flips.
+	* Note: There can be cases where an edge can not flip (e.g. if the resulting edge already exists)
+	*       - for this reason, there may be some "uncorrected" edges in the resulting mesh.
+	* 
+	* @param IntrinsicMesh - The mesh to be operated on.
+	* @param Uncorrected   - contains on return, all the edges that could not be flipped.
+	* @param Threshold     - edges with cotan weight less than threshold should be flipped.
+	*/
+	template <typename MeshType>
+	int32 FlipToDelaunayImpl(MeshType& Mesh, TSet<int32>& Uncorrected, const int32 MaxFlipCount, double Threshold = -TMathUtilConstants<double>::ZeroTolerance)
+	{
+	Uncorrected.Empty();
+
+	// returns true if an edge should be flipped.
+	auto EdgeShouldFlipFilter = [&Mesh, Threshold](int32 EID)->bool
+	{
+		if (!Mesh.IsEdge(EID) || Mesh.IsBoundaryEdge(EID))
+		{
+			return false;
+		}
+	
+		const double CotanWeightValue = Mesh.EdgeCotanWeight(EID);
+		
+		// Delaunay test
+		return CotanWeightValue < Threshold;
+	};
+
+	// Enqueue the "bad" edges.
+	FIndexFIFO EdgeQueue(Mesh.EdgeIndicesItr(), EdgeShouldFlipFilter);
+
+	// flip away the bad edges
+	int32 FlipCount = 0;
+	int32 EID;
+	while (EdgeQueue.Dequeue(EID) && FlipCount < MaxFlipCount)
+	{
+		if (EdgeShouldFlipFilter(EID))
+		{
+			FIntrinsicTriangulation::FEdgeFlipInfo EdgeFlipInfo;
+			const EMeshResult Result = Mesh.FlipEdge(EID, EdgeFlipInfo);
+
+			if (Result == EMeshResult::Ok)
+			{
+				FlipCount++;
+				for (int32 t = 0; t < 2; ++t)
+				{
+					const FIndex3i TriEIDs = Mesh.GetTriEdges(EdgeFlipInfo.Triangles[t]);
+					int32 IndexOf = TriEIDs.IndexOf(EID);
+					EdgeQueue.Enqueue(TriEIDs[(IndexOf + 1) % 3]);
+					EdgeQueue.Enqueue(TriEIDs[(IndexOf + 2) % 3]);
+				}
+			}
+			else
+			{
+				Uncorrected.Add(EID);
+			}
+		}
+	}
+
+	return FlipCount;
+	}
+
+};
+
+
+int32 UE::Geometry::FlipToDelaunay(FIntrinsicTriangulation& IntrinsicMesh, TSet<int32>& Uncorrected, const int32 MaxFlipCount)
+{
+	return FlipToDelaunayImpl(IntrinsicMesh, Uncorrected, MaxFlipCount);
+}
+
+int32 UE::Geometry::FlipToDelaunay(FIntrinsicEdgeFlipMesh& IntrinsicMesh, TSet<int32>& Uncorrected, const int32 MaxFlipCount)
+{
+	return FlipToDelaunayImpl(IntrinsicMesh, Uncorrected, MaxFlipCount);
+}
+
+
+
+
+/**------------------------------------------------------------------------------
+* FIntrinsicEdgeFlipMesh Methods
+*------------------------------------------------------------------------------ */
+
+FIntrinsicEdgeFlipMesh::FIntrinsicEdgeFlipMesh(const FDynamicMesh3& SrcMesh) :
+	Vertices(SrcMesh.GetVerticesBuffer()),
+	VertexRefCounts{ SrcMesh.GetVerticesRefCounts() },
+	VertexEdgeLists{ SrcMesh.GetVertexEdges() },
+
+	Triangles{ SrcMesh.GetTrianglesBuffer() },
+	TriangleRefCounts{ SrcMesh.GetTrianglesRefCounts() },
+	TriangleEdges{ SrcMesh.GetTriangleEdges() },
+
+	Edges{ SrcMesh.GetEdgesBuffer() },
+	EdgeRefCounts{ SrcMesh.GetEdgesRefCounts() }
+{
+	
+	const int32 MaxEID = MaxEdgeID();
+	EdgeLengths.SetNum(MaxEID);
+	for (int32 EID = 0; EID < MaxEID; ++EID)
+	{
+		if (!IsEdge(EID))
+		{
+			continue;
+		}
+		const FIndex2i EdgeV = GetEdgeV(EID);
+		const FVector3d Pos[2] = { GetVertex(EdgeV.A), GetVertex(EdgeV.B) };
+		EdgeLengths[EID] = (Pos[1] - Pos[0]).Length();
+	}
+
+	const int32 MaxTriID = MaxTriangleID();
+	InternalAngles.SetNum(MaxTriID);
+	for (int32 TID = 0; TID < MaxTriID; ++TID)
+	{
+		if (!IsTriangle(TID))
+		{
+			continue;
+		}
+		// angles at v0, v1, v2, in that order
+		InternalAngles[TID] = ComputeTriInternalAnglesR(TID);
+	}
+}
+
+
+FIndex2i FIntrinsicEdgeFlipMesh::GetEdgeOpposingV(int32 EID) const
+{
+	const FEdge& Edge = Edges[EID];
+	FIndex2i Result(InvalidID, InvalidID);
+
+	for (int32 i = 0; i < 2; ++i)
+	{
+		int32 TriID = Edge.Tri[i];
+		if (TriID == InvalidID) continue;
+
+		const FIndex3i TriEIDs = GetTriEdges(TriID);
+		const int32 IndexOf = TriEIDs.IndexOf(EID);
+		const FIndex3i TriVIDs = GetTriangle(TriID);
+		Result[i] = TriVIDs[AddTwoModThree[IndexOf]];
+	}
+
+	return Result;
+}
+
+
+FIndex2i FIntrinsicEdgeFlipMesh::GetOrientedEdgeV(int32 EID, int32 TID) const
+{
+	int32 IndexOf = GetTriEdges(TID).IndexOf(EID);
+	checkSlow(IndexOf != InvalidID);
+	FIndex3i TriVIDs = GetTriangle(TID);
+	return FIndex2i(TriVIDs[IndexOf], TriVIDs[AddOneModThree[IndexOf]]);
+}
+
+
+int32 FIntrinsicEdgeFlipMesh::ReplaceEdgeTriangle(int32 eID, int32 tOld, int32 tNew)
+{
+	FIndex2i& Tris = Edges[eID].Tri;
+	int32 a = Tris[0], b = Tris[1];
+	if (a == tOld) {
+		if (tNew == InvalidID)
+		{
+			Tris[0] = b;
+			Tris[1] = InvalidID;
+		}
+		else
+		{
+			Tris[0] = tNew;
+		}
+		return 0;
+	}
+	else if (b == tOld)
+	{
+		Tris[1] = tNew;
+		return 1;
+	}
+	else
+	{
+		return -1;
+	}
+}
+
+
+EMeshResult FIntrinsicEdgeFlipMesh::FlipEdgeTopology(int32 eab, FEdgeFlipInfo& FlipInfo)
+{
+	if (!IsEdge(eab))
+	{
+		return EMeshResult::Failed_NotAnEdge;
+	}
+	if (IsBoundaryEdge(eab))
+	{
+		return EMeshResult::Failed_IsBoundaryEdge;
+	}
+
+	// find oriented edge [a,b], tris t0,t1, and other verts c in t0, d in t1
+	const FEdge Edge = Edges[eab];
+	int32 t0 = Edge.Tri[0], t1 = Edge.Tri[1];
+
+	FIndex2i oppV = GetEdgeOpposingV(eab);
+	FIndex2i orientedV = GetOrientedEdgeV(eab, t0);
+	int32 a = orientedV.A, b = orientedV.B;
+
+	if (oppV[0] == InvalidID || oppV[1] == InvalidID)
+	{
+		return EMeshResult::Failed_BrokenTopology;
+	}
+	int32 c = oppV[0];
+	int32 d = oppV[1];
+
+	const FIndex3i T0te = GetTriEdges(t0);
+	const FIndex3i T1te = GetTriEdges(t1);
+
+	const int32 T0IndexOf = T0te.IndexOf(eab);
+	const int32 T1IndexOf = T1te.IndexOf(eab);
+	// find edges bc, ca, ad, db
+	const int32 ebc = T0te[AddOneModThree[T0IndexOf]];
+	const int32 eca = T0te[AddTwoModThree[T0IndexOf]];
+
+	const int32 ead = T1te[AddOneModThree[T1IndexOf]];
+	const int32 edb = T1te[AddTwoModThree[T1IndexOf]];
+
+	// update triangles
+	Triangles[t0] = FIndex3i(c, d, b);
+	Triangles[t1] = FIndex3i(d, c, a);
+
+	// update edge AB, which becomes flipped edge CD
+	SetEdgeVerticesInternal(eab, c, d);
+	SetEdgeTrianglesInternal(eab, t0, t1);
+	const int32 ecd = eab;
+
+	// update the two other edges whose triangle nbrs have changed
+	if (ReplaceEdgeTriangle(eca, t0, t1) == -1)
+	{
+		checkfSlow(false, TEXT("FIntrinsicEdgeFlipMesh.FlipEdge: first ReplaceEdgeTriangle failed"));
+		return EMeshResult::Failed_UnrecoverableError;
+	}
+	if (ReplaceEdgeTriangle(edb, t1, t0) == -1)
+	{
+		checkfSlow(false, TEXT("FIntrinsicEdgeFlipMesh.FlipEdge: second ReplaceEdgeTriangle failed"));
+		return EMeshResult::Failed_UnrecoverableError;
+	}
+
+	// update triangle nbr lists (these are edges)
+	TriangleEdges[t0] = FIndex3i(ecd, edb, ebc);
+	TriangleEdges[t1] = FIndex3i(ecd, eca, ead);
+
+	// remove old eab from verts a and b, and Decrement ref counts
+	if (VertexEdgeLists.Remove(a, eab) == false)
+	{
+		checkfSlow(false, TEXT("FIntrinsicEdgeFlipMesh.FlipEdge: first edge list remove failed"));
+		return EMeshResult::Failed_UnrecoverableError;
+	}
+	VertexRefCounts.Decrement(a);
+	if (a != b)
+	{
+		if (VertexEdgeLists.Remove(b, eab) == false)
+		{
+			checkfSlow(false, TEXT("FIntrinsicEdgeFlipMesh.FlipEdge: second edge list remove failed"));
+			return EMeshResult::Failed_UnrecoverableError;
+		}
+		VertexRefCounts.Decrement(b);
+	}
+	if (IsVertex(a) == false || IsVertex(b) == false)
+	{
+		checkfSlow(false, TEXT("FIntrinsicEdgeFlipMesh.FlipEdge: either a or b is not a vertex?"));
+		return EMeshResult::Failed_UnrecoverableError;
+	}
+
+	// add edge ecd to verts c and d, and increment ref counts
+	VertexEdgeLists.Insert(c, ecd);
+	VertexRefCounts.Increment(c);
+	if (c != d)
+	{
+		VertexEdgeLists.Insert(d, ecd);
+		VertexRefCounts.Increment(d);
+	}
+
+	// success! collect up results
+	FlipInfo.EdgeID = eab;
+	FlipInfo.OriginalVerts = FIndex2i(a, b);
+	FlipInfo.OpposingVerts = FIndex2i(c, d);
+	FlipInfo.Triangles = FIndex2i(t0, t1);
+	return EMeshResult::Ok;
+}
+
+
+FVector3d FIntrinsicEdgeFlipMesh::ComputeTriInternalAnglesR(const int32 TID) const
+{
+	const FVector3d Lengths = GetTriEdgeLengths(TID);
+	FVector3d Angles;
+	for (int32 v = 0; v < 3; ++v)
+	{
+		Angles[v] = InteriorAngle(Lengths[v], Lengths[AddOneModThree[v]], Lengths[AddTwoModThree[v]]);
+	}
+
+	return Angles;
+}
+
+
+double FIntrinsicEdgeFlipMesh::GetOpposingVerticesDistance(int32 EID) const
+{
+	const double OrgLength = GetEdgeLength(EID);
+	const FIndex2i EdgeTris = GetEdgeT(EID);
+	checkSlow(EdgeTris[1] != FDynamicMesh3::InvalidID);
+
+	// compute 2d locations of the verts opposite the EID edge.
+	FVector2d Opp2dVerts[2];
+	for (int32 i = 0; i < 2; ++i)
+	{
+		const FIndex3i TriEIDs = GetTriEdges(EdgeTris[i]);
+		const FVector3d Ls = GetEdgeLengthTriple(TriEIDs);
+		const int32 IndexOf = TriEIDs.IndexOf(EID);
+		const FVector3d PermutedLs = Permute(IndexOf, Ls);
+		Opp2dVerts[i] = ComputeOpposingVert2d(PermutedLs[0], PermutedLs[1], PermutedLs[2]);
+	}
+	// rotate second tri so the shared edge aligns
+	Opp2dVerts[1].Y = -Opp2dVerts[1].Y;
+	Opp2dVerts[1].X = OrgLength - Opp2dVerts[1].X;
+
+	return (Opp2dVerts[0] - Opp2dVerts[1]).Length();
+}
+
+
+EMeshResult FIntrinsicEdgeFlipMesh::FlipEdge(int32 EID, FEdgeFlipInfo& EdgeFlipInfo)
+{
+	if (!IsEdge(EID))
+	{
+		return EMeshResult::Failed_NotAnEdge;
+	}
+
+	if (IsBoundaryEdge(EID))
+	{
+		return EMeshResult::Failed_IsBoundaryEdge;
+	}
+
+	// prohibit case where the edge is shared by a non-convex pair of triangles.
+	{
+		// original triangles
+		FDynamicMesh3::FEdge OrgEdge = GetEdge(EID);
+
+		// Assumes both triangles have same orientation
+		double TotalAngleAtOrgVert[2] = { 0., 0. };
+		for (int32 i = 0; i < 2; ++i)
+		{	
+			int32 TriID = OrgEdge.Tri[i];
+			int32 IndexOf = GetTriEdges(TriID).IndexOf(EID);
+			TotalAngleAtOrgVert[i] += InternalAngles[TriID][IndexOf];
+			TotalAngleAtOrgVert[(i+1) % 2] += InternalAngles[TriID][AddOneModThree[IndexOf]];
+		}
+
+		if (TotalAngleAtOrgVert[0] > TMathUtilConstants<double>::Pi - TMathUtilConstants<double>::ZeroTolerance || TotalAngleAtOrgVert[1] > TMathUtilConstants<double>::Pi - TMathUtilConstants<double>::ZeroTolerance)
+		{
+			return EMeshResult::Failed_Unsupported;
+		}
+	}
+
+	// compute the length of edge after flip
+	const double PostFlipLength = GetOpposingVerticesDistance(EID);
+
+	// flip edge in the underlying mesh
+	const EMeshResult MeshFlipResult = FlipEdgeTopology(EID, EdgeFlipInfo);
+
+	if (MeshFlipResult != EMeshResult::Ok)
+	{
+		// could fail for many reasons, e.g. a boundary edge or the new edge already exists.
+		return MeshFlipResult;
+	}
+
+	// update the intrinsic edge length
+	EdgeLengths[EID] = PostFlipLength;
+
+	// update interior angles for the tris
+	for (int32 t = 0; t < 2; ++t)
+	{
+		const int32 TID = EdgeFlipInfo.Triangles[t];
+		InternalAngles[TID] = ComputeTriInternalAnglesR(TID);	
+	}
+
+	return MeshFlipResult;
+}
+
+
+FVector2d FIntrinsicEdgeFlipMesh::EdgeOpposingAngles(int32 EID) const
+{
+	FVector2d Result;
+	FDynamicMesh3::FEdge Edge = GetEdge(EID);
+	{
+		const int32 IndexOf = GetTriEdges(Edge.Tri.A).IndexOf(EID);
+		const int32 IndexOfOpp = AddTwoModThree[IndexOf];
+		const double Angle = GetTriInternalAngleR(Edge.Tri.A, IndexOfOpp);
+		Result[0] = Angle;
+	}
+	if (Edge.Tri.B != FDynamicMesh3::InvalidID)
+	{
+		const int32 IndexOf = GetTriEdges(Edge.Tri.B).IndexOf(EID);
+		const int32 IndexOfOpp = AddTwoModThree[IndexOf];
+		const double Angle = GetTriInternalAngleR(Edge.Tri.B, IndexOfOpp);
+		Result[1] = Angle;
+	}
+	else
+	{
+		Result[1] = -TMathUtilConstants<double>::MaxReal;
+	}
+
+	return Result;
+}
+
+
+double FIntrinsicEdgeFlipMesh::EdgeCotanWeight(int32 EID) const
+{
+	auto ComputeCotanOppAngle = [this](int32 EID, int32 TID)->double
+	{
+		const FIndex3i TriEIDs = GetTriEdges(TID);
+		const int32 IOf = TriEIDs.IndexOf(EID);
+		const FVector3d TriLs = GetEdgeLengthTriple(TriEIDs);
+
+		const FVector3d Ls(TriLs[IOf], TriLs[AddOneModThree[IOf]], TriLs[AddTwoModThree[IOf]]);
+		return ComputeCotangent(Ls[0], Ls[1], Ls[2]);
+	};
+
+	FDynamicMesh3::FEdge Edge = GetEdge(EID);
+	double Result = ComputeCotanOppAngle(EID, Edge.Tri.A);
+	if (Edge.Tri.B != FDynamicMesh3::InvalidID)
+	{
+		Result += ComputeCotanOppAngle(EID, Edge.Tri.B);
+		Result /= 2;
+	}
+
+	return Result;
+}
+
+
+
+/**------------------------------------------------------------------------------
+*  FIntrinsicTriangulation Methods
+*-------------------------------------------------------------------------------*/
+
+FIntrinsicTriangulation::FSurfacePoint::FSurfacePoint()
+{
+	PositionType = EPositionType::Vertex;
+	Position.VertexPosition.VID = FDynamicMesh3::InvalidID;
+}
+
+FIntrinsicTriangulation::FSurfacePoint::FSurfacePoint(int32 VID)
+{
+	PositionType = EPositionType::Vertex;
+	Position.VertexPosition.VID = VID;
+}
+
+FIntrinsicTriangulation::FSurfacePoint::FSurfacePoint(int32 EdgeID, double Alpha)
+{
+	PositionType = EPositionType::Edge;
+	Position.EdgePosition.EdgeID = EdgeID;
+	Position.EdgePosition.Alpha = Alpha;
+}
+
+FIntrinsicTriangulation::FSurfacePoint::FSurfacePoint(int32 TriID, const FVector3d& BaryCentrics)
+{
+	PositionType = EPositionType::Triangle;
+	Position.TriPosition.TriID = TriID;
+	Position.TriPosition.BarycentricCoords[0] = BaryCentrics[0];
+	Position.TriPosition.BarycentricCoords[1] = BaryCentrics[1];
+	Position.TriPosition.BarycentricCoords[2] = BaryCentrics[2];
+}
+
+
+UE::Geometry::FIntrinsicTriangulation::FIntrinsicTriangulation(const FDynamicMesh3& SrcMesh)
+	: MyBase(SrcMesh)
+	, ExtrinsicMesh(&SrcMesh)
+{
+	// During construction the mesh is a deep copy of the SrcMesh ( i.e. has the same IDs)
+
+	// pick local reference directions for each vertex and triangle.
+	const int32 MaxVertexID   = ExtrinsicMesh->MaxVertexID();
+	const int32 MaxTriangleID = ExtrinsicMesh->MaxTriangleID();
+
+	{
+		VIDToReferenceEID.AddUninitialized(MaxVertexID);
+		TIDToReferenceEID.AddUninitialized(MaxTriangleID);
+
+		for (int32 VID = 0; VID < MaxVertexID; ++VID)
+		{
+			if (ExtrinsicMesh->IsVertex(VID))
+			{
+				int32 FirstEID = -1;
+				for (int32 EID : ExtrinsicMesh->VtxEdgesItr(VID))
+				{
+					FirstEID = EID;
+					break;
+				}
+				VIDToReferenceEID[VID] = FirstEID;
+			}
+		}
+		for (int32 TID = 0; TID < MaxTriangleID; ++TID)
+		{
+			if (ExtrinsicMesh->IsTriangle(TID))
+			{
+				TIDToReferenceEID[TID] = ExtrinsicMesh->GetTriEdge(TID, 0);
+			}
+		}
+	}
+
+	// add surface position.
+	{
+		IntrinsicVertexPositions.SetNum(MaxVertexID);
+
+		// initialize: identify with vertex in Extrinsic Mesh
+		for (int32 VID = 0; VID < MaxVertexID; ++VID)
+		{
+			if (IsVertex(VID))
+			{
+				FSurfacePoint VertexSurfacePoint(VID);
+				IntrinsicVertexPositions[VID] = VertexSurfacePoint;
+			}
+		}
+	}
+
+	// add edge directions
+	{
+		const int32 MaxExtEdgeID = ExtrinsicMesh->MaxEdgeID();
+		IntrinsicEdgeAngles.SetNum(MaxTriangleID);
+
+		GeometricVertexInfo.SetNum(MaxVertexID);
+		IntrinsicVertexPositions.SetNum(MaxVertexID);
+
+		ParallelFor(MaxVertexID, [this, &SrcMesh](int32 VID)
+		{
+			if (!ExtrinsicMesh->IsVertex(VID))
+			{
+				return;
+			}
+
+			const int32 PlusTwoModThree[3] = { 2, 0, 1 }; // PlusTwoModThree[i] = (i + 2)%3
+
+			// initialize by computing the angles
+			TArray<int32> Triangles;
+			TArray<int32> ContiguousGroupLengths;
+			TArray<bool> GroupIsLoop;
+
+			{
+
+				FGeometricInfo GeometricInfo;
+				// note: since the structure and IDs are a deep copy of the source mesh, we can use it to find the contiguous triangles
+				const EMeshResult Result = SrcMesh.GetVtxContiguousTriangles(VID, Triangles, ContiguousGroupLengths, GroupIsLoop);
+				if (Result == EMeshResult::Ok && GroupIsLoop.Num() == 1)
+				{
+					GeometricInfo.bIsInterior = (GroupIsLoop[0] == true);
+					// Contiguous tris could be cw or ccw?  Need them to  be ccw
+					if (Triangles.Num() > 1)
+					{
+						const int32 TID0 = Triangles[0];
+						FIndex3i TriVIDs = GetTriangle(TID0);
+						FIndex3i TriEIDs = GetTriEdges(TID0);
+						int32 VertSubIdx = TriVIDs.IndexOf(VID);
+						const int32 EID = TriEIDs[VertSubIdx];
+						const int32 NextEID = TriEIDs[PlusTwoModThree[VertSubIdx]];
+						const int32 TID1 = Triangles[1];
+						TriVIDs = GetTriangle(TID1);
+						TriEIDs = GetTriEdges(TID1);
+						VertSubIdx = TriVIDs.IndexOf(VID);
+
+						if (TriEIDs[VertSubIdx] != NextEID)
+						{
+							// need to reverse order to have CCW
+							Algo::Reverse(Triangles);
+						}
+					}
+
+					const int32 NumEdges = ContiguousGroupLengths[0];
+
+					double AngleOffset = 0;
+					double WalkedAngle = 0;
+					const int32 ReferenceEID = VIDToReferenceEID[VID];
+					checkSlow(ReferenceEID != FDynamicMesh3::InvalidID)
+
+					TMap<int32, int32> VisitedTriSubID;
+					for (int32 TID : Triangles)
+					{
+						const FIndex3i TriVIDs = GetTriangle(TID);
+						const FIndex3i TriEIDs = GetTriEdges(TID);
+						const int32 VertSubIdx = TriVIDs.IndexOf(VID);
+						IntrinsicEdgeAngles[TID][VertSubIdx] = WalkedAngle;
+
+						VisitedTriSubID.Add(TID, VertSubIdx);
+						const int32 EID = TriEIDs[VertSubIdx];
+						if (EID == ReferenceEID)
+						{
+							AngleOffset = WalkedAngle;
+						}
+						
+						WalkedAngle += GetTriInternalAngleR(TID, VertSubIdx);
+					}
+					// need to include last edge if the one ring of triangles isn't a closed loop, it could be the reference edge.
+					if (GroupIsLoop[0] == false)
+					{
+						const int32 TID = Triangles.Last();
+						const FIndex3i TriVIDs = GetTriangle(TID);
+						const FIndex3i TriEIDs = GetTriEdges(TID);
+						const int32 VertSubIdx = TriVIDs.IndexOf(VID);
+						const int32 EID = TriEIDs[PlusTwoModThree[VertSubIdx]];
+						if (EID == ReferenceEID)
+						{
+							AngleOffset = WalkedAngle;
+						}
+					}
+
+					// orient with the reference edge at zero (or pi) angle  and rescale to radians and use 0,2pi to define polar angle.
+					const double ToRadians = TMathUtilConstants<double>::TwoPi / WalkedAngle;
+					GeometricInfo.ToRadians = ToRadians;
+					double RefAngle = (ExtrinsicMesh->GetEdgeV(ReferenceEID).B == VID) ? TMathUtilConstants<double>::Pi : 0.;
+					for (TPair<int32, int32>& TriSubIDPair : VisitedTriSubID)
+					{
+						int32 TID = TriSubIDPair.Key;
+						int32 SubID = TriSubIDPair.Value;
+						double& Angle = IntrinsicEdgeAngles[TID][SubID];
+						Angle = ToRadians * (Angle - AngleOffset) + RefAngle;
+						Angle = AsZeroToTwoPi(Angle);
+					}
+				}
+				else // bow-tie or boundary
+				{
+					GeometricInfo.bIsInterior = false;
+				}
+
+				GeometricVertexInfo[VID] = GeometricInfo;
+			}
+		}, false /* force single thread*/);
+	}
+
+}
+
+
+EMeshResult UE::Geometry::FIntrinsicTriangulation::FlipEdge(int32 EID, FEdgeFlipInfo& EdgeFlipInfo)
+{
+
+	if (IsBoundaryEdge(EID))
+	{
+		return EMeshResult::Failed_IsBoundaryEdge;
+	}
+	// capture the state of the triangles before the flip.
+	
+	FIndex2i Tris = GetEdgeT(EID);
+	const int32 PreFlipIndexOf[2] = {GetTriEdges(Tris[0]).IndexOf(EID), GetTriEdges(Tris[1]).IndexOf(EID)};
+
+	// edge polar directions before the flip
+	const FVector3d PreFlipDirections[2] = { IntrinsicEdgeAngles[Tris[0]], IntrinsicEdgeAngles[Tris[1]] };
+	const FVector3d PreFlipInteriorAngles[2] = { InternalAngles[Tris[0]], InternalAngles[Tris[1]] };
+
+	// polar directions for the four edges that don't change.
+	const double bcDir = PreFlipDirections[0][AddOneModThree[PreFlipIndexOf[0]]];
+	const double caDir = PreFlipDirections[0][AddTwoModThree[PreFlipIndexOf[0]]];
+
+	const double adDir = PreFlipDirections[1][AddOneModThree[PreFlipIndexOf[1]]];
+	const double dbDir = PreFlipDirections[1][AddTwoModThree[PreFlipIndexOf[1]]];
+
+	// flip edge in the underlying mesh
+	// this updates the edge lengths and the interior angles.
+	const EMeshResult MeshFlipResult = MyBase::FlipEdge(EID, EdgeFlipInfo);
+
+	if (MeshFlipResult != EMeshResult::Ok)
+	{
+		// could fail for many reasons, e.g. a boundary edge 
+		return MeshFlipResult;
+	}
+
+	// compute the polar edge directions for the new edge (ie. dcDir and cdDir)
+	// dcDir 
+	const double newAngleAtD = InternalAngles[Tris[0]][1];
+	const double ToPolarAtD = GeometricVertexInfo[EdgeFlipInfo.OpposingVerts[1]].ToRadians;
+	const double dcDir = AsZeroToTwoPi(dbDir + ToPolarAtD * newAngleAtD);
+
+	// cdDir 
+	const double newAngleAtC = InternalAngles[Tris[1]][1];
+	const double ToPolarAtC =  GeometricVertexInfo[EdgeFlipInfo.OpposingVerts[0]].ToRadians;
+	const double cdDir = AsZeroToTwoPi(caDir + ToPolarAtC * newAngleAtC);
+
+	IntrinsicEdgeAngles[Tris[0]] = FVector3d(cdDir, dbDir, bcDir); // edges leaving c, d, b
+	IntrinsicEdgeAngles[Tris[1]] = FVector3d(dcDir, caDir, adDir); // edges leaving d, c, a
+	
+	return EMeshResult::Ok;
+}
+
+
+double UE::Geometry::FIntrinsicTriangulation::UpdateVertexByEdgeTrace(const int32 NewVID, const int32 TraceStartVID, const double TracePolarAngle, const double TraceDist)
+{
+	const FSurfacePoint& StartSurfacePoint = IntrinsicVertexPositions[TraceStartVID];
+
+	double ActualDist = 0;
+	FMeshGeodesicSurfaceTracer SurfaceTracer(*ExtrinsicMesh);
+	switch (StartSurfacePoint.PositionType)
+	{
+	case FSurfacePoint::EPositionType::Vertex:
+	{
+		// convert vertex location
+		const int32 ExtrinsicVID = StartSurfacePoint.Position.VertexPosition.VID;
+		const int32 ExtrinsicRefEID = VIDToReferenceEID[ExtrinsicVID];
+
+		const FMeshGeodesicSurfaceTracer::FMeshTangentDirection TangentDirection = { ExtrinsicVID, ExtrinsicRefEID, TracePolarAngle };
+		ActualDist = SurfaceTracer.TraceMeshFromVertex(TangentDirection, TraceDist);
+
+	}
+	break;
+	case FSurfacePoint::EPositionType::Edge:
+	{
+		// convert edge location -  Not used, and might have bugs
+		check(0);
+
+		const int32 ExtrinsicRefEID = StartSurfacePoint.Position.EdgePosition.EdgeID;
+		const double Alpha = StartSurfacePoint.Position.EdgePosition.Alpha;
+		const FDynamicMesh3::FEdge ExtEdge = ExtrinsicMesh->GetEdge(ExtrinsicRefEID);
+		const int32 StartExtrinsicTID = ExtEdge.Tri.A; // Q? should I inspect the angle to determine which triangle we are entering?
+		FIndex3i ExtrinsicTriVIDs = ExtrinsicMesh->GetTriangle(StartExtrinsicTID);
+		FVector3d StartBarycentricCoords(0.);
+		StartBarycentricCoords[ExtrinsicTriVIDs.IndexOf(ExtEdge.Vert.A)] = Alpha;
+		StartBarycentricCoords[ExtrinsicTriVIDs.IndexOf(ExtEdge.Vert.B)] = 1. - Alpha;
+
+		FMeshGeodesicSurfaceTracer::FMeshSurfaceDirection SurfaceDirection(ExtrinsicRefEID, TracePolarAngle);
+	}
+	break;
+	case FSurfacePoint::EPositionType::Triangle:
+	{
+		// trace from BC point in triangle
+		const int32 StartExtrinsicTID = StartSurfacePoint.Position.TriPosition.TriID;
+		const int32 ExtrinsicRefEID = TIDToReferenceEID[StartExtrinsicTID];
+
+		FVector3d StartBarycentricCoords(0.);
+		for (int32 i = 0; i < 3; ++i)
+		{
+			StartBarycentricCoords[i] = StartSurfacePoint.Position.TriPosition.BarycentricCoords[i];
+		}
+
+		const FMeshGeodesicSurfaceTracer::FMeshSurfaceDirection SurfaceDirection(ExtrinsicRefEID, TracePolarAngle);
+		ActualDist = SurfaceTracer.TraceMeshFromBaryPoint(StartExtrinsicTID, StartBarycentricCoords, SurfaceDirection, TraceDist);
+	}
+	break;
+	default:
+	{
+		// not possible
+		check(0);
+	}
+	}
+
+	const FMeshGeodesicSurfaceTracer::FTraceResult& TraceResult = SurfaceTracer.GetTraceResults().Last();
+	// update new vertex position in two ways: the r3 position and the intrinsic position 
+	// 1) the r3 position
+	FVector3d TraceResultPos = ExtrinsicMesh->GetTriBaryPoint(TraceResult.TriID, TraceResult.Barycentric.X, TraceResult.Barycentric.Y, TraceResult.Barycentric.Z);
+	Vertices[NewVID] = TraceResultPos;
+	// 2) the intrinsic position (i.e. relative to the surface)  
+	//  Use "Triangle" type position for result ( Q: should we classify as "Edge" position if very close to edge ?)
+	FSurfacePoint TraceResultPosition(TraceResult.TriID, TraceResult.Barycentric);
+	IntrinsicVertexPositions.InsertAt(TraceResultPosition, NewVID);
+
+
+	// at the new vertex location: update the polar angle for each new edge relative to the reference direction in this extrinsic triangle. 
+
+
+
+	// fix directions relative to local reference edge on the extrinsic mesh 
+	// by finding direction of the TraceEID edge indecent on NewVID
+	double AngleOffset = 0;
+	{
+		FVector2d Dir = TraceResult.SurfaceDirection.Dir;
+
+		// translate to Dir about the reference edge for this triangle
+		int32 EndRefEID = TIDToReferenceEID[TraceResult.TriID];
+
+		const FMeshGeodesicSurfaceTracer::FTangentTri2& TangentTri2 = SurfaceTracer.GetLastTri();
+		// convert to local basis for first edge of tri2
+		if (TangentTri2.EdgeOrientationSign[0] == -1)
+		{
+			Dir = -Dir;
+		}
+		const int32 IndexOfEndRefEID = TangentTri2.PermutedTriEIDs.IndexOf(EndRefEID);
+		FVector2d DirRelToRefEID = TangentTri2.ChangeBasis(Dir, IndexOfEndRefEID);
+		if (TangentTri2.EdgeOrientationSign[IndexOfEndRefEID] == -1)
+		{
+			DirRelToRefEID = -DirRelToRefEID;
+		}
+		// angle of (new edge) path to NewVert relative to Ref edge
+		const double AngleToNewVert = FMath::Atan2(DirRelToRefEID.Y, DirRelToRefEID.X);
+		// reverse direction of path because we NewVert is the local origin for polar angles around new vert
+		const double AngleFromNewVert = AngleToNewVert + TMathUtilConstants<double>::Pi;
+		AngleOffset = AsZeroToTwoPi(AngleFromNewVert);
+	}
+
+	return AngleOffset;
+}
+
+
+UE::Geometry::EMeshResult UE::Geometry::FIntrinsicTriangulation::PokeTriangleTopology(int32 TriangleID, FPokeTriangleInfo& PokeResult)
+{
+	PokeResult = FPokeTriangleInfo();
+
+	if (!IsTriangle(TriangleID))
+	{
+		return EMeshResult::Failed_NotATriangle;
+	}
+
+	FIndex3i tv = GetTriangle(TriangleID);
+	FIndex3i te = GetTriEdges(TriangleID);
+
+	// create vertex with averaged position.. 
+	FVector3d vPos = (1./3.)*(GetVertex(tv[0]) + GetVertex(tv[1]) + GetVertex(tv[2]));
+
+	int32 newVertID = AppendVertex(vPos);
+
+
+	// add in edges to center vtx, do not connect to triangles yet
+	int32 eAN = AddEdgeInternal(tv[0], newVertID, -1, -1);
+	int32 eBN = AddEdgeInternal(tv[1], newVertID, -1, -1);
+	int32 eCN = AddEdgeInternal(tv[2], newVertID, -1, -1);
+	VertexRefCounts.Increment(tv[0]);
+	VertexRefCounts.Increment(tv[1]);
+	VertexRefCounts.Increment(tv[2]);
+	VertexRefCounts.Increment(newVertID, 3);
+
+	// old triangle becomes tri along first edge
+	Triangles[TriangleID] = FIndex3i(tv[0], tv[1], newVertID);
+	TriangleEdges[TriangleID] = FIndex3i(te[0], eBN, eAN);
+
+	// add two triangles
+	int32 t1 = AddTriangleInternal(tv[1], tv[2], newVertID, te[1], eCN, eBN);
+	int32 t2 = AddTriangleInternal(tv[2], tv[0], newVertID, te[2], eAN, eCN);
+
+	// second and third edges of original tri have neighbors
+	ReplaceEdgeTriangle(te[1], TriangleID, t1);
+	ReplaceEdgeTriangle(te[2], TriangleID, t2);
+
+	// set the triangles for the edges we created above
+	SetEdgeTrianglesInternal(eAN, TriangleID, t2);
+	SetEdgeTrianglesInternal(eBN, TriangleID, t1);
+	SetEdgeTrianglesInternal(eCN, t1, t2);
+
+
+	PokeResult.OriginalTriangle = TriangleID;
+	PokeResult.TriVertices = tv;
+	PokeResult.NewVertex = newVertID;
+	PokeResult.NewTriangles = FIndex2i(t1, t2);
+	PokeResult.NewEdges = FIndex3i(eAN, eBN, eCN);
+	PokeResult.BaryCoords = FVector3d(1./3., 1./3., 1./3.);
+	
+	return EMeshResult::Ok;
+}
+
+UE::Geometry::EMeshResult UE::Geometry::FIntrinsicTriangulation::PokeTriangle(int32 TID, const FVector3d& BaryCoordinates, FPokeTriangleInfo& PokeInfo)
+{
+	if (!IsTriangle(TID))
+	{
+		return EMeshResult::Failed_NotATriangle;
+	}
+
+	// state before the poke
+	const FIndex3i OriginalVIDs = GetTriangle(TID);
+	const FIndex3i OriginalEdges = GetTriEdges(TID);
+	const FVector3d OriginalTriEdgeLengths = GetTriEdgeLengths(TID);
+	const FVector3d OriginalEdgeDirs = IntrinsicEdgeAngles[TID];
+	
+	// Add a new vertex and faces to the IntrinsicMesh.   
+	// Note: the r3 position will be wrong initially since this poke will just interpolate the corners of the tri.
+	//       we fix this position as the last step in this function
+	EMeshResult PokeResult = PokeTriangleTopology(TID, PokeInfo);
+	if (PokeResult != EMeshResult::Ok)
+	{
+		return PokeResult;
+	}
+	
+
+	FIndex3i NewTris(TID, PokeInfo.NewTriangles[0], PokeInfo.NewTriangles[1]);
+	const int32 NewVID = PokeInfo.NewVertex;
+
+	// Need to update intrinsic information for the 3 triangles that resulted from the poke
+	// 1) update the edge lengths for the triangles
+	// 2) update the internal angles for the triangles
+	// 3) update the edge directions for the triangles. 
+	// 4) update the surface point for the new vertex by tracing one of the new edges. 
+	//    also update the edge directions leaving the new vertex to be relative to the direction defined on the extrinsic mesh
+
+
+	// (1) compute intrinsic edge lengths for the 3 new edges
+	{
+		FVector3d DistancesFromOldCorner;
+		for (int32 v = 0; v < 3; ++v)
+		{
+			FVector3d BaryCoordVertex(0., 0., 0.); BaryCoordVertex[v] = 1.;
+			const double Distance = DistanceBetweenBarycentricPoints( OriginalTriEdgeLengths[0],
+							  										  OriginalTriEdgeLengths[1],
+																	  OriginalTriEdgeLengths[2],
+																	  BaryCoordVertex, BaryCoordinates);
+			DistancesFromOldCorner[v] = Distance;
+		}
+	
+		{
+			// original tri is and verts (a, b,c) and edges (ab, bc, ca)
+			// the updated tri has verts (a, b, new) and edges (ab, bnew, newa)
+			FIndex3i T0EIDs = GetTriEdges(NewTris[0]);
+			EdgeLengths.InsertAt(DistancesFromOldCorner[1], T0EIDs[1]);
+			EdgeLengths.InsertAt(DistancesFromOldCorner[0], T0EIDs[2]);
+
+			// new tri[0] has verts (b, c, new) and edges (bc, cnew, newb)
+			FIndex3i T1EIDs = GetTriEdges(PokeInfo.NewTriangles[0]);
+			EdgeLengths.InsertAt(DistancesFromOldCorner[2], T1EIDs[1]);
+		}
+	}
+	// (2) update internal angles for the triangles.
+	{
+		InternalAngles.InsertAt(ComputeTriInternalAnglesR(NewTris[2]), NewTris[2]);
+		InternalAngles.InsertAt(ComputeTriInternalAnglesR(NewTris[1]), NewTris[1]);
+		InternalAngles[NewTris[0]] = ComputeTriInternalAnglesR(NewTris[0]);
+	}
+	
+	// (3) update the edge directions 
+	FVector3d TriEdgeDir[3]; // one for each tri in order (TID, NewTris[0], NewTris[1] )
+	// edges around the boundary of the original tri
+	TriEdgeDir[0][0] = OriginalEdgeDirs[0]; // AtoB dir
+	TriEdgeDir[1][0] = OriginalEdgeDirs[1]; // BtoC dir
+	TriEdgeDir[2][0] = OriginalEdgeDirs[2]; // CtoA dir
+	
+	// edges from the corners of the original tri towards the new vertex at the "center"
+	const double ToRadiansAtA = GeometricVertexInfo[OriginalVIDs[0]].ToRadians;
+	const double ToRadiansAtB = GeometricVertexInfo[OriginalVIDs[1]].ToRadians;
+	const double ToRadiansAtC = GeometricVertexInfo[OriginalVIDs[2]].ToRadians;
+	TriEdgeDir[0][1] = AsZeroToTwoPi( OriginalEdgeDirs[1] + ToRadiansAtB * InternalAngles[NewTris[1]][0]);  // BtoNew dir
+	TriEdgeDir[1][1] = AsZeroToTwoPi( OriginalEdgeDirs[2] + ToRadiansAtC * InternalAngles[NewTris[2]][0]);  // CtoNew dir
+	TriEdgeDir[2][1] = AsZeroToTwoPi( OriginalEdgeDirs[0] + ToRadiansAtA * InternalAngles[NewTris[0]][0]);  // AtoNew dir   
+
+	// edges from new vertex to a, to b, and to c, using new-to-A as the zero direction.  These will be updated later when we learn the correct angle for new-to-A
+	TriEdgeDir[0][2] = 0.;                                                         // NewToA dir 
+	TriEdgeDir[1][2] = InternalAngles[NewTris[0]][2];                                     // NewToB dir
+	TriEdgeDir[2][2] = InternalAngles[NewTris[0]][2] + InternalAngles[NewTris[1]][2];     // NewToC dir
+	GeometricVertexInfo.InsertAt(FGeometricInfo(), NewVID); // we want the default of false, and 1.
+
+	// (4) update the surface point for the new vertex by tracing one of the new edges. 
+	// 
+	// --- fix the r3 position of the new vertex and compute its SurfacePoint Attributes by doing a trace on the Extrinsic Mesh along one of the edges incident on NewVID
+
+	// Use first incident edge and find the distance and direction (angle) to trace
+	// Q: would picking the shortest of the new edges be better than just the first?
+	const int32 AtoNewEID = PokeInfo.NewEdges[0]; // a-to-new edge
+	const double AtoNewDist = EdgeLengths[AtoNewEID];
+
+	// trace from A vertex in the direction toward the new vertex the edge length distance. This updates the surface position entry for the new vert.
+	const int32 AVID = OriginalVIDs[0];
+	const double AtoNewDir = TriEdgeDir[2][1];
+	const double NewToAAngle = UpdateVertexByEdgeTrace(NewVID, AVID, AtoNewDir, AtoNewDist);
+
+	// update the Newto{A,B,C}  directions such that NewToA has the angle LocalEIDAngle.
+	for (int32 e = 0; e < 3; ++e)
+	{ 
+		TriEdgeDir[e][2] = AsZeroToTwoPi(TriEdgeDir[e][2] + NewToAAngle);
+	}
+	
+	// record the new edge directions.
+	IntrinsicEdgeAngles.InsertAt(TriEdgeDir[2], NewTris[2]);
+	IntrinsicEdgeAngles.InsertAt(TriEdgeDir[1], NewTris[1]);
+	IntrinsicEdgeAngles[NewTris[0]] = TriEdgeDir[0];
+	
+	// record the coordinate actually used. 
+	PokeInfo.BaryCoords = BaryCoordinates;
+
+	return EMeshResult::Ok;
+}
+
+
+EMeshResult UE::Geometry::FIntrinsicTriangulation::SplitEdgeTopology(int32 eab, FEdgeSplitInfo& SplitInfo)
+{
+	SplitInfo = FEdgeSplitInfo();
+
+	if (!IsEdge(eab))
+	{
+		return EMeshResult::Failed_NotAnEdge;
+	}
+
+	// look up primary edge & triangle
+	const FEdge Edge = Edges[eab];
+	const int32 t0 = Edge.Tri[0];
+	if (t0 == InvalidID)
+	{
+		return EMeshResult::Failed_BrokenTopology;
+	}
+	const FIndex3i T0tv = GetTriangle(t0);
+	const FIndex3i T0te = GetTriEdges(t0);
+	const int32 IndexOfe = T0te.IndexOf(eab);
+	const int32 a = T0tv[IndexOfe];
+	const int32 b = T0tv[AddOneModThree[IndexOfe]];
+	const int32 c = T0tv[AddTwoModThree[IndexOfe]];
+
+	// look up edge bc, which needs to be modified
+	const int32 ebc = T0te[AddOneModThree[IndexOfe]];
+
+	// RefCount overflow check. Conservatively leave room for
+	// extra increments from other operations.
+	if (VertexRefCounts.GetRawRefCount(c) > FRefCountVector::INVALID_REF_COUNT - 3)
+	{
+		return EMeshResult::Failed_HitValenceLimit;
+	}
+	
+
+	SplitInfo.OriginalEdge = eab;
+	SplitInfo.OriginalVertices = FIndex2i(a, b);   // this is the oriented a,b
+	SplitInfo.OriginalTriangles = FIndex2i(t0, InvalidID);
+	SplitInfo.SplitT = 0.5;
+
+	// quite a bit of code is duplicated between boundary and non-boundary case, but it
+	//  is too hard to follow later if we factor it out...
+	if (IsBoundaryEdge(eab))
+	{
+		// create vertex
+		const FVector3d vNew = 0.5 * (GetVertex(a) + GetVertex(b));
+		const int32 f = AppendVertex(vNew);
+		
+
+		// rewrite existing triangle
+		Triangles[t0][AddOneModThree[IndexOfe]] =  f;
+
+		// add second triangle
+		const int32 t2 = AddTriangleInternal(f, b, c, InvalidID, InvalidID, InvalidID);
+	
+
+		// rewrite edge bc, create edge af
+		ReplaceEdgeTriangle(ebc, t0, t2);
+		const int32 eaf = eab;
+		Edges[eaf].Vert = FIndex2i(FMath::Min(a,f), FMath::Max(a,f));
+		//ReplaceEdgeVertex(eaf, b, f);
+		if (a != b)
+		{ 
+			VertexEdgeLists.Remove(b, eab);
+		}
+		VertexEdgeLists.Insert(f, eaf);
+
+		// create edges fb and fc
+		const int32 efb = AddEdgeInternal(f, b, t2);
+		const int32 efc = AddEdgeInternal(f, c, t0, t2);
+
+		// update triangle edge-nbrs
+		ReplaceTriangleEdge(t0, ebc, efc);
+		TriangleEdges[t2] = FIndex3i(efb, ebc, efc);
+
+		// update vertex refcounts
+		VertexRefCounts.Increment(c);
+		VertexRefCounts.Increment(f, 2);
+
+		SplitInfo.bIsBoundary = true;
+		SplitInfo.OtherVertices = FIndex2i(c, InvalidID);
+		SplitInfo.NewVertex = f;
+		SplitInfo.NewEdges = FIndex3i(efb, efc, InvalidID);
+		SplitInfo.NewTriangles = FIndex2i(t2, InvalidID);
+
+		return EMeshResult::Ok;
+
+	}
+	else 		// interior triangle branch
+	{
+		// look up other triangle
+		const int32 t1 = Edges[eab].Tri[1];
+		SplitInfo.OriginalTriangles.B = t1;
+		const FIndex3i T1tv = GetTriangle(t1);
+		const FIndex3i T1te = GetTriEdges(t1);
+		const int32 T1IndexOfe = T1te.IndexOf(eab);
+		checkSlow(T1tv[T1IndexOfe] == b);
+
+		const int32 d = T1tv[AddTwoModThree[T1IndexOfe]];
+		const int32 edb = T1te[AddTwoModThree[T1IndexOfe]];
+
+		// RefCount overflow check. Conservatively leave room for
+		// extra increments from other operations.
+		if (VertexRefCounts.GetRawRefCount(d) > FRefCountVector::INVALID_REF_COUNT - 3)
+		{
+			return EMeshResult::Failed_HitValenceLimit;
+		}
+
+		// create vertex
+		FVector3d vNew = 0.5 * (GetVertex(a) + GetVertex(b));
+		int32 f = AppendVertex(vNew);
+
+
+		// rewrite existing triangles, replacing b with f
+		Triangles[t0][AddOneModThree[IndexOfe]] = f;
+		Triangles[t1][T1IndexOfe] = f;
+	
+
+		// add two triangles to close holes we just created
+		int32 t2 = AddTriangleInternal(f, b, c, InvalidID, InvalidID, InvalidID);
+		int32 t3 = AddTriangleInternal(f, d, b, InvalidID, InvalidID, InvalidID);
+		
+
+		// update the edges we found above, to point to triangles
+		ReplaceEdgeTriangle(ebc, t0, t2);
+		ReplaceEdgeTriangle(edb, t1, t3);
+
+		// edge eab became eaf
+		int32 eaf = eab; //Edge * eAF = eAB;
+		Edges[eaf].Vert = FIndex2i(FMath::Min(a, f), FMath::Max(a, f));
+
+		// update a/b/f vertex-edges
+		if (a != b)
+		{ 
+			VertexEdgeLists.Remove(b, eab);
+		}
+		VertexEdgeLists.Insert(f, eaf);
+
+		// create edges connected to f  (also updates vertex-edges)
+		int32 efb = AddEdgeInternal(f, b, t2, t3);
+		int32 efc = AddEdgeInternal(f, c, t0, t2);
+		int32 edf = AddEdgeInternal(d, f, t1, t3);
+
+		// update triangle edge-nbrs
+		ReplaceTriangleEdge(t0, ebc, efc);
+		ReplaceTriangleEdge(t1, edb, edf);
+		TriangleEdges[t2] = FIndex3i(efb, ebc, efc);
+		TriangleEdges[t3] = FIndex3i(edf, edb, efb);
+
+		// update vertex refcounts
+		VertexRefCounts.Increment(c);
+		VertexRefCounts.Increment(d);
+		VertexRefCounts.Increment(f, 4);
+
+		SplitInfo.bIsBoundary = false;
+		SplitInfo.OtherVertices = FIndex2i(c, d);
+		SplitInfo.NewVertex = f;
+		SplitInfo.NewEdges = FIndex3i(efb, efc, edf);
+		SplitInfo.NewTriangles = FIndex2i(t2, t3);
+
+		return EMeshResult::Ok;
+	}
+
+}
+
+
+UE::Geometry::EMeshResult  UE::Geometry::FIntrinsicTriangulation::SplitEdge(int32 EdgeAB, FEdgeSplitInfo& SplitInfo, double SplitParameterT)
+{
+	SplitParameterT = FMath::Clamp(SplitParameterT, 0., 1.);
+
+	if (!IsEdge(EdgeAB))
+	{
+		return EMeshResult::Failed_NotAnEdge;
+	}
+
+	
+	const FEdge OriginalEdge = GetEdge(EdgeAB);
+
+	if (IsBoundaryEdge(EdgeAB))
+	{
+		// state before the split
+		const int32 TID = OriginalEdge.Tri[0];
+		const int32 IndexOfe = GetTriEdges(TID).IndexOf(EdgeAB);
+
+		// Info about the original T0 tri, reordered to make the split edge the first edge..
+		const FIndex3i OriginalT0VIDs  = Permute(IndexOfe, GetTriangle(TID));  // as (a, b, c)
+		const FIndex3i OriginalT0Edges = Permute(IndexOfe, GetTriEdges(TID));  // as (ab, bc, ca)
+		
+		const FVector3d OriginalT0EdgeLengths = Permute(IndexOfe, GetTriEdgeLengths(TID));    // as (|ab|, |bc|, |ca|)
+		const FVector3d OriginalT0EdgeDirs    = Permute(IndexOfe, IntrinsicEdgeAngles[TID]);  // as ( aTob, bToc, cToa)
+
+		
+		// Update the connectivity with the edge split
+		EMeshResult Result = SplitEdgeTopology(EdgeAB, SplitInfo);
+
+		if (Result != EMeshResult::Ok)
+		{
+			return Result;
+		}
+		const int32 NewVID = SplitInfo.NewVertex;
+		const int32 NewTID = SplitInfo.NewTriangles[0]; // edges {fb, bc, cf}
+		const int32 NewEdgeFB = SplitInfo.NewEdges[0];
+		const int32 NewEdgeFC = SplitInfo.NewEdges[1];
+
+		const double DistAF = SplitParameterT * OriginalT0EdgeLengths[0];
+		const double DistFB = FMath::Max(0., OriginalT0EdgeLengths[0] - DistAF);
+		// new edge length
+		const double DistFC = DistanceBetweenBarycentricPoints( OriginalT0EdgeLengths[0],
+			                                                    OriginalT0EdgeLengths[1],
+																OriginalT0EdgeLengths[2],
+			                                                    FVector3d(SplitParameterT, 1. - SplitParameterT, 0.), 
+																FVector3d(0., 0., 1.));
+		
+		EdgeLengths.InsertAt(DistFC, NewEdgeFC);
+		EdgeLengths.InsertAt(DistFB, NewEdgeFB);
+		EdgeLengths[EdgeAB] = DistAF;
+		
+		// update the internal angles
+		InternalAngles.InsertAt(ComputeTriInternalAnglesR(NewTID), NewTID);
+		InternalAngles[TID] = ComputeTriInternalAnglesR(TID);
+		const FVector3d TIDInternalAngles = Permute(IndexOfe, InternalAngles[TID]);
+		
+		// update the directions.  Say DirFtoB is zero
+		const double ToRadiansAtC = GeometricVertexInfo[OriginalT0VIDs[2]].ToRadians;
+		GeometricVertexInfo.InsertAt(FGeometricInfo(), NewVID); // we want the default of false, and 1.
+
+		FVector3d TriEdgeDir[2]; // one for each tri in order (TID, NewTID )
+		// edges around the boundary of the original tri
+		TriEdgeDir[0][0] = OriginalT0EdgeDirs[0];    // AtoF dir
+		TriEdgeDir[1][1] = OriginalT0EdgeDirs[1];    // BtoC dir
+		TriEdgeDir[0][2] = OriginalT0EdgeDirs[2];    // CtoA dir  
+
+		TriEdgeDir[1][0] = 0.;  // FtoB dir
+		TriEdgeDir[1][2] = AsZeroToTwoPi(OriginalT0EdgeDirs[2] + ToRadiansAtC * TIDInternalAngles[2]);  // CtoF dir
+		TriEdgeDir[0][1] = AsZeroToTwoPi(InternalAngles[NewTID][0]);                                    // FtoC dir relative to FtoB
+		const double FtoADir = AsZeroToTwoPi(InternalAngles[NewTID][0] + TIDInternalAngles[1]);         // FtoA dir relative to FtoB
+		// update the surface position of the new vertex and the relative directions from it.
+
+		// trace from A vertex in the direction toward the new vertex the edge length distance. This updates the surface position entry for the new vert.
+		const int32 AVID = OriginalT0VIDs[0];
+		const double AtoFDir = TriEdgeDir[0][0];
+		const double FToAAngle = UpdateVertexByEdgeTrace(NewVID, AVID, AtoFDir, DistAF);
+		const double ToLocalDir = (FToAAngle - FtoADir); 
+		// update the Fto{B,C}  directions such that the direction from F to A would agree with FtoAAngle
+		TriEdgeDir[0][1] = AsZeroToTwoPi(TriEdgeDir[0][1] + ToLocalDir);
+		TriEdgeDir[1][0] = AsZeroToTwoPi(TriEdgeDir[1][0] + ToLocalDir);
+		
+		// store angle results
+		IntrinsicEdgeAngles.InsertAt(TriEdgeDir[1], NewTID);
+		IntrinsicEdgeAngles[TID] = Permute((3 - IndexOfe)%3,  TriEdgeDir[0]);
+
+
+		return Result;
+	}
+	else
+	{
+		// state before the split
+		const int32 TID = OriginalEdge.Tri[0];
+		const int32 IndexOfe = GetTriEdges(TID).IndexOf(EdgeAB);
+
+		// Info about the original T0 tri, reordered to make the split edge the first edge..
+		const FIndex3i OriginalT0VIDs  = Permute(IndexOfe, GetTriangle(TID));  // as (a, b, c)
+		const FIndex3i OriginalT0Edges = Permute(IndexOfe, GetTriEdges(TID));  // as (ab, bc, ca)
+
+		const FVector3d OriginalT0EdgeLengths = Permute(IndexOfe, GetTriEdgeLengths(TID));    // as (|ab|, |bc|, |ca|)
+		const FVector3d OriginalT0EdgeDirs    = Permute(IndexOfe, IntrinsicEdgeAngles[TID]);  // as ( aTob, bToc, cToa)
+		
+		// Info about the original T0 tri, reordered to make the split edge the first edge..
+		const int32 TID1 = OriginalEdge.Tri[1];
+		const int32 IndexOfT1e = GetTriEdges(TID1).IndexOf(EdgeAB);
+
+		const FIndex3i OriginalT1VIDs  = Permute(IndexOfT1e, GetTriangle(TID1));  // as (b, a, d)
+		const FIndex3i OriginalT1Edges = Permute(IndexOfT1e, GetTriEdges(TID1));  // as (ba, ad, db)
+		
+		const FVector3d OriginalT1EdgeLengths = Permute(IndexOfT1e, GetTriEdgeLengths(TID1));   // as (|ba|, |ad|, |db})
+		const FVector3d OriginalT1EdgeDirs    = Permute(IndexOfT1e, IntrinsicEdgeAngles[TID1]); // as (bToa, aTod, dTob) 
+		
+
+		// Update the connectivity with the edge split
+		EMeshResult Result = SplitEdgeTopology(EdgeAB, SplitInfo);
+
+		if (Result != EMeshResult::Ok)
+		{
+			return Result;
+		}
+		const int32 NewVID = SplitInfo.NewVertex;
+		const int32 NewTID0 = SplitInfo.NewTriangles[0]; // edges {fb, bc, cf}
+		const int32 NewTID1 = SplitInfo.NewTriangles[1]; // edges {fd, db, bc}
+		const int32 EdgeAF = EdgeAB;
+		const int32 NewEdgeFB = SplitInfo.NewEdges[0];
+		const int32 NewEdgeFC = SplitInfo.NewEdges[1];
+		const int32 NewEdgeFD = SplitInfo.NewEdges[2];
+
+		const double DistAF = SplitParameterT * OriginalT0EdgeLengths[0];
+		const double DistFB = FMath::Max(0., OriginalT0EdgeLengths[0] - DistAF);
+		// new edge length
+		const double DistFC = DistanceBetweenBarycentricPoints( OriginalT0EdgeLengths[0],
+			                                                    OriginalT0EdgeLengths[1],
+			                                                    OriginalT0EdgeLengths[2],
+			                                                    FVector3d(SplitParameterT, 1. - SplitParameterT, 0.),
+			                                                    FVector3d(0., 0., 1.));
+		
+		// new edge length
+		const double DistFD = DistanceBetweenBarycentricPoints( OriginalT1EdgeLengths[0],
+			                                                    OriginalT1EdgeLengths[1],
+			                                                    OriginalT1EdgeLengths[2],
+			                                                    FVector3d(1. - SplitParameterT, SplitParameterT, 0.),
+			                                                    FVector3d(0., 0., 1.));
+
+		EdgeLengths.InsertAt(DistFD, NewEdgeFD);
+		EdgeLengths.InsertAt(DistFC, NewEdgeFC);
+		EdgeLengths.InsertAt(DistFB, NewEdgeFB);
+		EdgeLengths[EdgeAF] = DistAF;
+
+		// update the internal angles (this uses edge lengths)
+		InternalAngles.InsertAt(ComputeTriInternalAnglesR(NewTID1), NewTID1);
+		InternalAngles.InsertAt(ComputeTriInternalAnglesR(NewTID0), NewTID0);
+		InternalAngles[TID] = ComputeTriInternalAnglesR(TID);
+		InternalAngles[TID1] = ComputeTriInternalAnglesR(TID1);
+
+		// update the directions.  Say DirFtoB is zero
+		const double ToRadiansAtC = GeometricVertexInfo[OriginalT0VIDs[2]].ToRadians;
+		const double ToRadiansAtD = GeometricVertexInfo[OriginalT1VIDs[2]].ToRadians;
+		GeometricVertexInfo.InsertAt(FGeometricInfo(), NewVID); // we want the default of false, and 1.
+
+		const FVector3d TIDInternalAngles  = Permute(IndexOfe, InternalAngles[TID]);
+		const FVector3d TID1InternalAngles = Permute(IndexOfT1e, InternalAngles[TID1]);
+
+		FVector3d TriEdgeDir[4]; // one for each tri in order (TID, NewTID0, TID1,  NewTID1 )
+		// edges around the boundary of the original tri
+		TriEdgeDir[0][0] = OriginalT0EdgeDirs[0];    // AtoF dir
+		TriEdgeDir[1][1] = OriginalT0EdgeDirs[1];    // BtoC dir
+		TriEdgeDir[0][2] = OriginalT0EdgeDirs[2];    // CtoA dir  
+
+		TriEdgeDir[3][2] = OriginalT1EdgeDirs[0];    // BtoF dir
+		TriEdgeDir[2][1] = OriginalT1EdgeDirs[1];    // AtoD dir
+		TriEdgeDir[3][1] = OriginalT1EdgeDirs[2];    // DtoB dir
+
+		
+		TriEdgeDir[1][2] = AsZeroToTwoPi(OriginalT0EdgeDirs[2] + ToRadiansAtC * TIDInternalAngles[2]);       // CtoF dir
+		TriEdgeDir[2][2] = AsZeroToTwoPi(OriginalT1EdgeDirs[2] + ToRadiansAtD * InternalAngles[NewTID1][1]); // Dtof dir 
+
+		TriEdgeDir[1][0] = 0.;                                                                          // FtoB dir
+		TriEdgeDir[0][1] = AsZeroToTwoPi(InternalAngles[NewTID0][0]);                                   // FtoC dir relative to FtoB
+		TriEdgeDir[2][0] = AsZeroToTwoPi(InternalAngles[NewTID0][0] + TIDInternalAngles[1]);            // FtoA dir relative to FtoB
+		
+		TriEdgeDir[3][0] = AsZeroToTwoPi(InternalAngles[NewTID0][0] + TIDInternalAngles[1] + TID1InternalAngles[0]); // FtoD dir relative to FtoB
+
+		// trace from A vertex in the direction toward the new vertex the edge length distance. This updates the surface position entry for the new vert.
+		const int32 AVID = OriginalT0VIDs[0];
+		const double AtoFDir = TriEdgeDir[0][0];
+		const double FToAAngle = UpdateVertexByEdgeTrace(NewVID, AVID, AtoFDir, DistAF);
+		const double ToLocalDir = (FToAAngle - TriEdgeDir[2][0]);
+
+		// fix angles relative to local dir.
+		TriEdgeDir[1][0] = AsZeroToTwoPi(ToLocalDir);
+		TriEdgeDir[0][1] = AsZeroToTwoPi(TriEdgeDir[0][1] + ToLocalDir);
+		TriEdgeDir[2][0] = AsZeroToTwoPi(TriEdgeDir[2][0] + ToLocalDir);
+		TriEdgeDir[3][0] = AsZeroToTwoPi(TriEdgeDir[3][0] + ToLocalDir);
+
+		// store angle results
+		IntrinsicEdgeAngles.InsertAt(TriEdgeDir[3], NewTID1);
+		IntrinsicEdgeAngles.InsertAt(TriEdgeDir[1], NewTID0);
+		IntrinsicEdgeAngles[TID]  = Permute((3 - IndexOfe) % 3, TriEdgeDir[0]);
+		IntrinsicEdgeAngles[TID1] = Permute((3 - IndexOfT1e) % 3, TriEdgeDir[2]);
+
+		return Result;
+	}
+}
+
diff --git a/Engine/Source/Runtime/GeometryCore/Private/Spatial/MeshGeodesicSurfaceTracer.cpp b/Engine/Source/Runtime/GeometryCore/Private/Spatial/MeshGeodesicSurfaceTracer.cpp
new file mode 100644
--- /dev/null
+++ b/Engine/Source/Runtime/GeometryCore/Private/Spatial/MeshGeodesicSurfaceTracer.cpp
@@ -0,0 +1,906 @@
+// Copyright Epic Games, Inc. All Rights Reserved.
+
+#include "Spatial/MeshGeodesicSurfaceTracer.h"
+
+#include "DynamicMesh/DynamicMesh3.h"
+#include "Intersection/IntrLine2Line2.h"
+#include "Intersection/IntrLine2Triangle2.h"
+#include "IndexTypes.h"
+#include "MathUtil.h" // for ::TwoPi
+#include "LineTypes.h"
+#include "Util/IndexUtil.h"
+
+#include "BoxTypes.h"
+
+#include "ExplicitUseGeometryMathTypes.h"		// using UE::Geometry::(math types)
+using namespace UE::Geometry;
+
+
+/**
+*  Local Utils for manipulating FIndex{2,3}-based mesh information
+*/ 
+namespace  
+{
+	// assumes VID0 and VID1 are contained in Edge
+	bool IsOrderReversed(int VID0, int32 VID1, const FIndex2i& Edge)
+	{
+		return (VID0 == Edge.B && VID1 == Edge.A);
+	}
+
+	/**
+	* Test if the ordering of the pair {VID0, VID1} is consistent with the vertex ordering in TriVIDs.
+	*
+	* returns 'false' if TriVIDs = {VID0, VID1, X} or {X, VID0, VID1} or { VID1, X, VID0}
+	* where 'X' is some other index in TriVIDs.  Any thing else will return 'true'.
+	* Note: Assumes VID0 and VID1 are contained in TriVIDs and it will  
+	*       return 'true' in the case that either or both VIDs are not contained in TriVIDs
+	*/
+	bool IsOrderReversed(int VID0, int32 VID1, const FIndex3i& TriVIDs)
+	{
+		return (IndexUtil::FindTriOrderedEdge(VID0, VID1, TriVIDs) == IndexConstants::InvalidID);
+	}
+	
+	// Permute the input so SrcIndex3[PivotIndex] becomes ResultIndex3[0] while preserving wrap around order
+	FIndex3i PermuteIndex3(const FIndex3i& SrcIndex3, int32 PivotIndex)
+	{
+		static const int32 AddOneModThree[3] = { 1, 2, 0 };
+		static const int32 AddTwoModThree[3] = { 2, 0, 1 };
+		return FIndex3i(SrcIndex3[PivotIndex], SrcIndex3[AddOneModThree[PivotIndex]], SrcIndex3[AddTwoModThree[PivotIndex]]);
+	}
+
+	// Helper to return the edges ordered as if VID is the start vertex.  
+	// note VID must be one of the Verts for the Tri, or something unexpected will happen
+	FIndex3i GetPermutedEdges(const FDynamicMesh3& Mesh, int32 TriID, int32 VID)
+	{
+		const FIndex3i TriVIDs = Mesh.GetTriangle(TriID);
+		const FIndex3i TriEIDs = Mesh.GetTriEdges(TriID);
+		const int32 TriVertIndex = IndexUtil::FindTriIndex(VID, TriVIDs);
+		return PermuteIndex3(TriEIDs, TriVertIndex);
+	}
+
+	int AbsMinIndex(const FVector3d& Values)
+	{
+		FVector3d AbsValues(FMath::Abs(Values[0]), FMath::Abs(Values[1]), FMath::Abs(Values[2]));
+
+		int MinIndex = 0;
+		for (int i = 1; i < 3; ++i)
+		{
+			if (AbsValues[i] < AbsValues[MinIndex])
+			{
+				MinIndex = i;
+			}
+		}
+		return MinIndex;
+	}
+
+	// Find TriID of the next counter clockwise Triangle when traveling around vertex VID.
+	// Note, this assumes that VID is a vertex in the current triangle.
+	int NextCCWTriangle(const FDynamicMesh3& Mesh, int32 CurTriID, int32 VID)
+	{
+		// Find next triangle traveling counter clockwise.
+		//  - do this by finding the edge shared with the current triangle the next triangle
+		const FIndex3i PermutedTriEIDs = GetPermutedEdges(Mesh, CurTriID, VID);
+		int NextEdgeID = PermutedTriEIDs.C;
+
+
+		FIndex2i NextEdgeTris = Mesh.GetEdgeT(NextEdgeID);
+		int NextTriID = (NextEdgeTris.A == CurTriID) ? NextEdgeTris.B : NextEdgeTris.A;
+		return NextTriID;
+	}
+};
+
+/**
+ *  Misc Triangle Utils
+ */
+namespace 
+{
+	/**
+	* Construct triangle in R2 with same vertex order and edge length as a triangle in R3
+	* where Pos0, Pos1, Pos2 are the vertices of the R3 triangle.
+	*
+	* in particular: Pos0 maps to V[0] = (0,0)
+	*                Pos1 maps to V[1] = (|Pos1 - Pos0|, 0)
+	*                Pos2 maps to V[2] = ( X, Y)  where Y > = 0  and X and Y are such that the edge lengths are preserved.
+	*/
+	template <typename ScalarType>
+	TTriangle2<ScalarType>  MakeTriangle2(const UE::Math::TVector<ScalarType>& Pos0, const UE::Math::TVector<ScalarType>& Pos1, const UE::Math::TVector<ScalarType>& Pos2)
+	{
+		typedef  UE::Math::TVector2<ScalarType>     Vector2Type;
+		typedef UE::Math::TVector<ScalarType>       Vector3Type;
+		typedef TTriangle2<ScalarType>              TriangleType;
+
+		const ScalarType Zero(0);
+
+		// spanning vectors 
+		const Vector3Type TriEdgeVectors[2] = { Pos1 - Pos0,
+												Pos2 - Pos0 };
+
+		const ScalarType BaseEdgeLength = TriEdgeVectors[0].Length();
+		const Vector3Type BaseEdgeDir = TriEdgeVectors[0] / BaseEdgeLength;
+
+
+		// BaseEdgeDir is to be aligned with x-axis.
+		// compute x,y coord of the TriEdgeVector[2] relative to BaseEdgeDir
+		const ScalarType XCoord = BaseEdgeDir.Dot(TriEdgeVectors[1]);
+		Vector3Type cp = BaseEdgeDir.Cross(TriEdgeVectors[1]);
+		const ScalarType YCoord = cp.Length();
+
+
+		return TriangleType(Vector2Type(Zero, Zero),
+			Vector2Type(BaseEdgeLength, Zero),
+			Vector2Type(XCoord, YCoord));
+	}
+
+
+	// Return a valid Barycentric position interior to defining triangle by clamping values and then rescaling to ensure alpha + beta + gamma = 1
+	FVector3d ClampBarycentric(const FVector3d& BCin)
+	{
+		FVector3d BClamped(FMath::Clamp(BCin.X, 0., 1.), FMath::Clamp(BCin.Y, 0., 1.), FMath::Clamp(BCin.Z, 0., 1.));
+		double Sum = BClamped.X + BClamped.Y + BClamped.Z;
+		if (Sum > TMathUtilConstants<double>::ZeroTolerance)
+		{
+			double InvSum = 1. / Sum;
+			BClamped *= InvSum;
+		}
+		else
+		{
+			// default to center.
+			BClamped = FVector3d(1. / 3., 1. / 3., 1. / 3.);
+		}
+		return BClamped;
+	}
+
+	// Test if the given BarycentricPoint describes a valid point with a defining triangle
+	bool IsBarycentricPointInTriangle(const FVector3d& BarycentricPoint, double Tol = TMathUtilConstants<double>::ZeroTolerance)
+	{
+		const double MinVal = FMath::Min3(BarycentricPoint.X, BarycentricPoint.Y, BarycentricPoint.Z);
+		const double MaxVal = FMath::Max3(BarycentricPoint.X, BarycentricPoint.Y, BarycentricPoint.Z);
+		const double Error = FMath::Abs(1. - ( BarycentricPoint.X + BarycentricPoint.Y + BarycentricPoint.Z));
+		return (MinVal > -Tol && MaxVal < 1. + Tol && Error < 3. * Tol);
+	}
+};
+
+
+// -- FTangentTri2 Implementations
+
+GeodesicSingleTriangleUtils::FTangentTri2::FTangentTri2(const FDynamicMesh3& Mesh, int32 TriID, int32 PrimaryEdgeID)
+{
+	Tri3ID = TriID;
+	const FIndex3i TriVIDs = Mesh.GetTriangle(TriID);
+	const FIndex3i TriEIDs = Mesh.GetTriEdges(TriID);
+	const FIndex2i EdgeVertIDs = Mesh.GetEdgeV(PrimaryEdgeID);
+
+	// current edge
+	// e0 = edge(v0, v1),  e1 = edge(v1, v2), e2 = edge(v2, v0)
+	const int32 EdgeIndexInTri = IndexUtil::FindEdgeIndexInTri(EdgeVertIDs.A, EdgeVertIDs.B, TriVIDs);
+
+	PermutedIndices = PermuteIndex3(FIndex3i(0, 1, 2), EdgeIndexInTri);
+
+	PermutedTriVIDs = PermuteIndex3(TriVIDs, EdgeIndexInTri);
+	PermutedTriEIDs = PermuteIndex3(TriEIDs, EdgeIndexInTri);
+
+	for (int i = 0; i < 3; ++i)
+	{
+		FIndex2i EdgeV = Mesh.GetEdgeV(PermutedTriEIDs[i]);
+		EdgeOrientationSign[i] = IsOrderReversed(EdgeV.A, EdgeV.B, PermutedTriVIDs) ? -1 : 1;
+	}
+
+	// The 3d vertices in the Permuted Order.
+	FVector3d Tri3dVerts[3] = { Mesh.GetVertex(PermutedTriVIDs[0]),
+								Mesh.GetVertex(PermutedTriVIDs[1]),
+								Mesh.GetVertex(PermutedTriVIDs[2]) };
+
+	Tri2d = MakeTriangle2(Tri3dVerts[0], Tri3dVerts[1], Tri3dVerts[2]);
+	
+
+	Tri2dEdges[0] = Tri2d.V[1] - Tri2d.V[0];
+	Tri2dEdges[1] = Tri2d.V[2] - Tri2d.V[1];
+	Tri2dEdges[2] = Tri2d.V[0] - Tri2d.V[2];
+
+	Tri2dLengths[0] = Tri2dEdges[0].Length();
+	Tri2dLengths[1] = Tri2dEdges[1].Length();
+	Tri2dLengths[2] = Tri2dEdges[2].Length();
+
+}
+
+void GeodesicSingleTriangleUtils::FTangentTri2::Reset()
+{
+	Tri3ID = IndexConstants::InvalidID;
+}
+
+double  GeodesicSingleTriangleUtils::FTangentTri2::DistanceToEdge(int EdgeIndex, const FVector2d& Point) const
+{
+	return DotPerp(Tri2dEdges[EdgeIndex], Point - Tri2d.V[EdgeIndex]) / Tri2dLengths[EdgeIndex];
+}
+
+// Closest Point on the edge, given in terms of 0,1 parameter measured along the edge (EdgeIndex = 0, 1, 2)
+double  GeodesicSingleTriangleUtils::FTangentTri2::ProjectOnEdge(int EdgeIndex, const FVector2d& Point) const
+{
+	return Tri2dEdges[EdgeIndex].Dot(Point - Tri2d.V[EdgeIndex]) / (Tri2dLengths[EdgeIndex] * Tri2dLengths[EdgeIndex]);
+}
+
+// Express the vector in terms of a basis that is aligned with the EdgeIndex edge of the triangle
+FVector2d GeodesicSingleTriangleUtils::FTangentTri2::ChangeBasis(const FVector2d& Vec, int32 EdgeIndex) const
+{
+	// rotate the Ray relative to the side of the triangle we hit.
+	const FVector2d FrameEdge = Tri2dEdges[EdgeIndex] / Tri2dLengths[EdgeIndex];
+	// Rotation matrix                                   (Cos = FrameEdge.X    Sin =  FrameEdge.Y)
+	// aligns FrameEdge with x-axis                      (-Sin = -FrameEdge.Y  Cos =  FrameEdge.X) 
+	FVector2d RotatedVector(FrameEdge.X * Vec.X + FrameEdge.Y * Vec.Y, -FrameEdge.Y * Vec.X + FrameEdge.X * Vec.Y);
+	RotatedVector = Normalized(RotatedVector);
+
+	return RotatedVector;
+}
+
+// Barycentric coords for this point in the tri2d
+FVector3d  GeodesicSingleTriangleUtils::FTangentTri2::GetBarycentricCoords(const FVector2d& PoinInTri2) const
+{
+	return Tri2d.GetBarycentricCoords(PoinInTri2);
+}
+
+// Convert a point on this tri, to Barycentric coords in the un-rotated source triangle.
+FVector3d  GeodesicSingleTriangleUtils::FTangentTri2::GetBarycentricCoordsInSrcTri(const FVector2d& PoinInTri2) const
+{
+	return ReorderVertexDataForSrcTri3d(GetBarycentricCoords(PoinInTri2));
+}
+
+// ----------------------------------------------------------------------------------------------------//
+// Methods for tracing across a single triangle in a mesh                                              //
+//-----------------------------------------------------------------------------------------------------//
+
+
+GeodesicSingleTriangleUtils::FTraceResult  GeodesicSingleTriangleUtils::TraceTangentTriangle(const FTangentTri2& TangentTri2, const FVector2d& RayOrigin, const FVector2d& RayDir, double MaxDistance)
+{
+	bool bHasIntersection = false;
+	double RayDistance = TMathUtilConstants<double>::MaxReal;
+	FVector2d HitPoint;
+
+	// special case of ray starting on edge and traveling parallel to the edge.  Say the vertex at the end of the edge is the hit
+	if (FMath::Abs(RayOrigin.Y) < TMathUtilConstants<double>::ZeroTolerance && FMath::Abs(RayDir.Y) < TMathUtilConstants<double>::ZeroTolerance)
+	{
+		bHasIntersection = true;
+		if (RayDir.X > 0)
+		{
+			HitPoint = TangentTri2.Tri2d.V[1];
+			RayDistance = (RayOrigin - HitPoint).Length();
+		}
+		else
+		{
+			HitPoint = TangentTri2.Tri2d.V[0];
+			RayDistance = (RayOrigin - HitPoint).Length();
+		}
+	}
+	else // test ray against triangle.
+	{ 
+
+		const FLine2d RayPath(RayOrigin, RayDir);
+
+		// Find Intr with Tri2
+		FIntrLine2Triangle2d LineTriIntr(RayPath, TangentTri2.Tri2d);
+		bHasIntersection = LineTriIntr.Find();
+
+		// classify the intersection(s)
+		const bool bIsSegmentIntersection = (LineTriIntr.Type == EIntersectionType::Segment);
+		const bool bIsPointIntersection = (LineTriIntr.Type == EIntersectionType::Point);
+
+
+		if (bIsSegmentIntersection)
+		{
+			// RayPath.Origin is located on the base of the triangle in the case of tracing from edge
+			// or inside the triangle when tracing from barycentric, 
+			// so there will be an intersection at Param <=0 at or behind the ray origin.
+			// we want the other intersection.
+			if (LineTriIntr.Param0 > LineTriIntr.Param1)
+			{
+				RayDistance = LineTriIntr.Param0;
+				HitPoint = LineTriIntr.Point0;
+			}
+			else
+			{
+				RayDistance = LineTriIntr.Param1;
+				HitPoint = LineTriIntr.Point1;
+			}
+		}
+		else if (bIsPointIntersection) // this is the grazing case where the line clips a single vertex of the triangle
+		{
+			RayDistance = LineTriIntr.Param0;
+			HitPoint = LineTriIntr.Point0;
+		}
+
+	}
+	
+
+
+	// populate the result.
+	FVector2d EndPoint;
+	FTraceResult  TraceResult;
+	if (bHasIntersection == false)
+	{
+		TraceResult.Classification = ETraceClassification::Failed;
+		return TraceResult;
+	}
+
+	bool bEdgeOrderReversed;
+	if (RayDistance <= MaxDistance)  // result on edge of triangle
+	{
+		EndPoint = HitPoint;
+
+		// Expect hit will be on either e1 = v2-v1 
+		//                          or  e2 = v0-v2
+		// unless the trace direction was actually out of the triangle
+		// compute the distance to each edge.  one should be zero, although all could be small.
+		const FVector3d DistHitToEdge(TangentTri2.DistanceToEdge(0, HitPoint), TangentTri2.DistanceToEdge(1, HitPoint), TangentTri2.DistanceToEdge(2, HitPoint));
+
+		// Pick the edge we hit
+		const int32 HitEdgeIndex = AbsMinIndex(DistHitToEdge);
+
+		bEdgeOrderReversed = (TangentTri2.EdgeOrientationSign[HitEdgeIndex] == -1);
+		// EdgeID for the edge we hit. 
+		const int32 HitEID = TangentTri2.PermutedTriEIDs[HitEdgeIndex];
+
+		// --- encode the hit relative to the edge we hit.
+		// 
+		// Distance measured from V[HitIndex] (along the edge) to the hit point as fraction of edge length.
+		const double HitEdgeAlpha = TangentTri2.ProjectOnEdge(HitEdgeIndex, HitPoint);
+
+		// rotate the Ray relative to the side of the triangle we hit. i.e. the rotation that aligns the hit edge with the x-axis
+		const FVector2d RotatedRayDir = TangentTri2.ChangeBasis(RayDir, HitEdgeIndex);
+
+		// populate trace result
+		TraceResult.Classification = ETraceClassification::Continuing;
+		TraceResult.TraceDist = RayDistance;
+		TraceResult.TriID = TangentTri2.Tri3ID;
+
+		TraceResult.SurfaceDirection = FMeshSurfaceDirection(HitEID, RotatedRayDir);
+
+		TraceResult.bIsEdgePoint = true;
+		TraceResult.EdgeID = HitEID;
+		TraceResult.EdgeAlpha = HitEdgeAlpha;
+
+	}
+	else // result internal to tri face
+	{
+		bEdgeOrderReversed = (TangentTri2.EdgeOrientationSign[0] == -1);
+		// Note: EndPoint isn't the HitPoint in this case
+		EndPoint = MaxDistance * RayDir + RayOrigin;
+
+		TraceResult.Classification = ETraceClassification::DistanceTerminated;
+		TraceResult.TraceDist = MaxDistance;
+		TraceResult.TriID = TangentTri2.Tri3ID;
+
+		TraceResult.bIsEdgePoint = false;
+		TraceResult.EdgeID = TangentTri2.PermutedTriEIDs[0]; //EID;
+		TraceResult.EdgeAlpha = 0.; // not used in this case
+
+		TraceResult.SurfaceDirection = FMeshSurfaceDirection(TraceResult.EdgeID, RayDir);
+
+	}
+	TraceResult.Barycentric = TangentTri2.GetBarycentricCoordsInSrcTri(EndPoint);
+
+	// may need to update this encoding to match the vertex order in Mesh.GetEdgeV 
+	if (bEdgeOrderReversed)
+	{
+		TraceResult.SurfaceDirection.Dir = -TraceResult.SurfaceDirection.Dir;
+		TraceResult.EdgeAlpha = 1. - TraceResult.EdgeAlpha;
+	}
+	TraceResult.EdgeAlpha = FMath::Clamp(TraceResult.EdgeAlpha, 0., 1.);
+
+	return TraceResult;
+}
+
+GeodesicSingleTriangleUtils::FTraceResult  GeodesicSingleTriangleUtils::TraceTriangleFromBaryPoint(const FDynamicMesh3& Mesh, const int32 TriID, const FVector3d& BaryPoint, const FMeshSurfaceDirection& Direction, double MaxDistance)
+{
+	FTangentTri2 ScratchTangentTri2;
+	return TraceTriangleFromBaryPoint(Mesh, TriID, BaryPoint, Direction, ScratchTangentTri2, MaxDistance );
+}
+
+GeodesicSingleTriangleUtils::FTraceResult  GeodesicSingleTriangleUtils::TraceTriangleFromBaryPoint(const FDynamicMesh3& Mesh, const int32 TriID, const FVector3d& BaryPointIn, const FMeshSurfaceDirection& Direction, FTangentTri2& TangentTri2, double MaxDistance)
+{
+	FIndex3i TriVIDs = Mesh.GetTriangle(TriID);
+
+	
+	FVector3d BaryPoint(BaryPointIn);
+	constexpr double CornerEpsilon = 1.e-4; // trace from vertex tolerance (in barycentric space )
+
+	// Don't trace from a corner of the triangle, slightly offsetset towards the interior and start there.
+	if (FMath::Abs(BaryPoint.X - 1.) < CornerEpsilon)
+	{
+		BaryPoint = FVector3d( 1. -CornerEpsilon, CornerEpsilon/2., CornerEpsilon/2.);
+	}
+	else if (FMath::Abs(BaryPoint.Y - 1.) < CornerEpsilon)
+	{
+		BaryPoint = FVector3d( CornerEpsilon / 2., 1. - CornerEpsilon, CornerEpsilon / 2.);
+	}
+	else if (FMath::Abs(BaryPoint.Z - 1.) < CornerEpsilon)
+	{
+		BaryPoint = FVector3d(CornerEpsilon / 2., CornerEpsilon / 2., 1. - CornerEpsilon);
+	}
+
+	FVector2d RayDir = Direction.Dir;
+	int EID = Direction.EdgeID;
+	// The 2d Tangent Triangle is in the first quadrant with EID aligned with the x-axis
+	TangentTri2 = FTangentTri2(Mesh, TriID, EID);
+	if (TangentTri2.EdgeOrientationSign[0] == -1)
+	{
+		// the frame that encoded the direction vector is rotated 180 degrees from the first edge of this triangle (our local reference )
+		// update the direction.
+		RayDir = -RayDir;
+	}
+
+	// Barycentric coord relative to the rotated 2d triangle.
+	const FVector3d Bary2d = { BaryPoint[TangentTri2.PermutedIndices[0]], BaryPoint[TangentTri2.PermutedIndices[1]], BaryPoint[TangentTri2.PermutedIndices[2]] };
+	const FVector2d RayOrigin = TangentTri2.Tri2d.BarycentricPoint(Bary2d);
+
+	// Path Entering the 2d triangle, relative to this 2d triangle  
+	FLine2d PathRay(RayOrigin, RayDir);
+
+	FTraceResult TraceResult = TraceTangentTriangle(TangentTri2, RayOrigin, RayDir, MaxDistance);
+	if (!IsTerminated(TraceResult))
+	{
+		FIndex2i EdgeT = Mesh.GetEdgeT(TraceResult.EdgeID);
+		if (EdgeT.A == FDynamicMesh3::InvalidID || EdgeT.B == FDynamicMesh3::InvalidID)
+		{
+			TraceResult.Classification = ETraceClassification::BoundaryTerminated;
+		}
+	}
+
+	checkSlow(TraceResult.bIsEdgePoint || TraceResult.EdgeID == EID);
+
+	return TraceResult;
+
+}
+
+GeodesicSingleTriangleUtils::FTraceResult  GeodesicSingleTriangleUtils::TraceTriangleFromEdge(const FDynamicMesh3& Mesh, const int32 TriID, double EdgeAlpha, const FMeshSurfaceDirection& Direction, double MaxDistance)
+{
+	FTangentTri2 ScratchTangentTri2;
+	return TraceTriangleFromEdge(Mesh, TriID, EdgeAlpha,  Direction, ScratchTangentTri2, MaxDistance);
+}
+
+GeodesicSingleTriangleUtils::FTraceResult  GeodesicSingleTriangleUtils::TraceTriangleFromEdge(const FDynamicMesh3& Mesh, const int32 TriID, double EdgeAlpha, const FMeshSurfaceDirection& Direction, FTangentTri2& TangentTri2, double MaxDistance)
+{
+
+	FVector2d RayDir = Direction.Dir;
+	const int32 EID = Direction.EdgeID;
+
+	// The 2d Tangent Triangle is in the first quadrant with EID aligned with the x-axis
+	TangentTri2 = FTangentTri2(Mesh, TriID, EID);
+	if (TangentTri2.EdgeOrientationSign[0] == -1)
+	{
+		// the crossing data was encoded with the vector pointing in the direction opposite the local triangle order, update it
+		EdgeAlpha = 1. - EdgeAlpha;
+		RayDir = -RayDir;
+	}
+	EdgeAlpha = FMath::Clamp(EdgeAlpha, 0., 1.);
+
+	FVector2d RayOrigin(TangentTri2.Tri2dEdges[0] * EdgeAlpha);
+
+	FTraceResult TraceResult = TraceTangentTriangle(TangentTri2, RayOrigin, RayDir, MaxDistance);
+	if (!IsTerminated(TraceResult))
+	{
+		FIndex2i EdgeT = Mesh.GetEdgeT(TraceResult.EdgeID);
+		if (EdgeT.A == FDynamicMesh3::InvalidID || EdgeT.B == FDynamicMesh3::InvalidID)
+		{
+			TraceResult.Classification = ETraceClassification::BoundaryTerminated;
+		}
+	}
+
+	checkSlow(TraceResult.bIsEdgePoint || TraceResult.EdgeID == EID);
+
+	return TraceResult;
+}
+
+
+
+// Note: Assumes RayDir is defined in a triangle adjacent to the edge, not on the tangent plane!
+GeodesicSingleTriangleUtils::FTraceResult GeodesicSingleTriangleUtils::TraceFromVertex(const FDynamicMesh3& Mesh, int32 VID, const FMeshSurfaceDirection& Direction, double MaxDistance)
+{
+
+	FTangentTri2 ScratchTangentTri2;
+	return TraceFromVertex(Mesh, VID, Direction, ScratchTangentTri2, MaxDistance);
+}
+GeodesicSingleTriangleUtils::FTraceResult GeodesicSingleTriangleUtils::TraceFromVertex(const FDynamicMesh3& Mesh, int32 VID, const FMeshSurfaceDirection& Direction, FTangentTri2& ScratchTangentTri2, double MaxDistance)
+{
+
+	int32 EID = Direction.EdgeID;
+	FIndex2i Edge = Mesh.GetEdgeV(EID);
+	check(Edge.A == VID || Edge.B == VID);
+
+	// This doesn't work at a bowtie, or a boundary vertex. 
+	check(!Mesh.IsBowtieVertex(VID));
+	check(!Mesh.IsBoundaryVertex(VID));
+
+	// @todo: return a boundary terminated result in place of the checks
+
+
+
+	// Compute local tangent plane information. 
+	struct FInternalAngle
+	{
+		int32 TriID;
+		double Angle; // internal angle
+		double OffsetAngle; // sum of prior internal angles.
+	};
+
+	double TotalAngle = 0;
+	TArray<FInternalAngle> InternalAngles;
+	FIndex2i FirstLastTri = Mesh.GetEdgeT(EID);
+
+	// Travel counter clockwise around the vertex, visiting each triangle in order.
+	// Start with the first triangle we visit when traveling counter clockwise around VID, starting with edge EID.
+	int CurTriID = FirstLastTri.A;
+	if (GetPermutedEdges(Mesh, FirstLastTri.B, VID).A == EID)
+	{
+		CurTriID = FirstLastTri.B;
+	}
+	int32 FirstTriID = CurTriID;
+	do
+	{
+		FIndex3i TriVIDS = Mesh.GetTriangle(CurTriID);
+
+		const int32 TriIndex = IndexUtil::FindTriIndex(VID, TriVIDS);
+		const double AngleR = Mesh.GetTriInternalAngleR(CurTriID, TriIndex);
+
+
+		FInternalAngle& InternalAngle = InternalAngles.AddDefaulted_GetRef();
+		InternalAngle.TriID = CurTriID;
+		InternalAngle.Angle = AngleR;
+		InternalAngle.OffsetAngle = TotalAngle;
+
+		CurTriID = NextCCWTriangle(Mesh, CurTriID, VID);
+		TotalAngle += AngleR;
+
+	} while (CurTriID != FirstTriID && CurTriID != FDynamicMesh3::InvalidID);
+
+	// The direction of the ray was encoded as RayDir by assuming the edge is directed from Edge.A to Edge.B
+	// and aligned with the x-axis.
+	// if we trace from the vertex Edge.B, we need to flip the ray
+	const FVector2d RayDir = (VID == Edge.B) ? -Direction.Dir : Direction.Dir;
+
+
+	FVector2d InComingRay = -RayDir;
+	double IndicentAngle = FMath::Atan2(InComingRay.Y, InComingRay.X); // angle on a triangle face, not the tangent plane
+
+	double ExitPolarAngle = IndicentAngle * (TMathUtilConstants<double>::TwoPi / TotalAngle) + TMathUtilConstants<double>::Pi;
+	ExitPolarAngle = AsZeroToTwoPi(ExitPolarAngle);
+
+	double ExitAngle = ExitPolarAngle * (TotalAngle / TMathUtilConstants<double>::TwoPi);
+
+	// find the triangle that the ray enters, and angle from reference to the first side of that triangle
+	int32 DstTriID = FDynamicMesh3::InvalidID;
+	double DstOffsetAngle = 0;
+	for (const FInternalAngle& AngleData : InternalAngles)
+	{
+		if (AngleData.OffsetAngle <= ExitAngle && (AngleData.OffsetAngle + AngleData.Angle) >= ExitAngle)
+		{
+			DstTriID = AngleData.TriID;
+			DstOffsetAngle = AngleData.OffsetAngle;
+			break;
+		}
+	}
+	checkSlow(DstTriID != FDynamicMesh3::InvalidID);
+
+	// What is the first edge of the Dst Tri?
+	const FIndex3i PermutedDstEdges = GetPermutedEdges(Mesh, DstTriID, VID);
+	int DstEdgeID = PermutedDstEdges.A; // First Edge.
+
+	double AngleEdgeToRay = ExitAngle - DstOffsetAngle;
+
+	// ray relative to edge exiting VID
+	double DstAlpha = 0.;
+	FVector2d DstRayDir(FMath::Cos(AngleEdgeToRay), FMath::Sin(AngleEdgeToRay));
+	FIndex2i DstEdge = Mesh.GetEdgeV(DstEdgeID);
+	if (DstEdge.B == VID)
+	{
+		DstAlpha = 1. - DstAlpha;
+		DstRayDir = -DstRayDir;
+	}
+	FMeshSurfaceDirection OutGoingDirection(DstEdgeID, DstRayDir);
+	return TraceTriangleFromEdge(Mesh, DstTriID, DstAlpha, OutGoingDirection, ScratchTangentTri2, MaxDistance);
+}
+
+
+GeodesicSingleTriangleUtils::FTraceResult  GeodesicSingleTriangleUtils::TraceFromVertex(const FDynamicMesh3& Mesh, const FMeshTangentDirection& TangentDirection, double MaxDistance)
+{
+
+	FTangentTri2 ScratchTangentTri2;
+	return TraceFromVertex(Mesh, TangentDirection, ScratchTangentTri2, MaxDistance);
+}
+GeodesicSingleTriangleUtils::FTraceResult  GeodesicSingleTriangleUtils::TraceFromVertex(const FDynamicMesh3& Mesh, const FMeshTangentDirection& TangentDirection, FTangentTri2& ScratchTangentTri2, double MaxDistance)
+{
+	int32 EID = TangentDirection.EdgeID;
+	int32 VID = TangentDirection.VID;
+
+	FIndex2i EdgeV = Mesh.GetEdgeV(EID);
+	check(EdgeV.A == VID || EdgeV.B == VID);
+	// This doesn't work at a bowtie, or a boundary vertex. 
+	// @todo: return a boundary terminated result in place of the checks
+	check(!Mesh.IsBowtieVertex(VID));
+	check(!Mesh.IsBoundaryVertex(VID));
+
+
+	// Compute local tangent plane information. 
+	struct FInternalAngle
+	{
+		int32 TriID;
+		double Angle; // internal angle
+		double OffsetAngle; // sum of prior internal angles.
+	};
+
+	double TotalAngle = 0;
+	TArray<FInternalAngle> InternalAngles;
+	FIndex2i FirstLastTri = Mesh.GetEdgeT(EID);
+
+	// Travel counter-clockwise around the vertex, visiting each triangle in order.
+	// Start with the first triangle we visit when traveling counter clockwise around VID, starting with edge EID.
+	int CurTriID = FirstLastTri.A;
+	if (GetPermutedEdges(Mesh, FirstLastTri.B, VID).A == EID)
+	{
+		CurTriID = FirstLastTri.B;
+	}
+	int FirstTriID = CurTriID;
+	do
+	{
+		FIndex3i TriVIDS = Mesh.GetTriangle(CurTriID);
+
+		const int32 TriIndex = IndexUtil::FindTriIndex(VID, TriVIDS);
+		const double AngleR = Mesh.GetTriInternalAngleR(CurTriID, TriIndex);
+
+
+		FInternalAngle& InternalAngle = InternalAngles.AddDefaulted_GetRef();
+		InternalAngle.TriID = CurTriID;
+		InternalAngle.Angle = AngleR;
+		InternalAngle.OffsetAngle = TotalAngle;
+
+		CurTriID = NextCCWTriangle(Mesh, CurTriID, VID);
+		TotalAngle += AngleR;
+
+	} while (CurTriID != FirstTriID && CurTriID != FDynamicMesh3::InvalidID);
+
+	const double ToSurfaceAngle = TotalAngle / TMathUtilConstants<double>::TwoPi;
+	// if the edge "ends" at VID then it aligns with Pi in the tangent space (i.e. the vector A->B aligns with 0 radians, and the vector B->A aligns with pi )
+	double PolarAngle = (EdgeV.B == VID) ? TangentDirection.PolarAngle + TMathUtilConstants<double>::Pi : TangentDirection.PolarAngle;
+	const double SurfaceAngle = AsZeroToTwoPi(PolarAngle) * ToSurfaceAngle;
+
+
+	// find the triangle that the ray enters, and angle from reference to the first side of that triangle
+	int DstTriID = FDynamicMesh3::InvalidID;
+	double DstOffsetAngle = 0;
+	for (const FInternalAngle& AngleData : InternalAngles)
+	{
+		if (AngleData.OffsetAngle <= SurfaceAngle && (AngleData.OffsetAngle + AngleData.Angle) >= SurfaceAngle)
+		{
+			DstTriID = AngleData.TriID;
+			DstOffsetAngle = AngleData.OffsetAngle;
+			break;
+		}
+	}
+	checkSlow(DstTriID != FDynamicMesh3::InvalidID);
+
+	// What is the first edge of the Dst Tri?
+	const FIndex3i PermutedDstEdges = GetPermutedEdges(Mesh, DstTriID, VID);
+	int DstEdgeID = PermutedDstEdges.A; // First Edge.
+
+	double AngleEdgeToRay = SurfaceAngle - DstOffsetAngle;
+
+	// ray relative to edge exiting VID
+	double DstAlpha = 0.;
+	FVector2d DstRayDir(FMath::Cos(AngleEdgeToRay), FMath::Sin(AngleEdgeToRay));
+	FIndex2i DstEdge = Mesh.GetEdgeV(DstEdgeID);
+	if (DstEdge.B == VID)
+	{
+		DstAlpha = 1. - DstAlpha;
+		DstRayDir = -DstRayDir;
+	}
+	FMeshSurfaceDirection OutGoingDirection(DstEdgeID, DstRayDir);
+	return TraceTriangleFromEdge(Mesh, DstTriID, DstAlpha, OutGoingDirection, ScratchTangentTri2, MaxDistance);
+}
+
+
+GeodesicSingleTriangleUtils::FTraceResult  GeodesicSingleTriangleUtils::TraceNextTriangle(const FDynamicMesh3& Mesh, const  FMeshGeodesicSurfaceTracer::FTraceResult& StartTrace, double MaxDistance)
+{
+	FTangentTri2 ScratchTangentTri2;
+	return TraceNextTriangle(Mesh, StartTrace, ScratchTangentTri2, MaxDistance);
+}
+GeodesicSingleTriangleUtils::FTraceResult  GeodesicSingleTriangleUtils::TraceNextTriangle(const FDynamicMesh3& Mesh, const FMeshGeodesicSurfaceTracer::FTraceResult& LastTrace, FTangentTri2& ScratchTangentTri2, double MaxDistance)
+{
+
+	constexpr double CornerEpsilon = 1.e-4;
+	if (LastTrace.bIsEdgePoint)
+	{
+		double EdgeAlpha = LastTrace.EdgeAlpha;
+		// don't allow starting on a vertex: displace slightly from the vertex and trace from edge.
+		if (FMath::Abs(EdgeAlpha - 1.) < CornerEpsilon)
+		{
+			EdgeAlpha = 1. - CornerEpsilon;
+		}
+		else if (FMath::Abs(EdgeAlpha - 0.) < CornerEpsilon)
+		{
+			EdgeAlpha = CornerEpsilon;
+		}
+
+		const int32 EdgeID = LastTrace.SurfaceDirection.EdgeID;
+		int TriID; // The triangle we are entering.
+		{
+			FIndex2i AdjTriIDs = Mesh.GetEdgeT(EdgeID);
+			TriID = (AdjTriIDs.A == LastTrace.TriID) ? AdjTriIDs.B : AdjTriIDs.A;
+		}
+
+		if (TriID == FDynamicMesh3::InvalidID)
+		{
+			// we hit a boundary, copy the trace surface point and direction, but update the classification and distance
+			FTraceResult Result = LastTrace;
+
+			Result.Classification = FMeshGeodesicSurfaceTracer::ETraceClassification::BoundaryTerminated;
+			Result.TraceDist = 0.;
+
+			return Result;
+		}
+
+		//start from edge
+		return TraceTriangleFromEdge(Mesh, TriID, EdgeAlpha, LastTrace.SurfaceDirection, ScratchTangentTri2, MaxDistance);
+	}
+	else
+	{
+		return TraceTriangleFromBaryPoint(Mesh, LastTrace.TriID, LastTrace.Barycentric, LastTrace.SurfaceDirection, ScratchTangentTri2, MaxDistance);
+	}
+
+}
+
+
+double FMeshGeodesicSurfaceTracer::TraceMeshFromBaryPoint(const int32 TriID, const FVector3d& BaryPoint, const FVector3d& RayDir3, double MaxDistance)
+{
+	using namespace GeodesicSingleTriangleUtils;
+
+	check(Mesh);
+
+	checkSlow(IsBarycentricPointInTriangle(BaryPoint));
+
+	// sanitize.. 
+	FVector3d BaryPointClamped = ClampBarycentric(BaryPoint);
+
+	// project the input direction onto the plane of the specified triangle.
+
+	// create a frame for the face with the origin at V0	
+	const FIndex3i TriVIDs = Mesh->GetTriangle(TriID);
+	const FVector3d ZDir = Mesh->GetTriNormal(TriID);
+	const FVector3d XDir = Normalized(Mesh->GetVertex(TriVIDs[1]) - Mesh->GetVertex(TriVIDs[0]));
+	// Unreal has Left-Hand Coordinate System so we need to reverse this cross-product. (would be Y = Z x X  normally)
+	const FVector3d YDir = Normalized(FVector3d::CrossProduct(XDir, ZDir));
+
+	// project the ray onto the X,Y plane ( i.e. the face of the triangle )
+	FVector2d ProjectedRay(FVector3d::DotProduct(XDir, RayDir3), FVector3d::DotProduct(YDir, RayDir3));
+
+	// lower level code expects the 2d ray to be encoded relative to the "direction" of the first edge, edge(v0,v1)
+	const int32 EID = Mesh->GetTriEdge(TriID, 0);
+	const FIndex2i Edge0 = Mesh->GetEdgeV(EID);
+
+	if (Edge0.B == TriVIDs.A && Edge0.A == TriVIDs.B)
+	{
+		ProjectedRay = -ProjectedRay;
+	}
+
+	FMeshSurfaceDirection Direction(EID, ProjectedRay);
+	return TraceMeshFromBaryPoint(TriID, BaryPointClamped, Direction, MaxDistance);
+}
+
+double FMeshGeodesicSurfaceTracer::TraceMeshFromBaryPoint(const int32 TriID, const FVector3d& BaryPoint, const FMeshSurfaceDirection& Direction, double MaxDistance)
+{
+	using namespace GeodesicSingleTriangleUtils;
+
+	check(Mesh);
+
+	checkSlow(IsBarycentricPointInTriangle(BaryPoint));
+
+	// sanitize.. 
+	FVector3d BaryPointClamped = ClampBarycentric(BaryPoint);
+
+	FVector2d RayDir = Normalized(Direction.Dir);
+	const int32 EID = Direction.EdgeID;
+	// early out if the edge isn't part of the tri
+	FIndex3i TriEdges = Mesh->GetTriEdges(TriID);
+	if (TriEdges.A != EID && TriEdges.B != EID && TriEdges.C != EID)
+	{
+		return 0;
+	}
+
+	double TotalDist = 0.;
+	FMeshSurfaceDirection StartDirection(EID, RayDir);
+	{
+		FTraceResult TraceStart;
+		TraceStart.Classification = FMeshGeodesicSurfaceTracer::ETraceClassification::Start;
+		TraceStart.TraceDist = 0.;
+		TraceStart.TriID = TriID;
+		TraceStart.SurfaceDirection = StartDirection;
+		TraceStart.Barycentric = BaryPointClamped;
+		TraceStart.bIsEdgePoint = false;
+		TraceStart.EdgeID = EID;
+		TraceStart.EdgeAlpha = -1.;
+
+		SurfaceTrace.Add(TraceStart);
+	}
+
+	FTraceResult TraceResult = TraceTriangleFromBaryPoint(*Mesh, TriID, BaryPointClamped, StartDirection, ScratchTri2, MaxDistance);
+	bool bTerminated = IsTerminated(TraceResult);
+	TotalDist += TraceResult.TraceDist;
+	SurfaceTrace.Add(TraceResult);
+
+	MaxDistance -= TraceResult.TraceDist;
+	if (bTerminated)
+	{
+		return TotalDist;
+	}
+
+	do
+	{
+		TraceResult = TraceNextTriangle(*Mesh, TraceResult, ScratchTri2, MaxDistance);
+		bTerminated = IsTerminated(TraceResult);
+		TotalDist += TraceResult.TraceDist;
+		MaxDistance -= TraceResult.TraceDist;
+		SurfaceTrace.Add(TraceResult);
+
+	} while (bTerminated == false);
+
+	return TotalDist;
+}
+
+
+double FMeshGeodesicSurfaceTracer::TraceMeshFromVertex(const FMeshTangentDirection& Direction, double MaxDistance)
+{
+	using namespace GeodesicSingleTriangleUtils;
+
+	check(Mesh);
+
+	const int32 VID = Direction.VID;
+	const int32 EID = Direction.EdgeID;
+	// early out if we try to trace from a bowtie
+	bool bIsBowTie = Mesh->IsBowtieVertex(VID);
+	if (bIsBowTie)
+	{
+		return 0;
+	}
+
+	double TotalDist = 0.;
+	FMeshSurfaceDirection StartDirection(EID, FVector2d(FMath::Cos(Direction.PolarAngle), FMath::Sin(Direction.PolarAngle)));
+	{
+		int TriID = Mesh->GetEdgeT(EID).A;
+		FIndex3i TriVIDs = Mesh->GetTriangle(TriID);
+		int IndexOfVID = TriVIDs.IndexOf(VID);
+		FVector3d BaryPoint(0);
+		BaryPoint[IndexOfVID] = 1.;
+
+		FTraceResult TraceStart;
+		TraceStart.Classification = ETraceClassification::Start;
+		TraceStart.TraceDist = 0.;
+		TraceStart.TriID = TriID;
+		TraceStart.SurfaceDirection = StartDirection;
+		TraceStart.Barycentric = BaryPoint;
+		TraceStart.bIsEdgePoint = false;
+		TraceStart.EdgeID = EID;
+		TraceStart.EdgeAlpha = -1.;
+
+		SurfaceTrace.Add(TraceStart);
+	}
+
+	FTraceResult TraceResult = TraceFromVertex(*Mesh, Direction, ScratchTri2, MaxDistance);
+	bool bTerminated = IsTerminated(TraceResult);
+	TotalDist += TraceResult.TraceDist;
+	SurfaceTrace.Add(TraceResult);
+
+	MaxDistance -= TraceResult.TraceDist;
+	if (bTerminated)
+	{
+		return TotalDist;
+	}
+
+	do
+	{
+		TraceResult = TraceNextTriangle(*Mesh, TraceResult, ScratchTri2, MaxDistance);
+		bTerminated = IsTerminated(TraceResult);
+		TotalDist += TraceResult.TraceDist;
+		MaxDistance -= TraceResult.TraceDist;
+		SurfaceTrace.Add(TraceResult);
+
+	} while (bTerminated == false);
+
+	return TotalDist;
+}
diff --git a/Engine/Source/Runtime/GeometryCore/Public/Spatial/IntrinsicTriangulationMesh.h b/Engine/Source/Runtime/GeometryCore/Public/Spatial/IntrinsicTriangulationMesh.h
new file mode 100644
--- /dev/null
+++ b/Engine/Source/Runtime/GeometryCore/Public/Spatial/IntrinsicTriangulationMesh.h
@@ -0,0 +1,585 @@
+// Copyright Epic Games, Inc. All Rights Reserved.
+
+#pragma once
+
+#include "CoreMinimal.h"
+#include "DynamicMesh/DynamicMesh3.h"
+#include "DynamicMesh/InfoTypes.h"
+#include "Util/DynamicVector.h"
+#include "VectorTypes.h"
+#include "MathUtil.h"
+
+
+namespace UE
+{
+namespace Geometry
+{
+/**
+* FIntrinsicEdgeFlipMesh supports edge flips, but no other operations that change
+* the mesh connectivity, and no operations that affect the vertices.
+*
+* An Intrinsic Mesh can be thought of as a mesh that overlays another mesh, sharing
+* the original surface (and in this case vertex locations), but the edges of the intrinsic mesh
+* (while constrained to be on the mesh surface) need not align with the edges of the original mesh and the
+* lengths of these edges are measured on the surface of the original mesh (not the R3 distance).
+*
+* The FIntrinsicEdgeFlipMesh is designed to work with the function FlipToDelaunay() to make an
+* Intrinsic Delaunay Triangulation (IDT) of the same surface as a given FDynamicMesh,
+* allowing for more robust cotan-laplacians (see LaplacianMatrixAssembly.h)
+*
+* Note: the implementation is a simple triangle-based mesh but there is no restriction that 
+* edge(a, b) is unique (e.g. multiple intrinsic edges may connect the same two vertices with different paths). 
+* In fact this mesh allows triangles and edges with repeated vertices. Such structures arise naturally 
+* with intrinsic triangulation, for example an edge that starts and ends at the
+* same vertex may encircle a mesh.
+*
+* The API of FIntrinsicEdgeFlipMesh is similar to FDynamicMesh3, but only has a minimal subset of methods and some
+* such as ::FlipEdge() and ::GetEdgeOpposingV() have very different implementations
+*/
+class GEOMETRYCORE_API FIntrinsicEdgeFlipMesh
+{
+public:
+	
+	using FEdge = FDynamicMesh3::FEdge;
+	using FEdgeFlipInfo = DynamicMeshInfo::FEdgeFlipInfo;
+	/** InvalidID indicates that a vertex/edge/triangle ID is invalid */
+	constexpr static int InvalidID = IndexConstants::InvalidID;
+
+	/** Constructor does ID-matching deep copy the basic mesh topology (but no attributes or groups) */
+	FIntrinsicEdgeFlipMesh(const FDynamicMesh3& SrcMesh);
+
+	FIntrinsicEdgeFlipMesh() = delete;
+	FIntrinsicEdgeFlipMesh(const FIntrinsicEdgeFlipMesh&) = delete;
+
+	/**
+	* Flip a single edge in the Intrinsic Mesh.
+	* @param EdgeFlipInfo [out] populated on return.
+	*
+	* @return a success or failure info.
+	*/
+	EMeshResult FlipEdge(int32 EID, FEdgeFlipInfo& EdgeFlipInfo);
+
+
+	/** @return the intrinsic edge length of specified edge */
+	double GetEdgeLength(int32 EID) const
+	{
+		checkSlow(IsEdge(EID));
+		return EdgeLengths[EID];
+	}
+
+	/**
+	* @return the intrinsic edge lengths for specified tri
+	* ordered to match the edges returned by ::GetEdges(TID)
+	*/
+	FVector3d GetTriEdgeLengths(const int32 TID) const
+	{
+		checkSlow(IsTriangle(TID));
+		return GetEdgeLengthTriple(GetTriEdges(TID));
+	}
+
+	/**
+	* @return the interior angle of specified triangle corner
+	*
+	* @param TID             - ID of containing triangle
+	* @param IndexOf         - Index (i.e 0, 1, 2) of specified wedge of the triangle
+	*/
+	double GetTriInternalAngleR(const int32 TID, const int32 IndexOf) const
+	{
+		checkSlow(IndexOf == 0 || IndexOf == 1 || IndexOf == 2);
+		return InternalAngles[TID][IndexOf];
+	}
+
+
+	/**
+	* return FVector2( alpha_{ij} , beta_{ij} )
+	* where alpha_{ij} and beta_{ij} are the two interior angles
+	* opposite the specified edge, EID.
+	*
+	* Expected angles are in the [0,2pi) range
+	*
+	* Note: if EID is a mesh boundary beta_{ij} is replaced by -MAX_DOUBLE
+	*/
+	FVector2d EdgeOpposingAngles(int32 EID) const;
+
+
+	/**
+	* return 1/2 ( cot(alpha_{ij}) + cot(beta_{ij}) )
+	* where alpha_{ij} and beta_{ij} are the to interior angles
+	* opposite the specified edge, EID.
+	*
+	* Note: if EID is a mesh boundary, cot(alpha_ij) is returned
+	*/
+	double EdgeCotanWeight(int32 EID) const;
+
+	// -- Minimal mesh interface --//
+
+
+	/** @return enumerable object for valid edge indices suitable for use with range-based for, ie for ( int i : EdgeIndicesItr() ) */
+	FRefCountVector::IndexEnumerable EdgeIndicesItr() const
+	{
+		return EdgeRefCounts.Indices();
+	}
+
+	/** @return enumerable object for one-ring edges of a vertex, suitable for use with range-based for, ie for ( int i : VtxEdgesItr(VertexID) ) */
+	FSmallListSet::ValueEnumerable VtxEdgesItr(int32 VID) const
+	{
+		checkSlow(VertexRefCounts.IsValid(VID));
+		return VertexEdgeLists.Values(VID);
+	}
+
+	/** @return the valence of a vertex (the number of connected edges) */
+	int32 GetVtxEdgeCount(int32 VID) const
+	{
+		return VertexRefCounts.IsValid(VID) ? VertexEdgeLists.GetCount(VID) : -1;
+	}
+
+	/** @return number of triangles in the intrinsic mesh*/
+	int32 TriangleCount() const
+	{
+		return (int32)TriangleRefCounts.GetCount();
+	}
+
+	/** @return number of edges in the intrinsic mesh*/
+	int32 EdgeCount() const
+	{
+		return (int32)EdgeRefCounts.GetCount();
+	}
+
+	/** @return upper bound on the edge ID used in the mesh.  i.e. all edge IDs < MaxEdgeID*/
+	int32 MaxEdgeID() const
+	{
+		return (int32)EdgeRefCounts.GetMaxIndex();
+	}
+
+	/** @return true if intrinsic mesh contains specified edge */
+	inline bool IsEdge(int32 EdgeID) const
+	{
+		return EdgeRefCounts.IsValid(EdgeID);
+	}
+
+	/** @return true if specified edge is on the boundary of intrinsic mesh*/
+	inline bool IsBoundaryEdge(int32 EdgeID) const
+	{
+		checkSlow(IsEdge(EdgeID));
+		return Edges[EdgeID].Tri[1] == InvalidID;
+	}
+
+	/** @return specified edge*/
+	inline FEdge GetEdge(int32 EdgeID) const
+	{
+		checkSlow(IsEdge(EdgeID));
+		return Edges[EdgeID];
+	}
+
+	/** @return the vertex IDs that define the given edge in intrinsic mesh*/
+	inline FIndex2i GetEdgeV(int32 EdgeID) const
+	{
+		checkSlow(IsEdge(EdgeID));
+		return Edges[EdgeID].Vert;
+	}
+
+	/** @return the triangle pair for an edge. The second triangle may be InvalidID */
+	inline FIndex2i GetEdgeT(int32 EdgeID) const
+	{
+		checkSlow(IsEdge(EdgeID));
+		return Edges[EdgeID].Tri;
+	}
+
+	/** @return the vertex pair opposite specified edge. The second vertex may be InvalidID */
+	FIndex2i GetEdgeOpposingV(int32 EID) const;
+
+	/** @return true if intrinsic mesh contains specified vertex */
+	inline bool IsVertex(int32 VertexID) const
+	{
+		return VertexRefCounts.IsValid(VertexID);
+	}
+
+	/** @return the r3 position of the specified vertex */
+	inline FVector3d GetVertex(int32 VertexID) const
+	{
+		checkSlow(IsVertex(VertexID));
+		return Vertices[VertexID];
+	}
+
+	/** @return upper bound on the triangle ID used in the mesh.  i.e. all triangle IDs < MaxTriangleID */
+	int32 MaxTriangleID() const
+	{
+		return (int32)TriangleRefCounts.GetMaxIndex();
+	}
+
+	/** @return true if intrinsic mesh contains specified triangle */
+	inline bool IsTriangle(int32 TriangleID) const
+	{
+		return TriangleRefCounts.IsValid(TriangleID);
+	}
+	
+	/** @return vertex ids for the three corners of the triangle (e.g. a, b, c)*/
+	inline FIndex3i GetTriangle(int32 TriangleID) const
+	{
+		checkSlow(IsTriangle(TriangleID));
+		return Triangles[TriangleID];
+	}
+
+	/** @return edge ids for the three sides of the triangle (e.g. ab, bc, ca)*/
+	inline FIndex3i GetTriEdges(int32 TriangleID) const
+	{
+		checkSlow(IsTriangle(TriangleID));
+		return TriangleEdges[TriangleID];
+	}
+
+
+protected:
+	
+	/** @return three requested edge lengths */
+	FVector3d GetEdgeLengthTriple(const FIndex3i& EIDs) const
+	{
+		return FVector3d(EdgeLengths[EIDs.A], EdgeLengths[EIDs.B], EdgeLengths[EIDs.C]);
+	}
+
+	/** 
+	* cyclic permutation of the src vector such that the 'Index' entry in the src is the first entry in the result
+	* NB: Index must be 0, 1, or 2.
+	*/ 
+	template <typename Vector3Type>
+	Vector3Type Permute(int32 Index, const Vector3Type& Src) const
+	{
+		checkSlow(Index == 0 || Index == 1 || Index == 2);
+		return Vector3Type(Src[Index], Src[AddOneModThree[Index]], Src[AddTwoModThree[Index]]);
+	}
+
+	/**
+	* computes the angles based on the edge lengths, so must have valid EdgeLengths for this tri before calling.
+	* the order of the angles matches the order of the vertices that define the triangle, e.g. if (a, b, c) = GetTriangle()
+	* the (angle at a, angle at b, angle at c) = ComputeTriInternalAnglesR
+	*/  
+	FVector3d ComputeTriInternalAnglesR(const int32 TID) const;
+
+	/**
+	* @return the intrinsic distance between the two vertices opposite the specified edge
+	* NB: Will fail if EID is a boundary edge.
+	*/
+	double GetOpposingVerticesDistance(int32 EID) const;
+
+protected:
+	// basic mesh operations
+
+	/** update the mesh topology (e.g. triangles and edges) by doing an edge flip, but does not update the intrinsic edge lengths or angles.*/
+	EMeshResult FlipEdgeTopology(int32 eab, FEdgeFlipInfo& FlipInfo);
+
+	/** update the existing edge to reference vertices 'a' and 'b'.*/
+	void SetEdgeVerticesInternal(int32 EdgeID, int32 a, int32 b)
+	{
+		if (a > b)
+		{
+			Swap(a, b);
+		}
+		Edges[EdgeID].Vert[0] = a;
+		Edges[EdgeID].Vert[1] = b;
+	}
+	/** update the existing edge to reference triangles t0 and t1*/
+	inline void SetEdgeTrianglesInternal(int32 EdgeID, int32 t0, int32 t1)
+	{
+		Edges[EdgeID].Tri[0] = t0;
+		Edges[EdgeID].Tri[1] = t1;
+	}
+	/** replace tOld with tNew in an existing edge */
+	int32 ReplaceEdgeTriangle(int32 eID, int32 tOld, int32 tNew);
+	/** replace eOld with eNew in and existing triangle */
+	inline int32 ReplaceTriangleEdge(int32 tID, int32 eOld, int32 eNew)
+	{
+		FIndex3i& TriEdgeIDs = TriangleEdges[tID];
+		for (int32 j = 0; j < 3; ++j)
+		{
+			if (TriEdgeIDs[j] == eOld)
+			{
+				TriEdgeIDs[j] = eNew;
+				return j;
+			}
+		}
+		return -1;
+	}
+
+	/**
+	* @return vertex pair in specified edge, oriented by the given triangle.  
+	* note: the vertices may be the same.
+	* note: assumes EID is an edge in Triangle TID
+	*/
+	FIndex2i GetOrientedEdgeV(int32 EID, int32 TID) const;
+protected:
+
+	const int32 AddOneModThree[3] = { 1, 2, 0 };
+	const int32 AddTwoModThree[3] = { 2, 0, 1 };
+
+	// the basic mesh 
+
+	TDynamicVector<FVector3d> Vertices{};   // list of vertex positions
+	FRefCountVector VertexRefCounts{};      // Reference counts of vertex indices. For vertices that exist, the count is 1 + num_triangle_using_vertex. Iterate over this to find out which vertex indices are valid. 
+	FSmallListSet VertexEdgeLists;          // List of per-vertex edge one-rings
+
+	TDynamicVector<FIndex3i> Triangles;     // List of triangle vertex - index triplets[Vert0 Vert1 Vert2]
+	FRefCountVector TriangleRefCounts{};    // Reference counts of triangle indices. Ref count is always 1 if the triangle exists. Iterate over this to find out which triangle indices are valid.
+	TDynamicVector<FIndex3i> TriangleEdges; // List of triangle edge triplets [Edge0 Edge1 Edge2]
+
+	TDynamicVector<FEdge> Edges;            // List of edge elements. An edge is four elements [VertA, VertB, Tri0, Tri1], where VertA < VertB, and Tri1 may be InvalidID (if the edge is a boundary edge)
+	FRefCountVector EdgeRefCounts{};        // Reference counts of edge indices. Ref count is always 1 if the edge exists. Iterate over this to find out which edge indices are valid.
+
+protected:
+	
+	// intrinsic quantities 
+	
+	TDynamicVector<double>            EdgeLengths;           // Length of each intrinsic edge in the Mesh: note these may differ from (Vert0 - Vert1).Lenght().
+	TDynamicVector<FVector3d>         InternalAngles;        // Interior Angles for each triangle: note may differ from FDynamicMesh::GetTriInternalAnglesR(TID) as they are computed from intrinsic edge lenghts
+
+};
+
+
+/**
+* Class that manages the intrinsic triangulation of a given FDynamicMesh.
+*
+* An Intrinsic triangulation can be thought of as a mesh that overlays another mesh, sharing
+* the original surface, but the intrinsic triangles can fold over the edges of the original mesh
+* That is to say, the edges of the intrinsic mesh are still on the mesh surface
+* but need not align with the edges of the original mesh and the lengths of these edges are measured
+* on the surface of the original mesh (not the R3 distance).
+*
+* This supports mesh operations that: SplitEdge, PokeTriangle, FlipEdge.
+*
+* Both SplitEdge and PokeTriangle will generate new vertices and triangles in the
+* intrinsic triangulation.  These vertices will live on the surface of the original mesh
+* and their location (in R3 and relative to the original surface ) is computed by tracing
+* a path on the original mesh.
+*
+* NB: The lifetime of this structure should not exceed that of the original
+* FDynamicMesh as this class holds a pointer to that mesh to reference locations on its surface.
+*/
+class GEOMETRYCORE_API FIntrinsicTriangulation : public FIntrinsicEdgeFlipMesh
+{
+
+public:
+
+	typedef FIntrinsicEdgeFlipMesh MyBase;
+
+	using FEdgeFlipInfo = DynamicMeshInfo::FEdgeFlipInfo;
+	using FEdgeSplitInfo = DynamicMeshInfo::FEdgeSplitInfo;
+	using FPokeTriangleInfo = DynamicMeshInfo::FPokeTriangleInfo;
+
+
+	FIntrinsicTriangulation(const FDynamicMesh3& SrcMesh);
+
+	FIntrinsicTriangulation() = delete;
+	FIntrinsicTriangulation(const FIntrinsicTriangulation&) = delete;
+	
+
+	/** @return pointer the reference mesh, note the FIntrinsicTriangulation does not manage the lifetime of this mesh */
+	const FDynamicMesh3* GetExtrinsicMesh()  const
+	{
+		return ExtrinsicMesh;
+	}
+	
+	// vertex information
+
+
+	struct FSurfacePoint;
+	/**
+	* return FSurfacePoint information used to map an intrinsic vertex to the surface of the original mesh.
+	*/
+	const FIntrinsicTriangulation::FSurfacePoint& GetVertexSurfacePoint(int32 VID) const
+	{
+		checkSlow(IsVertex(VID));
+		return IntrinsicVertexPositions[VID];
+	}
+
+	/**
+	* return position in r3 of an intrinsic vertex.
+	*/
+	FVector3d GetVertexPosition(int32 VID) const
+	{
+		return GetVertex(VID);
+	}
+	
+	/**
+	* Flip a single edge in the Intrinsic Mesh.
+	* @param EdgeFlipInfo [out] populated on return.
+	*
+	* @return a success or failure info.
+	*/
+	EMeshResult FlipEdge(int32 EID, FEdgeFlipInfo& EdgeFlipInfo);
+
+
+	/**
+	* Insert a new vertex inside an intrinsic triangle, ie do a 1 to 3 triangle split
+	* @param TID             - index of triangle to poke
+	* @param BaryCoordinates - barycentric coordinates of poke position
+	* @param PokeInfo        - returned information about new and modified mesh elements
+	* 
+	* @return Ok on success, or enum value indicates why operation cannot be applied. Mesh remains unmodified on error.
+	*/
+	EMeshResult PokeTriangle(int32 TID, const FVector3d& BaryCoordinates, FPokeTriangleInfo& PokeInfo);
+
+
+	/**
+	 * Split an intrinsic edge of the mesh by inserting a vertex. This creates a new triangle on either side of the edge (ie a 2-4 split).
+	 * If the original edge had vertices [a,b], with triangles t0=[a,b,c] and t1=[b,a,d],  then the split inserts new vertex f.
+	 * After the split t0=[a,f,c] and t1=[f,a,d], and we have t2=[f,b,c] and t3=[f,d,b]  (it's best to draw it out on paper...)
+	 * SplitInfo.OriginalTriangles = {t0, t1} and SplitInfo.NewTriangles = {t2, t3}
+	 *
+	 * @param EdgeAB          - index of the edge to be split
+	 * @param SplitInfo       - returned information about new and modified mesh elements
+	 * @param SplitParameterT - defines the position along the edge that we split at, must be between 0 and 1, and is assumed to be based on the order of vertices in t0
+	 * 
+	 * @return Ok on success, or enum value indicates why operation cannot be applied. Mesh remains unmodified on error.
+	 */
+	EMeshResult SplitEdge(int32 EdgeAB, FEdgeSplitInfo& SplitInfo, double SplitParameterT = 0.5);
+
+	
+
+public:
+
+
+	struct FGeometricInfo
+	{
+		double ToRadians = 1.;    // 2pi / sum of angles at this vertex.  Will be 1 if the Gaussian curvature is 0
+		bool bIsInterior = false; // false if on boundary or at a bow-tie.
+	};
+
+	// Location on (extrinsic) mesh stored as a union.
+	struct FSurfacePoint
+	{
+		enum class EPositionType
+		{
+			Vertex,
+			Edge,
+			Triangle
+		};
+
+		// constructors respect the correct position type
+		FSurfacePoint(); // defaults to invalid vertex position
+		FSurfacePoint(const int32 VID);
+		FSurfacePoint(const int32 EdgeID, const double Alpha);
+		FSurfacePoint(const int32 TriID, const FVector3d& BaryCentrics);
+
+		struct FVertexPosition
+		{
+			int32 VID;                      // Reference (extrinsic) vertex
+		};
+		struct FEdgePosition
+		{
+			int32 EdgeID;                  // Reference (extrinsic) edge
+			double Alpha;                  // Lerp coordinate along referenced edge as Pos(Edge.A) (Alpha) + (1-Alpha) Pos(Edge.B)
+		};
+		struct FTrianglePosition
+		{
+			int32 TriID;                    // Reference (extrinsic) triangle
+			double BarycentricCoords[3];    // Barycentric coordinates within referenced triangle
+		};
+		union FSurfacePositionUnion
+		{
+			FVertexPosition   VertexPosition;
+			FEdgePosition     EdgePosition;
+			FTrianglePosition TriPosition;
+		};
+
+		EPositionType PositionType;      // Position type, needed to correctly interpret the Position union
+		FSurfacePositionUnion Position;  // Position relative to vertex, edge, or triangle. 
+	};
+
+protected:
+
+	
+
+	/**
+	* Given edge EID connects UpdateVID and StartVID, this update the SurfacePosition of the vertex UpdateVID by tracing the edge EID
+	* a distance TraceDistance from StartVID.
+	*
+	* @return               - the polar angle of the tracing edge incident to VID ( relative to the local reference edge. )
+	* @param UpdateVID      - ID of intrinsic vertex who's surface position is updated by this trace
+	* @param StartVID       - ID of intrinsic vertex where the trace starts
+	* @param TracePolarDir  - Local direction of the trace. 
+	* @param TraceDistance  - Distance, measured on the extrinsic mesh surface to trace.
+	*
+	*/
+	double UpdateVertexByEdgeTrace(const int32 UpdateVID, const int32 StartVID, const double TracePolarDir, const double TraceDistance);
+
+
+	/**
+	* Updates the mesh connectivity by adding a new vertex at the vertex-averaged location. Note: this does not update any of the intrinsic lengths, positions, angles, etc.
+	* those quantities must be updated after.  See ::PokeTriangle()
+	*/ 
+	EMeshResult PokeTriangleTopology(int32 TID, FPokeTriangleInfo& PokeInfo);
+	
+	/**
+	* Updates the mesh connectivity by adding a new vertex at the vertex-averaged location. Note: this does not update any of the intrinsic lengths, positions, angles, etc.
+	* those quantities must be updated after.  See ::SplitEdge()
+	*/
+	EMeshResult SplitEdgeTopology(int32 EdgeAB, FEdgeSplitInfo& SplitInfo);
+
+	/** allocate, or clear existing edge list for specified vertex*/
+	inline void AllocateEdgesList(int32 VertexID)
+	{
+		if (VertexID < (int)VertexEdgeLists.Size())
+		{
+			VertexEdgeLists.Clear(VertexID);
+		}
+		VertexEdgeLists.AllocateAt(VertexID);
+	}
+
+	/** add new vertex to the mesh topology, does not update any of the intrinsic quantities*/
+	inline int32 AppendVertex(const FVector3d& vPos)
+	{
+		int32 newVertID = VertexRefCounts.Allocate();
+		Vertices.InsertAt(vPos, newVertID);
+		AllocateEdgesList(newVertID);
+		return newVertID;
+	}
+
+	/** add new edge to the mesh topology, does not update any of the intrinsic quantities*/
+	inline int32 AddEdgeInternal(int32 vA, int32 vB, int32 tA, int32 tB = InvalidID)
+	{
+		if (vB < vA) {
+			int32 t = vB; vB = vA; vA = t;
+		}
+		int32 eid = EdgeRefCounts.Allocate();
+		Edges.InsertAt(FEdge{ {vA, vB},{tA, tB} }, eid);
+		VertexEdgeLists.Insert(vA, eid);
+		if (vA != vB)
+		{ 
+			VertexEdgeLists.Insert(vB, eid);
+		}
+
+		return eid;
+	}
+	/** add new triangle to the mesh topology, does not update any of the intrinsic quantities*/
+	inline int32 AddTriangleInternal(int32 a, int32 b, int32 c, int32 e0, int32 e1, int32 e2)
+	{
+		int32 tid = TriangleRefCounts.Allocate();
+		Triangles.InsertAt(FIndex3i(a, b, c), tid);
+		TriangleEdges.InsertAt(FIndex3i(e0, e1, e2), tid);
+		return tid;
+	}
+protected:
+
+
+	const FDynamicMesh3* ExtrinsicMesh;                                   // Original mesh. The triangles of this mesh define the surface.  We don't manage the lifetime of this mesh
+	TArray<int32> VIDToReferenceEID;                                      // Reference (original mesh) Edge per (original mesh) Vertex, defines local zero angle on tangent plane
+	TArray<int32> TIDToReferenceEID;                                      // Reference (original mesh) Edge per (original mesh) Triangle, defines local zero angle on tangent plane
+	
+	TDynamicVector<FVector3d>            IntrinsicEdgeAngles;             // Per-intrinsic Triangle - polar angles of each edge relative to the vertex it exits (i.e. edge 0 relative to vertex 0)
+	TDynamicVector<FSurfacePoint>        IntrinsicVertexPositions;        // Per-intrinsic Vertex - Locates intrinsic mesh vertices relative to the surface defined by extrinsic mesh.
+	TDynamicVector<FGeometricInfo>       GeometricVertexInfo;             // Per-intrinsic Vertex - Simple geometric information about the extrinsic surface in the neighborhood of intrinsic vert 	
+
+};
+
+
+/**
+* Perform edge flips on intrinsic mesh until either the MaxFlipCount is reached or the mesh is fully Delaunay
+*
+* @param IntrinsicMesh - The mesh to be operated on.
+* @param Uncorrected   - contains on return, all the edges that could not be flipped (e.g. edge shared by non-convex pair of triangles)
+* 
+* @return  the number of flips.
+*/
+int32 GEOMETRYCORE_API FlipToDelaunay(FIntrinsicEdgeFlipMesh& IntrinsicMesh, TSet<int>& Uncorrected, const int32 MaxFlipCount = TMathUtilConstants<int>::MaxReal);
+int32 GEOMETRYCORE_API FlipToDelaunay(FIntrinsicTriangulation& IntrinsicMesh, TSet<int>& Uncorrected, const int32 MaxFlipCount = TMathUtilConstants<int>::MaxReal);
+
+}; // end namespace Geometry
+}; // end namespace UE
diff --git a/Engine/Source/Runtime/GeometryCore/Public/Spatial/MeshGeodesicSurfaceTracer.h b/Engine/Source/Runtime/GeometryCore/Public/Spatial/MeshGeodesicSurfaceTracer.h
new file mode 100644
--- /dev/null
+++ b/Engine/Source/Runtime/GeometryCore/Public/Spatial/MeshGeodesicSurfaceTracer.h
@@ -0,0 +1,384 @@
+// Copyright Epic Games, Inc. All Rights Reserved.
+
+#pragma once
+
+#include "CoreMinimal.h"
+#include "HAL/Platform.h"
+#include "IndexTypes.h"
+#include "VectorTypes.h"
+#include "MathUtil.h" // for max real
+#include "TriangleTypes.h"
+
+
+
+namespace UE {
+namespace Geometry {
+
+class FDynamicMesh3;
+
+
+/**
+* Utils to define vectors relative to the local surface of a mesh 
+* and given such a vector, functions to trace across the face of a single triangle within a mesh
+* to find where the vector intersects a triangle edge.
+*/
+namespace GeodesicSingleTriangleUtils
+{
+
+	/**
+	* Normalized vector defined relative to specified Mesh edge
+	* and only valid in the two faces adjacent to the edge.
+	*/
+	struct FMeshSurfaceDirection;
+
+	/**
+	* Direction in mesh tangent space, defined relative to a vertex and a specified Mesh edge
+	*/
+	struct FMeshTangentDirection;
+	/**
+	* Classification for a trace result (i.e. for a points recorded along the path).
+	*/
+	enum class ETraceClassification
+	{
+		Start,                  // corresponds to initial location when starting a trace 
+		Continuing,             // non-terminated result computed when crossing a single triangle
+		DistanceTerminated,     // result terminated by reaching prescribed max distance
+		BoundaryTerminated,     // result terminated by reaching mesh boundary
+		Failed                  // failed trace, generally a sign that you started a trace with a ray origin outside the triangle.
+	};
+
+	/**
+	* Data returned from a trace across a triangle face.
+	* provides distance traveled, and a new Tangent Vector and accompanying point on the mesh surface.
+	*/
+	struct FTraceResult;
+
+	/**
+	* Triangle in R2, constructed by rotating and aligning a Triangle in R3.
+	* Traces carried out across these triangles.
+	*/
+	class FTangentTri2;
+
+
+	static bool IsTerminated(const FTraceResult& Result);
+	// ----------------------------------------------------------------------------------------
+	//  functions for tracing across a single triangle on a mesh
+	//-----------------------------------------------------------------------------------------
+
+
+	/**
+	* Trace FTangentTri2 using a local coordinate system defined by aligning the first edge of the triangle with the positive x-direction
+	* and the opposing vertex with positive y-coordinate.
+	* @param TangentTri2    - Triangle to trace, expressed in 2d
+	* @param RayOrigin - ray origin
+	* @param Direction - direction of trace, normalized
+	* @param MaxDist   - Max distance the path is allowed to travel for this triangle trace.
+	*
+	* @return Trace Result.
+	*/
+	FTraceResult TraceTangentTriangle(const FTangentTri2& TangentTri2, const FVector2d& RayOrigin, const FVector2d& RayDir, double MaxDistance);
+
+
+	/**
+	* Trace From an Edge crossing into the triangle TriID
+	* @param TriID     - FDynamicMesh3 TriangleID that indicated that triangle we trace across
+	* @param BaryPoint - Barycentric coordinates of ray origin.
+	* @param Direction - Surface direction of trace
+	* @param MaxDist   - Max distance the path is allowed to travel for this triangle trace.
+	*
+	* @return Trace Result.
+	*
+	* NB: it is assumed the edge EID is part of the triangle TriID.  ( otherwise the result will be unexpected)
+	* NB: it is assumed the BaryPoint is a valid point in the triangle ( otherwise the results will be unexpected)
+	*/
+	FTraceResult TraceTriangleFromBaryPoint(const FDynamicMesh3& Mesh, const int32 TriID, const FVector3d& BaryPoint, const FMeshSurfaceDirection& Direction, double MaxDistance = TMathUtilConstants<double>::MaxReal);
+	FTraceResult TraceTriangleFromBaryPoint(const FDynamicMesh3& Mesh, const int32 TriID, const FVector3d& BaryPoint, const FMeshSurfaceDirection& Direction, FTangentTri2& ScratchTangentTri2, double MaxDistance = TMathUtilConstants<double>::MaxReal);
+
+	/**
+	* Trace from  a Vertex Crossing.
+	* @param TangentDirection   - Direction in the local tangent space ( also defines the frame by center VID and zero angle EID)
+	* @param MaxDistance        - Maximal distance this trace can travel.
+	*
+	* @return Trace Result.
+	*/
+	FTraceResult TraceFromVertex(const FDynamicMesh3& Mesh, const FMeshTangentDirection& Direction, double MaxDistance = TMathUtilConstants<double>::MaxReal);
+	FTraceResult TraceFromVertex(const FDynamicMesh3& Mesh, const FMeshTangentDirection& Direction, FTangentTri2& ScratchTangentTri2, double MaxDistance = TMathUtilConstants<double>::MaxReal);
+
+	/**
+	* Trace from  a Vertex Crossing.
+	* @param VID         - Vertex ID (per the FDynamicMesh3) where the trace starts
+	* @param Direction   - Surface direction of trace.
+	* @param MaxDistance - Maximal distance this trace can travel.
+	*
+	* @return Trace Result.
+	*
+	* NB: the direction vector is defined on the surface, not the tangent plane.   This means, it must lie in a triangle adjacent to Direction.EdgeID
+	*/
+	FTraceResult TraceFromVertex(const FDynamicMesh3& Mesh, int32 VID, const FMeshSurfaceDirection& Direction, double MaxDistance = TMathUtilConstants<double>::MaxReal);
+	FTraceResult TraceFromVertex(const FDynamicMesh3& Mesh, int32 VID, const FMeshSurfaceDirection& Direction, FTangentTri2& ScratchTangentTri2, double MaxDistance = TMathUtilConstants<double>::MaxReal);
+
+	/**
+	* Trace From an Edge crossing into the triangle TriID
+	* @param TriID     - FDynamicMesh3 TriangleID that indicated that triangle we trace across
+	* @param EdgeAlpha - [0,1] parameter that determines the location on the Direction edge where the trace originates (measured from Edge.A to Edge.B as ordered in mesh)
+	* @param Direction - Surface direction of trace
+	*
+	* @return Trace Result.
+	*
+	* NB: it is assumed that the edge EID is part of the triangle TriID.  (otherwise the result will be unexpected)
+	*/
+	FTraceResult TraceTriangleFromEdge(const FDynamicMesh3& Mesh, const int32 TriID, double EdgeAlpha, const FMeshSurfaceDirection& Direction, double MaxDistance = TMathUtilConstants<double>::MaxReal);
+	FTraceResult TraceTriangleFromEdge(const FDynamicMesh3& Mesh, const int32 TriID, double EdgeAlpha, const FMeshSurfaceDirection& Direction, FTangentTri2& ScratchTangentTri2, double MaxDistance = TMathUtilConstants<double>::MaxReal);
+
+
+	/**
+	* Trace, starting with the result of an earlier trace.
+	* @param LastTrace   - Struct that hold previous trace.  Note StartTrace.TriID will be the triangle last visited. Will trace across the other edge adjacent tri
+	* @param MaxDistance - Maximal distance this trace can travel.
+	*
+	* @return Trace Result.
+	*/
+	FTraceResult TraceNextTriangle(const FDynamicMesh3& Mesh, const  FTraceResult& LastTrace, double MaxDistance = TMathUtilConstants<double>::MaxReal);
+	FTraceResult TraceNextTriangle(const FDynamicMesh3& Mesh, const  FTraceResult& LastTrace, FTangentTri2& ScratchTangentTri2, double MaxDistance = TMathUtilConstants<double>::MaxReal);
+
+}; // end namespace GeodesicSingleTriangleUtils
+
+class GeodesicSingleTriangleUtils::FTangentTri2
+{
+public:
+	FTangentTri2() = default;
+
+	FTangentTri2(const FTangentTri2& other) = default;
+
+	// The PrimaryEdge will be aligned with the x-axis in 2d ( and must be one of the edges of the triangle indicated by TriID)
+	FTangentTri2(const FDynamicMesh3& Mesh, int32 TriID, int32 PrimaryEdgeID);
+
+	// Reset just invalidates Tri3ID  
+	void Reset();
+
+	// EdgeIndex = 0, 1, 2
+	double DistanceToEdge(int EdgeIndex, const FVector2d& Point) const;
+
+	// Closest Point on the edge, given in terms of 0,1 parameter measured along the edge (EdgeIndex = 0, 1, 2)
+	double ProjectOnEdge(int EdgeIndex, const FVector2d& Point) const;
+
+	// Express the vector in terms of a basis that is aligned with the EdgeIndex edge of the triangle
+	FVector2d ChangeBasis(const FVector2d& Vec, int32 EdgeIndex) const;
+
+	//Convert vertex data ordered for this FTangentTri2 into vertex data ordered for the source 3d triangle
+	template <typename Vector3Type>
+	Vector3Type ReorderVertexDataForSrcTri3d(const Vector3Type& DataOrderedForTri2d) const;
+
+	// Barycentric coords for this point in the tri2d
+	FVector3d GetBarycentricCoords(const FVector2d& PoinInTri2) const;
+
+	// Convert a point on this tri, to Barycentric coords in the un-rotated source triangle.
+	FVector3d GetBarycentricCoordsInSrcTri(const FVector2d& PoinInTri2) const;
+
+public:
+
+	int32 Tri3ID = IndexConstants::InvalidID;  // TriID in the DynamicMesh3 that defined the triangle in R3
+
+	FIndex3i PermutedTriVIDs;      // Map verts to original triangle
+	FIndex3i PermutedTriEIDs;      // Map edges to original triangle
+	FIndex3i PermutedIndices;      // Maps R2 triangle index order to original R3 triangle index order
+
+	FIndex3i EdgeOrientationSign;  // Per-edge information: +1 if Mesh.GetEdgeV() orientation agrees with ccw traversal of this triangle. -1 if reversed
+	FTriangle2d Tri2d;             // The triangle in R2
+
+	FVector2d Tri2dEdges[3];         // cached R2 triangle edge vectors
+	FVector3d Tri2dLengths;          // cached R2 triangle edge lengths
+};
+
+/**
+* Class that, given a starting location and direction, traces a geodesic path along a mesh.
+* As the path enters a triangle, it is unfolded a two-dimensional triangle and the path is transformed into 
+* the local frame of the triangle.
+* 
+* At each point where the path crosses into a new triangle, this class records the point and direction 
+* in terms of barycentric coordinates, triangle references and local path direction relative 
+* to a local frame (in the form of one of the edges of the triangle).
+*/
+class FMeshGeodesicSurfaceTracer
+{
+public:
+
+	using FMeshSurfaceDirection = GeodesicSingleTriangleUtils::FMeshSurfaceDirection;
+	using FMeshTangentDirection = GeodesicSingleTriangleUtils::FMeshTangentDirection;
+	using ETraceClassification  = GeodesicSingleTriangleUtils::ETraceClassification;
+	using FTraceResult          = GeodesicSingleTriangleUtils::FTraceResult;
+	using FTangentTri2          = GeodesicSingleTriangleUtils::FTangentTri2;
+
+
+	FMeshGeodesicSurfaceTracer() = default; // will need to associate a mesh to actually use..
+
+	FMeshGeodesicSurfaceTracer(const FDynamicMesh3& MeshIn)
+		: Mesh(&MeshIn)
+	{}
+
+	void Reset(const FDynamicMesh3& MeshIn)
+	{
+		SurfaceTrace.Reset();
+		ScratchTri2.Reset();
+		Mesh = &MeshIn;
+	}
+
+	/**
+	* Populate internal SurfaceTrace array of FTraceResults while following a "straight" path on the surface of a mesh.
+	* The starting location and direction are given in terms of a local parameterization of the initial triangle.
+	*
+	* @param TriID     - FDynamicMesh3 TriangleID that indicated that triangle where the trace originates
+	* @param BaryPoint - Barycentric coordinates of ray origin relative to the specified triangle. See NB below.
+	* @param RayDir3   - Initial Direction given in R3.  This will be projected on the face of the specified Triangle (TriID)
+	* @param MaxDist   - Max distance the path is allowed to travel for this triangle trace.
+	*
+	* @return total distance of the trace, this may be less than the MaxDist (e.g. in the event the trace encounters a mesh edge)
+	*
+	* NB: the BaryPoint must be a valid point within or on the boundary of the triangle.
+	*/
+	double TraceMeshFromBaryPoint(const int32 TriID, const FVector3d& BaryPoint, const FVector3d& RayDir3, double MaxDistance = TMathUtilConstants<double>::MaxReal);
+
+	/**
+	* Populate internal SurfaceTrace array of FTraceResults while following a "straight" path on the surface of a mesh.
+	* The starting location and direction are given in terms of a local parameterization of the initial triangle.
+	*
+	* @param TriID     - FDynamicMesh3 TriangleID that indicated that triangle where the trace originates
+	* @param BaryPoint - Barycentric coordinates of ray origin relative to the specified triangle. See NB below.
+	* @param Direction - Surface direction of trace
+	* @param MaxDist   - Max distance the path is allowed to travel for this triangle trace.
+	*
+	* @return total distance of the trace, this may be less than the MaxDist (e.g. in the event the trace encounters a mesh edge)
+	*
+	* NB: the BaryPoint must be a valid point within or on the boundary of the triangle.
+	* NB: it is assumed the edge Direction.EdgeID is part of the triangle TriID.  ( otherwise the result will be unexpected)
+	*/
+	double TraceMeshFromBaryPoint(const int32 TriID, const FVector3d& BaryPoint, const FMeshSurfaceDirection& Direction, double MaxDistance = TMathUtilConstants<double>::MaxReal);
+
+	/**
+	* Populate internal SurfaceTrace array of FTraceResults while following a "straight" path on the surface of a mesh.
+	* The starting location and direction are given in terms of a local parameterization of the initial triangle.
+	*
+	* @param Direction - Tangent plane direction of trace, defines the origin (VID) and the zero radian direction (EID.A to EID.B) along with a polar angle
+	* @param MaxDist   - Max distance the path is allowed to travel for this triangle trace.
+	*
+	* @return total distance of the trace, this may be less than the MaxDist (e.g. in the event the trace encounters a mesh edge)
+	*
+	* NB: it is assumed the edge Direction.EdgeID is adjacent to VID.  ( otherwise the result will be unexpected)
+	*/
+	double TraceMeshFromVertex( const FMeshTangentDirection& Direction, double MaxDistance);
+
+	/** Access to the last tri traced in 2d form.*/
+	const FTangentTri2& GetLastTri()
+	{
+		return ScratchTri2;
+	}
+
+	/** Access to the trace results*/
+	TArray<FTraceResult>& GetTraceResults()
+	{
+		return SurfaceTrace;
+	}
+	/** Access to the trace results*/
+	const TArray<FTraceResult>& GetTraceResults() const
+	{
+		return SurfaceTrace;
+	}
+
+	
+protected:
+
+	const FDynamicMesh3* Mesh = nullptr;       // Pointer to the Mesh to be traced.
+	TArray<FTraceResult> SurfaceTrace;         // Vector of points along the computed trace
+	FTangentTri2 ScratchTri2;                  // 2d representation of the last triangle traversed.  Used as scratch during computations
+
+};
+
+
+
+struct GeodesicSingleTriangleUtils::FMeshSurfaceDirection
+{
+	FMeshSurfaceDirection() :EdgeID(IndexConstants::InvalidID), Dir(0., 0.) {}
+
+	FMeshSurfaceDirection(int EID, const FVector2d& DirIn)
+		: EdgeID(EID)
+		, Dir(DirIn)
+	{}
+
+	FMeshSurfaceDirection(int EID, const double PolarAngle)
+		: EdgeID(EID)
+		, Dir(FVector2d(FMath::Cos(PolarAngle), FMath::Sin(PolarAngle)))
+	{}
+
+	FMeshSurfaceDirection(const FMeshSurfaceDirection& Other) = default;
+	
+
+	int32 EdgeID;           // Mesh Edge, has implicit order (from Edge.A to Edge.B)
+	FVector2d  Dir;         // Normalized direction relative to the frame defined by the mesh edge EdgeID aligned with the x-axis
+};
+
+struct GeodesicSingleTriangleUtils::FMeshTangentDirection
+{
+	int32 VID;           // Tangent space centered on VID vertex
+	int32 EdgeID;        // Mesh Edge (adjacent to VID), the vector Edge.A to Edge.B defines the zero radians direction (Edge.B to Edge.A is pi radians)
+	double PolarAngle;   // Angle measured between [0, 2pi)  
+};
+
+struct GeodesicSingleTriangleUtils::FTraceResult
+{
+	ETraceClassification Classification = ETraceClassification::Start;
+
+	double TraceDist = -1.;                    // Distance traveled in this tri before either meeting and edge or the max distance.
+	int32 TriID = IndexConstants::InvalidID;    // Triangle traversed
+	FMeshSurfaceDirection SurfaceDirection;    // Tangent vector.
+
+	// -- Information about the surface point --//
+	FVector3d Barycentric;                     // Surface point as Barycentric coords.  Can be reconstructed with FDynamicMesh::GetTriBaryPoint(TriID, Barycentric[0], Barycentric[1], Barycentric[2]);
+	bool bIsEdgePoint;                         // Duplicate description of the surface point in the case of edge termination
+	int32 EdgeID = IndexConstants::InvalidID;  // In the case of Edge termination, SurfaceDirection.EdgeID = EdgeID
+	double EdgeAlpha = -1.;                    // When edge point, locates position between the vertices lerping parameter, e.g. Alpha Edge.A + ( 1-Alpha) Edge.B
+};
+
+bool GeodesicSingleTriangleUtils::IsTerminated(const FTraceResult& Result)
+{
+	return (Result.Classification == ETraceClassification::DistanceTerminated
+		|| Result.Classification == ETraceClassification::BoundaryTerminated);
+}
+
+
+
+template <typename Vector3Type>
+Vector3Type GeodesicSingleTriangleUtils::FTangentTri2::ReorderVertexDataForSrcTri3d(const Vector3Type& DataOrderdForTri2d) const
+{
+	Vector3Type Result;
+	Result[PermutedIndices[0]] = DataOrderdForTri2d[0];
+	Result[PermutedIndices[1]] = DataOrderdForTri2d[1];
+	Result[PermutedIndices[2]] = DataOrderdForTri2d[2];
+
+	return Result;
+}
+
+
+/**
+* return AngleR in [0,2pi) range.  
+*
+* The same polar point (r, AngleR) can be described by an infinite number of alternative 
+* angles, each offset by some multiple of 2pi,  i.e. (r, AngleR + n * 2Pi).
+* This selects the unique angle that falls in the [0,2pi) range.
+*/
+static double AsZeroToTwoPi(double AngleR)
+{
+	while (AngleR < 0.)
+	{
+		AngleR += TMathUtilConstants<double>::TwoPi;
+	}
+
+	while (AngleR > TMathUtilConstants<double>::TwoPi)
+	{
+		AngleR -= TMathUtilConstants<double>::TwoPi;
+	}
+	return AngleR;
+}
+
+}; // end namespace Geometry
+}; // end namespace UE
+
