IntrinsicTriangulationMesh.cpp  Adding missing include  - fixing non unity build error

diff --git a/Engine/Source/Runtime/GeometryCore/Private/Spatial/IntrinsicTriangulationMesh.cpp b/Engine/Source/Runtime/GeometryCore/Private/Spatial/IntrinsicTriangulationMesh.cpp
--- a/Engine/Source/Runtime/GeometryCore/Private/Spatial/IntrinsicTriangulationMesh.cpp
+++ b/Engine/Source/Runtime/GeometryCore/Private/Spatial/IntrinsicTriangulationMesh.cpp
@@ -6,6 +6,7 @@
 #include "Util/IndexUtil.h"
 #include "Algo/Reverse.h"
 #include "Containers/BitArray.h"
+#include "Containers/Queue.h"
 #include "MathUtil.h"
 #include "Async/ParallelFor.h"
 
