GeometryProcessing: force UV/normal overlays to be initialized on creation inside a DynamicMeshAttributeSet.
#rnx

diff --git a/Engine/Source/Runtime/GeometryCore/Private/DynamicMesh/DynamicMeshAttributeSet.cpp b/Engine/Source/Runtime/GeometryCore/Private/DynamicMesh/DynamicMeshAttributeSet.cpp
--- a/Engine/Source/Runtime/GeometryCore/Private/DynamicMesh/DynamicMeshAttributeSet.cpp
+++ b/Engine/Source/Runtime/GeometryCore/Private/DynamicMesh/DynamicMeshAttributeSet.cpp
@@ -362,7 +362,9 @@ void FDynamicMeshAttributeSet::SetNumUVLayers(int Num)
 	{
 		for (int i = (int)UVLayers.Num(); i < Num; ++i)
 		{
-			UVLayers.Add(new FDynamicMeshUVOverlay(ParentMesh));
+			FDynamicMeshUVOverlay* NewUVLayer = new FDynamicMeshUVOverlay(ParentMesh);
+			NewUVLayer->InitializeTriangles(ParentMesh->MaxTriangleID());
+			UVLayers.Add(NewUVLayer);
 		}
 	}
 	else
@@ -395,7 +397,9 @@ void FDynamicMeshAttributeSet::SetNumNormalLayers(int Num)
 	{
 		for (int32 i = NormalLayers.Num(); i < Num; ++i)
 		{
-			NormalLayers.Add(new FDynamicMeshNormalOverlay(ParentMesh));
+			FDynamicMeshNormalOverlay* NewNormalLayer = new FDynamicMeshNormalOverlay(ParentMesh);
+			NewNormalLayer->InitializeTriangles(ParentMesh->MaxTriangleID());
+			NormalLayers.Add(NewNormalLayer);
 		}
 	}
 	else
@@ -410,6 +414,7 @@ void FDynamicMeshAttributeSet::EnablePrimaryColors()
 	if (HasPrimaryColors() == false)
 	{ 
 		ColorLayer = MakeUnique<FDynamicMeshColorOverlay>(ParentMesh);
+		ColorLayer->InitializeTriangles(ParentMesh->MaxTriangleID());
 	}
 }
 
