GeometryTypes.h -  add missing includes to fix non-unity build errors

diff --git a/Engine/Source/Runtime/GeometryCore/Public/GeometryTypes.h b/Engine/Source/Runtime/GeometryCore/Public/GeometryTypes.h
--- a/Engine/Source/Runtime/GeometryCore/Public/GeometryTypes.h
+++ b/Engine/Source/Runtime/GeometryCore/Public/GeometryTypes.h
@@ -2,6 +2,9 @@
 
 #pragma once
 
+#include "Containers/Map.h"
+#include "Misc/AssertionMacros.h"
+
 namespace UE
 {
 namespace Geometry
