GeometryProcessing: small fixes.
- handle undefined normal attributes in Simplifier ComputeFaceQuadric()
- Add separate SetGridCellSize() / SetMeshCellSize() functions to TImplicitMorphology
- define explicit enum values in RemoveOccludedTriangles
- Add support for linear scaling along path in FGeneralizedCylinderGenerator
#rnx

diff --git a/Engine/Plugins/Runtime/GeometryProcessing/Source/DynamicMesh/Private/MeshSimplification.cpp b/Engine/Plugins/Runtime/GeometryProcessing/Source/DynamicMesh/Private/MeshSimplification.cpp
--- a/Engine/Plugins/Runtime/GeometryProcessing/Source/DynamicMesh/Private/MeshSimplification.cpp
+++ b/Engine/Plugins/Runtime/GeometryProcessing/Source/DynamicMesh/Private/MeshSimplification.cpp
@@ -31,7 +31,15 @@ FAttrBasedQuadricErrord TMeshSimplification<FAttrBasedQuadricErrord>::ComputeFac
 
 	if (NormalOverlay != nullptr)
 	{
-		NormalOverlay->GetTriElements(tid, n0, n1, n2);
+		if (NormalOverlay->IsSetTriangle(tid))
+		{
+			NormalOverlay->GetTriElements(tid, n0, n1, n2);
+		}
+		else
+		{
+			FVector3f FaceNormal = (FVector3f)Mesh->GetTriNormal(tid);
+			n0 = n1 = n2 = FaceNormal;
+		}
 	}
 	else
 	{
diff --git a/Engine/Plugins/Runtime/GeometryProcessing/Source/DynamicMesh/Public/Operations/RemoveOccludedTriangles.h b/Engine/Plugins/Runtime/GeometryProcessing/Source/DynamicMesh/Public/Operations/RemoveOccludedTriangles.h
--- a/Engine/Plugins/Runtime/GeometryProcessing/Source/DynamicMesh/Public/Operations/RemoveOccludedTriangles.h
+++ b/Engine/Plugins/Runtime/GeometryProcessing/Source/DynamicMesh/Public/Operations/RemoveOccludedTriangles.h
@@ -29,15 +29,15 @@ namespace Geometry
 
 enum class EOcclusionTriangleSampling
 {
-	Centroids,
-	Vertices,
-	VerticesAndCentroids
+	Centroids = 0,
+	Vertices = 1,
+	VerticesAndCentroids = 2
 };
 
 enum class EOcclusionCalculationMode
 {
-	FastWindingNumber,
-	SimpleOcclusionTest
+	FastWindingNumber = 0,
+	SimpleOcclusionTest = 1
 };
 
 
diff --git a/Engine/Source/Runtime/GeometryCore/Public/Generators/SweepGenerator.h b/Engine/Source/Runtime/GeometryCore/Public/Generators/SweepGenerator.h
--- a/Engine/Source/Runtime/GeometryCore/Public/Generators/SweepGenerator.h
+++ b/Engine/Source/Runtime/GeometryCore/Public/Generators/SweepGenerator.h
@@ -626,6 +626,7 @@ public:
 };
 
 
+
 /**
  * Sweep a 2D Profile Polygon along a 3D Path.
  * 
@@ -645,6 +646,9 @@ public:
 	bool bLoop = false;
 	ECapType CapType = ECapType::FlatTriangulation;
 
+	double StartScale = 1.0;
+	double EndScale = 1.0;
+
 	// When true, the generator attempts to scale UV's in a way that preserves scaling across different mesh
 	// results, aiming for 1.0 in UV space to be equal to UnitUVInWorldCoordinates in world space. This in
 	// practice means adjusting the U scale relative to the CrossSection curve length and V scale relative
@@ -668,15 +672,18 @@ public:
 		}
 		int PathNum = Path.Num();
 		
+		bool bApplyScaling = (StartScale != 1.0 || EndScale != 1.0) && (bLoop == false);
+		bool bNeedArcLength = (bApplyScaling || bUVScaleRelativeWorld);
+		double TotalPathArcLength = (bNeedArcLength) ? UE::Geometry::CurveUtil::ArcLength<double, FVector3d>(Path, bLoop) : 1.0;
+
 		FAxisAlignedBox2f Bounds = (FAxisAlignedBox2f)CrossSection.Bounds();
 		double BoundsMaxDimInv = 1.0 / FMathd::Max(Bounds.MaxDim(), .001);
 		FVector2f SectionScale(1, 1), CapScale(BoundsMaxDimInv, BoundsMaxDimInv);
 		if (bUVScaleRelativeWorld)
 		{
 			double Perimeter = CrossSection.Perimeter();
-			double PathLen = UE::Geometry::CurveUtil::ArcLength<double, FVector3d>(Path, bLoop);
 			SectionScale.X = Perimeter / UnitUVInWorldCoordinates;
-			SectionScale.Y = PathLen / UnitUVInWorldCoordinates;
+			SectionScale.Y = TotalPathArcLength / UnitUVInWorldCoordinates;
 			CapScale.X = CapScale.Y = 1.0f / UnitUVInWorldCoordinates;
 		}
 		ConstructMeshTopology(CrossSection, {}, {}, {}, true, Path, PathNum + (bLoop ? 1 : 0), bLoop, Caps, SectionScale, CapScale, Bounds.Center());
@@ -688,6 +695,7 @@ public:
 			XNormals[Idx] = CrossSection.GetNormal_FaceAvg(Idx);
 		}
 
+		double AccumArcLength = 0;
 		FFrame3d CrossSectionFrame = InitialFrame;
 		for (int PathIdx = 0; PathIdx < PathNum; ++PathIdx)
 		{
@@ -699,10 +707,21 @@ public:
 			for (int SubIdx = 0; SubIdx < XNum; SubIdx++)
 			{
 				FVector2d XP = CrossSection[SubIdx];
+				if (bApplyScaling)
+				{
+					double T = FMathd::Clamp((AccumArcLength / TotalPathArcLength), 0.0, 1.0);
+					XP *= FMathd::Lerp(StartScale, EndScale, T);
+				}
+
 				FVector2d XN = XNormals[SubIdx];
 				Vertices[SubIdx + PathIdx * XNum] = C + X * XP.X + Y * XP.Y;
 				Normals[SubIdx + PathIdx * XNum] = (FVector3f)(X * XN.X + Y * XN.Y);
 			}
+
+			if (PathIdx < PathNum - 1)
+			{
+				AccumArcLength += Distance(C, Path[PathIdx + 1]);
+			}
 		}
 		if (bCapped && !bLoop)
 		{
diff --git a/Engine/Source/Runtime/GeometryCore/Public/Implicit/Morphology.h b/Engine/Source/Runtime/GeometryCore/Public/Implicit/Morphology.h
--- a/Engine/Source/Runtime/GeometryCore/Public/Implicit/Morphology.h
+++ b/Engine/Source/Runtime/GeometryCore/Public/Implicit/Morphology.h
@@ -64,8 +64,20 @@ public:
 	void SetCellSizesAndDistance(FAxisAlignedBox3d Bounds, double DistanceIn, int TargetInputVoxelCount, int TargetOutputVoxelCount)
 	{
 		Distance = DistanceIn;
+		SetGridCellSize(Bounds, DistanceIn, TargetInputVoxelCount);
+		SetMeshCellSize(Bounds, DistanceIn, TargetOutputVoxelCount);
+	}
+
+	// Set input grid cell size to hit the target voxel counts along the max dimension of the bounds
+	void SetGridCellSize(FAxisAlignedBox3d Bounds, double DistanceIn, int TargetInputVoxelCount)
+	{
 		GridCellSize = (Bounds.MaxDim() + Distance * 2.0) / double(TargetInputVoxelCount);
-		MeshCellSize = (Bounds.MaxDim() + Distance * 2.0) / double(TargetOutputVoxelCount);
+	}
+
+	// Set output meshing cell size to hit the target voxel counts along the max dimension of the bounds
+	void SetMeshCellSize(FAxisAlignedBox3d Bounds, double DistanceIn, int TargetInputVoxelCount)
+	{
+		MeshCellSize = (Bounds.MaxDim() + Distance * 2.0) / double(TargetInputVoxelCount);
 	}
 
 	/** if this function returns true, we should abort calculation */
