GeometryCore: Fix over-allocating triangles in FDynamicMesh3::CompactInPlace()

This CL fixes an issue where we accidentally allocate three times the amount of necessary triangles during `CompactInPlace()`. I also added a test during `CheckValidity()` to make sure that if we have dense ref counts the size of the vertex/triangle vectors is identical to the number of ref counts.

diff --git a/Engine/Source/Runtime/GeometryCore/Private/DynamicMesh/DynamicMesh3.cpp b/Engine/Source/Runtime/GeometryCore/Private/DynamicMesh/DynamicMesh3.cpp
--- a/Engine/Source/Runtime/GeometryCore/Private/DynamicMesh/DynamicMesh3.cpp
+++ b/Engine/Source/Runtime/GeometryCore/Private/DynamicMesh/DynamicMesh3.cpp
@@ -774,6 +774,9 @@ bool FDynamicMesh3::CheckValidity(FValidityOptions ValidityOptions, EValidityChe
 		};
 	}
 
+	// When ref counts are dense, the used ref count must match the size of the vertex/triangle vector. 
+	CheckOrFailF(!VertexRefCounts.IsDense() || (VertexRefCounts.IsDense() && VertexRefCounts.GetCount() == Vertices.Num()));
+	CheckOrFailF(!TriangleRefCounts.IsDense() || (TriangleRefCounts.IsDense() && TriangleRefCounts.GetCount() == Triangles.Num()));
 
 	for (int tID : TriangleIndicesItr())
 	{
diff --git a/Engine/Source/Runtime/GeometryCore/Private/DynamicMesh/DynamicMesh3_Edits.cpp b/Engine/Source/Runtime/GeometryCore/Private/DynamicMesh/DynamicMesh3_Edits.cpp
--- a/Engine/Source/Runtime/GeometryCore/Private/DynamicMesh/DynamicMesh3_Edits.cpp
+++ b/Engine/Source/Runtime/GeometryCore/Private/DynamicMesh/DynamicMesh3_Edits.cpp
@@ -463,7 +463,7 @@ void FDynamicMesh3::CompactInPlace(FCompactMaps* CompactInfo)
 
 	// trim triangles data structures
 	TriangleRefCounts.Trim(TriangleCount());
-	Triangles.Resize(TriangleCount() * 3);
+	Triangles.Resize(TriangleCount());
 	TriangleEdges.Resize(TriangleCount() * 3);
 	if (HasTriangleGroups())
 	{
