GeometryCore: make DynamicMesh::IsSameAs function const
#rnx

diff --git a/Engine/Source/Runtime/GeometryCore/Private/DynamicMesh/DynamicMesh3.cpp b/Engine/Source/Runtime/GeometryCore/Private/DynamicMesh/DynamicMesh3.cpp
--- a/Engine/Source/Runtime/GeometryCore/Private/DynamicMesh/DynamicMesh3.cpp
+++ b/Engine/Source/Runtime/GeometryCore/Private/DynamicMesh/DynamicMesh3.cpp
@@ -605,7 +605,7 @@ FString FDynamicMesh3::MeshInfoString()
 	return VtxString + "\n" + TriString + "\n" + EdgeString + "\n" + AttribString + "\n" + InfoString;
 }
 
-bool FDynamicMesh3::IsSameAs(const FDynamicMesh3& m2, const FSameAsOptions& Options)
+bool FDynamicMesh3::IsSameAs(const FDynamicMesh3& m2, const FSameAsOptions& Options) const
 {
 	if (VertexCount() != m2.VertexCount())
 	{
diff --git a/Engine/Source/Runtime/GeometryCore/Public/DynamicMesh/DynamicMesh3.h b/Engine/Source/Runtime/GeometryCore/Public/DynamicMesh/DynamicMesh3.h
--- a/Engine/Source/Runtime/GeometryCore/Public/DynamicMesh/DynamicMesh3.h
+++ b/Engine/Source/Runtime/GeometryCore/Public/DynamicMesh/DynamicMesh3.h
@@ -1277,7 +1277,7 @@ public:
 	 * Check if another mesh is the same as this mesh. By default only checks
 	 * vertices and triangles, turn on other parameters w/ flags
 	 */
-	virtual bool IsSameAs(const FDynamicMesh3& OtherMesh, const FSameAsOptions& Options);
+	virtual bool IsSameAs(const FDynamicMesh3& OtherMesh, const FSameAsOptions& Options) const;
 
 	/**
 	 * Options for what the validity check will permit
