GeometryCore: Fix TDynamicMeshOverlay<>::CompactInPlace() not compacting ElementTriangles vector

#rnx

diff --git a/Engine/Source/Runtime/GeometryCore/Private/DynamicMesh/DynamicMeshOverlay.cpp b/Engine/Source/Runtime/GeometryCore/Private/DynamicMesh/DynamicMeshOverlay.cpp
--- a/Engine/Source/Runtime/GeometryCore/Private/DynamicMesh/DynamicMeshOverlay.cpp
+++ b/Engine/Source/Runtime/GeometryCore/Private/DynamicMesh/DynamicMeshOverlay.cpp
@@ -1422,6 +1422,12 @@ bool TDynamicMeshOverlay<RealType, ElementSize>::CheckValidity(bool bAllowNonMan
 
 	// @todo: check that all connected element-pairs are also edges in parentmesh
 
+	// Check that the number of parent vertices is consistent with number of elements.
+	CheckOrFailF(ParentVertices.Num() * ElementSize == Elements.Num());
+
+	// Check that the per-triangle data does not exceed the number of triangles in the parent mesh.
+	CheckOrFailF(!ParentMesh || ElementTriangles.Num() <= ParentMesh->MaxTriangleID() * 3);
+
 	// check that parent vtx of a non-isolated element is actually a vertex
 	for (int elemid : ElementIndicesItr())
 	{
diff --git a/Engine/Source/Runtime/GeometryCore/Public/DynamicMesh/DynamicMeshOverlay.h b/Engine/Source/Runtime/GeometryCore/Public/DynamicMesh/DynamicMeshOverlay.h
--- a/Engine/Source/Runtime/GeometryCore/Public/DynamicMesh/DynamicMeshOverlay.h
+++ b/Engine/Source/Runtime/GeometryCore/Public/DynamicMesh/DynamicMeshOverlay.h
@@ -187,19 +187,24 @@ public:
 		Elements.Resize(ElementCount() * ElementSize);
 		ParentVertices.Resize(ElementCount());
 
+		// Remap and compact triangle element indices.
+		int32 MaxNewTID = 0;
 		for (int TID = 0, OldMaxTID = ElementTriangles.Num() / 3; TID < OldMaxTID; TID++)
 		{
-			int OldStart = TID * 3;
+			const int32 OldStart = TID * 3;
 			if (ElementTriangles[OldStart] == -1)
 			{
 				continue; // triangle was not set; skip it
 			}
-			int NewStart = CompactMaps.GetTriangleMapping(TID) * 3;
+			const int32 NewTID = CompactMaps.GetTriangleMapping(TID);
+			MaxNewTID = FMath::Max(MaxNewTID, NewTID);
+			const int32 NewStart = NewTID * 3;
 			for (int SubIdx = 0; SubIdx < 3; SubIdx++)
 			{
 				ElementTriangles[NewStart + SubIdx] = MapE[ElementTriangles[OldStart + SubIdx]];
 			}
 		}
+		ElementTriangles.Resize((MaxNewTID + 1) * 3);
 
 		checkSlow(IsCompact());
 	}
