GeometryCore: add TMeshQueries::MeshDistanceStatistics
#rnx

diff --git a/Engine/Source/Runtime/GeometryCore/Public/MeshQueries.h b/Engine/Source/Runtime/GeometryCore/Public/MeshQueries.h
--- a/Engine/Source/Runtime/GeometryCore/Public/MeshQueries.h
+++ b/Engine/Source/Runtime/GeometryCore/Public/MeshQueries.h
@@ -585,6 +585,47 @@ public:
 		return FMath::Max(HausdorffDistanceSerial(MeshA, SpatialB), HausdorffDistanceSerial(MeshB, SpatialA));
 	}
 
+	/**
+	 * Compute various statistics on distances between two meshes
+	 */
+	template<typename MeshSpatialType>
+	static void MeshDistanceStatistics(
+		const TriangleMeshType& MeshA, 
+		const MeshSpatialType& SpatialB,
+		const TriangleMeshType* MeshB, 
+		const MeshSpatialType* SpatialA,
+		bool bSymmetric,
+		double& MaxDistance,
+		double& MinDistance,
+		double& AverageDistance,
+		double& RootMeanSqrDeviation
+		)
+	{
+		TArray<double> Distances;
+		VertexToSurfaceDistances(MeshA, SpatialB, Distances);
+
+		if (bSymmetric && MeshB != nullptr && SpatialA != nullptr)
+		{
+			TArray<double> Distances2;
+			VertexToSurfaceDistances(*MeshB, *SpatialA, Distances2);
+			Distances.Append(Distances2);
+		}
+
+		double NumDistances = (double)Distances.Num();
+
+		MaxDistance = -BIG_NUMBER;
+		MinDistance = BIG_NUMBER;
+		AverageDistance = 0;
+		for (double& Dist : Distances)
+		{
+			MaxDistance = FMath::Max(Dist, MaxDistance);
+			MinDistance = FMath::Min(Dist, MinDistance);
+			AverageDistance += Dist / NumDistances;
+			RootMeanSqrDeviation += (Dist * Dist);
+		}
+		RootMeanSqrDeviation = FMathd::Sqrt(RootMeanSqrDeviation / NumDistances);
+	}
+
 };
 
 
