Fixes to CompactInPlace: keep ElementTriangles 1:1 with ParentMesh triangles, don't try to copy when there's no mapping, but copy back InvalidID when the (mapped) element tri is unset

#rnx

diff --git a/Engine/Source/Runtime/GeometryCore/Private/DynamicMesh/DynamicMeshOverlay.cpp b/Engine/Source/Runtime/GeometryCore/Private/DynamicMesh/DynamicMeshOverlay.cpp
--- a/Engine/Source/Runtime/GeometryCore/Private/DynamicMesh/DynamicMeshOverlay.cpp
+++ b/Engine/Source/Runtime/GeometryCore/Private/DynamicMesh/DynamicMeshOverlay.cpp
@@ -1425,8 +1425,8 @@ bool TDynamicMeshOverlay<RealType, ElementSize>::CheckValidity(bool bAllowNonMan
 	// Check that the number of parent vertices is consistent with number of elements.
 	CheckOrFailF(ParentVertices.Num() * ElementSize == Elements.Num());
 
-	// Check that the per-triangle data does not exceed the number of triangles in the parent mesh.
-	CheckOrFailF(!ParentMesh || ElementTriangles.Num() <= ParentMesh->MaxTriangleID() * 3);
+	// Check that the per-triangle data matches the number of triangles in the parent mesh.
+	CheckOrFailF(!ParentMesh || ElementTriangles.Num() == ParentMesh->MaxTriangleID() * 3);
 
 	// check that parent vtx of a non-isolated element is actually a vertex
 	for (int elemid : ElementIndicesItr())
diff --git a/Engine/Source/Runtime/GeometryCore/Public/DynamicMesh/DynamicMeshOverlay.h b/Engine/Source/Runtime/GeometryCore/Public/DynamicMesh/DynamicMeshOverlay.h
--- a/Engine/Source/Runtime/GeometryCore/Public/DynamicMesh/DynamicMeshOverlay.h
+++ b/Engine/Source/Runtime/GeometryCore/Public/DynamicMesh/DynamicMeshOverlay.h
@@ -192,18 +192,35 @@ public:
 		for (int TID = 0, OldMaxTID = ElementTriangles.Num() / 3; TID < OldMaxTID; TID++)
 		{
 			const int32 OldStart = TID * 3;
-			if (ElementTriangles[OldStart] == -1)
+			const int32 NewTID = CompactMaps.GetTriangleMapping(TID);
+			if (NewTID == IndexConstants::InvalidID)
 			{
-				continue; // triangle was not set; skip it
+				// skip if there's no mapping
+				continue;
 			}
-			const int32 NewTID = CompactMaps.GetTriangleMapping(TID);
-			MaxNewTID = FMath::Max(MaxNewTID, NewTID);
+
+			MaxNewTID = FMath::Max(NewTID, MaxNewTID);
+
 			const int32 NewStart = NewTID * 3;
-			for (int SubIdx = 0; SubIdx < 3; SubIdx++)
+			if (ElementTriangles[OldStart] == IndexConstants::InvalidID)
+			{
+				// triangle was not set; copy back InvalidID
+				for (int SubIdx = 0; SubIdx < 3; SubIdx++)
+				{
+					ElementTriangles[NewStart + SubIdx] = IndexConstants::InvalidID;
+				}
+			}
+			else
 			{
-				ElementTriangles[NewStart + SubIdx] = MapE[ElementTriangles[OldStart + SubIdx]];
+				for (int SubIdx = 0; SubIdx < 3; SubIdx++)
+				{
+					ElementTriangles[NewStart + SubIdx] = MapE[ElementTriangles[OldStart + SubIdx]];
+				}
 			}
 		}
+		// ElementTriangles should never grow during a compaction, so just resizing is ok
+		// (i.e., we shouldn't need to set InvalidID on any added triangles)
+		checkSlow(ElementTriangles.Num() >= (MaxNewTID + 1) * 3);
 		ElementTriangles.Resize((MaxNewTID + 1) * 3);
 
 		checkSlow(IsCompact());
