GeometryScripting: make RevolvePolygon and Torus support extended revolve options, add SpiralRevolvePolygon
GeometryCore: add RevolveGenerator, generates surface-of-revolution meshes for open and closed profiles, the latter w/ spirals. Port of more general code in CurveSweepOp/RevolveProperties to expose at lower levels (needs further code cleanup)
#rnx

diff --git a/Engine/Plugins/Experimental/GeometryScripting/Source/GeometryScriptingCore/Private/MeshPrimitiveFunctions.cpp b/Engine/Plugins/Experimental/GeometryScripting/Source/GeometryScriptingCore/Private/MeshPrimitiveFunctions.cpp
--- a/Engine/Plugins/Experimental/GeometryScripting/Source/GeometryScriptingCore/Private/MeshPrimitiveFunctions.cpp
+++ b/Engine/Plugins/Experimental/GeometryScripting/Source/GeometryScriptingCore/Private/MeshPrimitiveFunctions.cpp
@@ -16,11 +16,9 @@
 #include "Generators/DiscMeshGenerator.h"
 #include "Generators/SweepGenerator.h"
 #include "Generators/FlatTriangulationMeshGenerator.h"
+#include "Generators/RevolveGenerator.h"
 #include "ConstrainedDelaunay2.h"
 #include "Arrangement2d.h"
-#include "Util/RevolveUtil.h"
-
-#include "CompositionOps/CurveSweepOp.h"
 
 #include "ExplicitUseGeometryMathTypes.h"		// using UE::Geometry::(math types)
 using namespace UE::Geometry;
@@ -88,6 +86,64 @@ static void AppendPrimitive(
 
 
 
+static void AppendPrimitiveMesh(
+	UDynamicMesh* TargetMesh, 
+	FDynamicMesh3& AppendMesh,
+	FTransform Transform, 
+	FGeometryScriptPrimitiveOptions PrimitiveOptions,
+	FVector3d PreTranslate = FVector3d::Zero())
+{
+	auto ApplyOptionsToMesh = [&Transform, &PrimitiveOptions, PreTranslate](FDynamicMesh3& Mesh)
+	{
+		if (PreTranslate.SquaredLength() > 0)
+		{
+			MeshTransforms::Translate(Mesh, PreTranslate);
+		}
+
+		MeshTransforms::ApplyTransform(Mesh, (UE::Geometry::FTransform3d)Transform);
+		if (PrimitiveOptions.PolygroupMode == EGeometryScriptPrimitivePolygroupMode::SingleGroup)
+		{
+			for (int32 tid : Mesh.TriangleIndicesItr())
+			{
+				Mesh.SetTriangleGroup(tid, 0);
+			}
+		}
+		if (PrimitiveOptions.bFlipOrientation)
+		{
+			Mesh.ReverseOrientation(true);
+			if (Mesh.HasAttributes())
+			{
+				FDynamicMeshNormalOverlay* Normals = Mesh.Attributes()->PrimaryNormals();
+				for (int elemid : Normals->ElementIndicesItr())
+				{
+					Normals->SetElement(elemid, -Normals->GetElement(elemid));
+				}
+			}
+		}
+	};
+
+	if (TargetMesh->IsEmpty())
+	{
+		TargetMesh->EditMesh([&](FDynamicMesh3& EditMesh)
+		{
+			EditMesh.Copy(AppendMesh);
+			ApplyOptionsToMesh(EditMesh);
+		}, EDynamicMeshChangeType::GeneralEdit, EDynamicMeshAttributeChangeFlags::Unknown, false);
+	}
+	else
+	{
+		ApplyOptionsToMesh(AppendMesh);
+		TargetMesh->EditMesh([&](FDynamicMesh3& EditMesh)
+		{
+			FMeshIndexMappings TmpMappings;
+			FDynamicMeshEditor Editor(&EditMesh);
+			Editor.AppendMesh(&AppendMesh, TmpMappings);
+		}, EDynamicMeshChangeType::GeneralEdit, EDynamicMeshAttributeChangeFlags::Unknown, false);
+	}
+}
+
+
+
 
 UDynamicMesh* UGeometryScriptLibrary_MeshPrimitiveFunctions::AppendBox(
 	UDynamicMesh* TargetMesh,
@@ -302,6 +358,7 @@ UDynamicMesh* UGeometryScriptLibrary_MeshPrimitiveFunctions::AppendTorus(
 	UDynamicMesh* TargetMesh,
 	FGeometryScriptPrimitiveOptions PrimitiveOptions,
 	FTransform Transform,
+	FGeometryScriptRevolveOptions RevolveOptions,
 	float MajorRadius,
 	float MinorRadius,
 	int32 MajorSteps,
@@ -316,15 +373,16 @@ UDynamicMesh* UGeometryScriptLibrary_MeshPrimitiveFunctions::AppendTorus(
 	{
 		PolygonVertices.Add((FVector2D)(v+Shift));
 	}
-	return AppendSimpleRevolvePolygon(TargetMesh, PrimitiveOptions, Transform, PolygonVertices, MajorRadius, MajorSteps, Debug);
+	return AppendRevolvePolygon(TargetMesh, PrimitiveOptions, Transform, PolygonVertices, RevolveOptions, MajorRadius, MajorSteps, Debug);
 }
 
 
-UDynamicMesh* UGeometryScriptLibrary_MeshPrimitiveFunctions::AppendSimpleRevolvePolygon(
+UDynamicMesh* UGeometryScriptLibrary_MeshPrimitiveFunctions::AppendRevolvePolygon(
 	UDynamicMesh* TargetMesh,
 	FGeometryScriptPrimitiveOptions PrimitiveOptions,
 	FTransform Transform,
 	const TArray<FVector2D>& PolygonVertices,
+	FGeometryScriptRevolveOptions RevolveOptions,
 	float Radius,
 	int32 Steps,
 	UGeometryScriptDebug* Debug)
@@ -340,178 +398,151 @@ UDynamicMesh* UGeometryScriptLibrary_MeshPrimitiveFunctions::AppendSimpleRevolve
 		return TargetMesh;
 	}
 
-	FGeneralizedCylinderGenerator RevolveGen;
+	FRevolvePlanarPolygonGenerator RevolveGen;
 	for (FVector2D Point : PolygonVertices)
 	{
-		RevolveGen.CrossSection.AppendVertex(FVector2d(Point.X, Point.Y));
+		RevolveGen.PolygonVertices.Add(FVector2d(Point.X + Radius, Point.Y));
 	}
-	
-	FPolygon2d PathPoly = FPolygon2d::MakeCircle(FMath::Max(FMathf::ZeroTolerance, Radius), FMath::Max(3, Steps));
-	for (FVector2d v : PathPoly.GetVertices())
+	RevolveGen.Steps = FMath::Max(Steps, 2);
+	RevolveGen.RevolveDegrees = RevolveOptions.RevolveDegrees;
+	RevolveGen.DegreeOffset = RevolveOptions.DegreeOffset;
+	RevolveGen.bReverseDirection = RevolveOptions.bReverseDirection;
+	RevolveGen.bProfileAtMidpoint = RevolveOptions.bProfileAtMidpoint;
+	RevolveGen.bFillPartialRevolveEndcaps = RevolveOptions.bFillPartialRevolveEndcaps;
+
+	FDynamicMesh3 ResultMesh = RevolveGen.GenerateMesh();
+
+	// generate hard normals
+	if (RevolveOptions.bHardNormals)
 	{
-		RevolveGen.Path.Add(FVector3d(v.X, v.Y, 0.0));
+		float NormalDotProdThreshold = FMathf::Cos( FMathf::Clamp(RevolveOptions.HardNormalAngle, 0.0, 180.0) * FMathf::DegToRad);
+		FMeshNormals FaceNormalsCalc(&ResultMesh);
+		FaceNormalsCalc.ComputeTriangleNormals();
+		const TArray<FVector3d>& FaceNormals = FaceNormalsCalc.GetNormals();
+		ResultMesh.Attributes()->PrimaryNormals()->CreateFromPredicate([&](int VID, int TA, int TB) {
+			return !(FaceNormals[TA].Dot(FaceNormals[TB]) < NormalDotProdThreshold);
+		}, 0);
+		FMeshNormals RecomputeMeshNormals(&ResultMesh);
+		RecomputeMeshNormals.RecomputeOverlayNormals(ResultMesh.Attributes()->PrimaryNormals(), true, true);
+		RecomputeMeshNormals.CopyToOverlay(ResultMesh.Attributes()->PrimaryNormals(), false);
 	}
-	Algo::Reverse(RevolveGen.Path);
-
-	RevolveGen.bLoop = true;
-	RevolveGen.bCapped = false;
-	RevolveGen.bPolygroupPerQuad = (PrimitiveOptions.PolygroupMode == EGeometryScriptPrimitivePolygroupMode::PerQuad);
-	RevolveGen.InitialFrame = FFrame3d(RevolveGen.Path[0], FVector3d::UnitX(), FVector3d::UnitZ(), -FVector3d::UnitY());
-	RevolveGen.Generate();
 
-	AppendPrimitive(TargetMesh, &RevolveGen, Transform, PrimitiveOptions);
+	AppendPrimitiveMesh(TargetMesh, ResultMesh, Transform, PrimitiveOptions);
 	return TargetMesh;
 }
 
 
-
-
-UDynamicMesh* UGeometryScriptLibrary_MeshPrimitiveFunctions::AppendRevolvePath(
+UDynamicMesh* UGeometryScriptLibrary_MeshPrimitiveFunctions::AppendSpiralRevolvePolygon(
 	UDynamicMesh* TargetMesh,
 	FGeometryScriptPrimitiveOptions PrimitiveOptions,
 	FTransform Transform,
-	const TArray<FVector2D>& PathVertices,
+	const TArray<FVector2D>& PolygonVertices,
 	FGeometryScriptRevolveOptions RevolveOptions,
-	int32 Steps,
-	bool bCapped,
+	float Radius,
+	int Steps,
+	float RisePerRevolution,
 	UGeometryScriptDebug* Debug)
 {
 	if (TargetMesh == nullptr)
 	{
-		UE::Geometry::AppendError(Debug, EGeometryScriptErrorType::InvalidInputs, LOCTEXT("AppendRevolvePath_NullMesh", "AppendRevolvePath: TargetMesh is Null"));
+		UE::Geometry::AppendError(Debug, EGeometryScriptErrorType::InvalidInputs, LOCTEXT("AppendSpiralRevolvePolygon_NullMesh", "AppendSpiralRevolvePolygon: TargetMesh is Null"));
 		return TargetMesh;
 	}
-	if (PathVertices.Num() < 2)
+	if (PolygonVertices.Num() < 3)
 	{
-		UE::Geometry::AppendError(Debug, EGeometryScriptErrorType::InvalidInputs, LOCTEXT("AppendRevolvePath_InvalidPolygon", "AppendRevolvePath: PathVertices array requires at least 2 positions"));
+		UE::Geometry::AppendError(Debug, EGeometryScriptErrorType::InvalidInputs, LOCTEXT("AppendSpiralRevolvePolygon_InvalidPolygon", "AppendSpiralRevolvePolygon: PolygonVertices array requires at least 3 positions"));
 		return TargetMesh;
 	}
 
-	FCurveSweepOp CurveSweepOp;
-
-	FVector3d AxisDirection(0, 0, 1);
-	FVector3d AxisOrigin(0, 0, 0);
-	for (FVector2D Point : PathVertices)
-	{
-		CurveSweepOp.ProfileCurve.Add(FVector3d(Point.X, 0, Point.Y));
-	}
-	// unclear why but the sweep code seems to be written for a clockwise ordering...
-	Algo::Reverse(CurveSweepOp.ProfileCurve);
-
-	// Project first and last points onto the revolution axis to cap it
-	if (bCapped)
-	{
-		FVector3d FirstPoint = CurveSweepOp.ProfileCurve[0];
-		FVector3d LastPoint = CurveSweepOp.ProfileCurve.Last();
-
-		double DistanceAlongAxis = AxisDirection.Dot(LastPoint - AxisOrigin);
-		FVector3d ProjectedPoint = AxisOrigin + (AxisDirection * DistanceAlongAxis);
-		CurveSweepOp.ProfileCurve.Add(ProjectedPoint);
-
-		DistanceAlongAxis = AxisDirection.Dot(FirstPoint - AxisOrigin);
-		ProjectedPoint = AxisOrigin + (AxisDirection * DistanceAlongAxis);
-		CurveSweepOp.ProfileCurve.Add(ProjectedPoint);
-
-		CurveSweepOp.bProfileCurveIsClosed = true;
-	}
-	else
+	FSpiralRevolvePlanarPolygonGenerator RevolveGen;
+	for (FVector2D Point : PolygonVertices)
 	{
-		CurveSweepOp.bProfileCurveIsClosed = false;
+		RevolveGen.PolygonVertices.Add(FVector2d(Point.X + Radius, Point.Y));
 	}
+	RevolveGen.Steps = FMath::Max(Steps, 2);
+	RevolveGen.RevolveDegrees = FMath::Max(0.0001, RevolveOptions.RevolveDegrees);
+	RevolveGen.DegreeOffset = RevolveOptions.DegreeOffset;
+	RevolveGen.bReverseDirection = RevolveOptions.bReverseDirection;
+	RevolveGen.bProfileAtMidpoint = RevolveOptions.bProfileAtMidpoint;
+	RevolveGen.bFillPartialRevolveEndcaps = RevolveOptions.bFillPartialRevolveEndcaps;
+	RevolveGen.RisePerFullRevolution = RisePerRevolution;
 
-	//double TotalRevolutionDegrees = RevolveDegrees;// (AlongAxisOffsetPerDegree == 0) ? ClampedRevolutionDegrees : RevolutionDegrees;
-	double TotalRevolutionDegrees = FMath::Clamp(RevolveOptions.RevolveDegrees, 0.1f, 360.0f);
+	FDynamicMesh3 ResultMesh = RevolveGen.GenerateMesh();
 
-	double DegreesPerStep = TotalRevolutionDegrees / (double)Steps;
-	double DegreesOffset = RevolveOptions.DegreeOffset;
-	if (RevolveOptions.ReverseDirection)
+	// generate hard normals
+	if (RevolveOptions.bHardNormals)
 	{
-		DegreesPerStep *= -1;
-		DegreesOffset *= -1;
+		float NormalDotProdThreshold = FMathf::Cos( FMathf::Clamp(RevolveOptions.HardNormalAngle, 0.0, 180.0) * FMathf::DegToRad);
+		FMeshNormals FaceNormalsCalc(&ResultMesh);
+		FaceNormalsCalc.ComputeTriangleNormals();
+		const TArray<FVector3d>& FaceNormals = FaceNormalsCalc.GetNormals();
+		ResultMesh.Attributes()->PrimaryNormals()->CreateFromPredicate([&](int VID, int TA, int TB) {
+			return !(FaceNormals[TA].Dot(FaceNormals[TB]) < NormalDotProdThreshold);
+		}, 0);
+		FMeshNormals RecomputeMeshNormals(&ResultMesh);
+		RecomputeMeshNormals.RecomputeOverlayNormals(ResultMesh.Attributes()->PrimaryNormals(), true, true);
+		RecomputeMeshNormals.CopyToOverlay(ResultMesh.Attributes()->PrimaryNormals(), false);
 	}
-	//double DownAxisOffsetPerStep = TotalRevolutionDegrees * AlongAxisOffsetPerDegree / Steps;
 
-	if (RevolveOptions.bProfileAtMidpoint && DegreesPerStep != 0 && FMathd::Abs(DegreesPerStep) < 180)
-	{
-		RevolveUtil::MakeProfileCurveMidpointOfFirstStep(CurveSweepOp.ProfileCurve, DegreesPerStep, AxisOrigin, AxisDirection);
-	}
+	AppendPrimitiveMesh(TargetMesh, ResultMesh, Transform, PrimitiveOptions);
+	return TargetMesh;
+}
 
-	// Generate the sweep curve
-	//CurveSweepOpOut.bSweepCurveIsClosed = bWeldFullRevolution && AlongAxisOffsetPerDegree == 0 && TotalRevolutionDegrees == 360;
-	CurveSweepOp.bSweepCurveIsClosed = (TotalRevolutionDegrees == 360);
-	int32 NumSweepFrames = CurveSweepOp.bSweepCurveIsClosed ? Steps : Steps + 1; // If closed, last sweep frame is also first
-	CurveSweepOp.SweepCurve.Reserve(NumSweepFrames);
-	RevolveUtil::GenerateSweepCurve(AxisOrigin, AxisDirection, DegreesOffset,
-		DegreesPerStep, 0.0, NumSweepFrames, CurveSweepOp.SweepCurve);
 
-	// Weld any vertices that are on the axis
-	RevolveUtil::WeldPointsOnAxis(CurveSweepOp.ProfileCurve, AxisOrigin,
-		AxisDirection, 0.1, CurveSweepOp.ProfileVerticesToWeld);
 
-	CurveSweepOp.bSharpNormals = RevolveOptions.bHardNormals;
-	CurveSweepOp.SharpNormalAngleTolerance = RevolveOptions.HardNormalAngle;
-	//CurveSweepOpOut.DiagonalTolerance = DiagonalProportionTolerance;
-	//double UVScale = MaterialProperties.UVScale;
-	double UVScale = 1.0;
-	CurveSweepOp.UVScale = FVector2d(UVScale, UVScale);
-	//if (bReverseProfileCurve ^ bFlipVs)
-	//{
-	//	CurveSweepOpOut.UVScale[1] *= -1;
-	//	CurveSweepOpOut.UVOffset = FVector2d(0, UVScale);
-	//}
-	CurveSweepOp.bUVsSkipFullyWeldedEdges = true;
-	CurveSweepOp.bUVScaleRelativeWorld = true; // MaterialProperties.bWorldSpaceUVScale;
-	CurveSweepOp.UnitUVInWorldCoordinates = 100; // This seems to be the case in the AddPrimitiveTool
-	CurveSweepOp.QuadSplitMode = EProfileSweepQuadSplit::Uniform;
 
-	if (PrimitiveOptions.PolygroupMode == EGeometryScriptPrimitivePolygroupMode::PerQuad)
-	{
-		CurveSweepOp.PolygonGroupingMode = EProfileSweepPolygonGrouping::PerFace;
-	}
-	if (PrimitiveOptions.PolygroupMode == EGeometryScriptPrimitivePolygroupMode::PerFace)
+UDynamicMesh* UGeometryScriptLibrary_MeshPrimitiveFunctions::AppendRevolvePath(
+	UDynamicMesh* TargetMesh,
+	FGeometryScriptPrimitiveOptions PrimitiveOptions,
+	FTransform Transform,
+	const TArray<FVector2D>& PathVertices,
+	FGeometryScriptRevolveOptions RevolveOptions,
+	int32 Steps,
+	bool bCapped,
+	UGeometryScriptDebug* Debug)
+{
+	if (TargetMesh == nullptr)
 	{
-		CurveSweepOp.PolygonGroupingMode = EProfileSweepPolygonGrouping::PerProfileSegment;
+		UE::Geometry::AppendError(Debug, EGeometryScriptErrorType::InvalidInputs, LOCTEXT("AppendRevolvePath_NullMesh", "AppendRevolvePath: TargetMesh is Null"));
+		return TargetMesh;
 	}
-	else
+	if (PathVertices.Num() < 2)
 	{
-		CurveSweepOp.PolygonGroupingMode = EProfileSweepPolygonGrouping::Single;
+		UE::Geometry::AppendError(Debug, EGeometryScriptErrorType::InvalidInputs, LOCTEXT("AppendRevolvePath_InvalidPolygon", "AppendRevolvePath: PathVertices array requires at least 2 positions"));
+		return TargetMesh;
 	}
 
-	CurveSweepOp.CapFillMode = (bCapped) ?
-		FCurveSweepOp::ECapFillMode::EarClipping : FCurveSweepOp::ECapFillMode::None;		// delaunay? hits ensure...
-
-	CurveSweepOp.CalculateResult(nullptr);
-	TUniquePtr<FDynamicMesh3> ResultMesh = CurveSweepOp.ExtractResult();
-	MeshTransforms::ApplyTransform(*ResultMesh, (UE::Geometry::FTransform3d)Transform);
-
-	if (PrimitiveOptions.bFlipOrientation)
+	FRevolvePlanarPathGenerator RevolveGen;
+	for (FVector2D Point : PathVertices)
 	{
-		ResultMesh->ReverseOrientation(true);
-		if (ResultMesh->HasAttributes() && ResultMesh->Attributes()->PrimaryNormals() != nullptr)
-		{
-			FDynamicMeshNormalOverlay* Normals = ResultMesh->Attributes()->PrimaryNormals();
-			for (int elemid : Normals->ElementIndicesItr())
-			{
-				Normals->SetElement(elemid, -Normals->GetElement(elemid));
-			}
-		}
+		RevolveGen.PathVertices.Add((FVector2d)Point);
 	}
+	RevolveGen.Steps = FMath::Max(Steps, 2);
+	RevolveGen.bCapped = bCapped;
+	RevolveGen.RevolveDegrees = RevolveOptions.RevolveDegrees;
+	RevolveGen.DegreeOffset = RevolveOptions.DegreeOffset;
+	RevolveGen.bReverseDirection = RevolveOptions.bReverseDirection;
+	RevolveGen.bProfileAtMidpoint = RevolveOptions.bProfileAtMidpoint;
+	RevolveGen.bFillPartialRevolveEndcaps = RevolveOptions.bFillPartialRevolveEndcaps;
 
-	if (TargetMesh->IsEmpty())
-	{
-		TargetMesh->SetMesh(MoveTemp(*ResultMesh));
-	}
-	else
-	{
-		TargetMesh->EditMesh([&](FDynamicMesh3& EditMesh)
-		{
-			FMeshIndexMappings TmpMappings;
-			FDynamicMeshEditor Editor(&EditMesh);
-			Editor.AppendMesh(ResultMesh.Get(), TmpMappings);
+	FDynamicMesh3 ResultMesh = RevolveGen.GenerateMesh();
 
-		}, EDynamicMeshChangeType::GeneralEdit, EDynamicMeshAttributeChangeFlags::Unknown, false);
+	// generate hard normals
+	if (RevolveOptions.bHardNormals)
+	{
+		float NormalDotProdThreshold = FMathf::Cos(FMathf::Clamp(RevolveOptions.HardNormalAngle, 0.0, 180.0) * FMathf::DegToRad);
+		FMeshNormals FaceNormalsCalc(&ResultMesh);
+		FaceNormalsCalc.ComputeTriangleNormals();
+		const TArray<FVector3d>& FaceNormals = FaceNormalsCalc.GetNormals();
+		ResultMesh.Attributes()->PrimaryNormals()->CreateFromPredicate([&](int VID, int TA, int TB) {
+			return !(FaceNormals[TA].Dot(FaceNormals[TB]) < NormalDotProdThreshold);
+		}, 0);
+		FMeshNormals RecomputeMeshNormals(&ResultMesh);
+		RecomputeMeshNormals.RecomputeOverlayNormals(ResultMesh.Attributes()->PrimaryNormals(), true, true);
+		RecomputeMeshNormals.CopyToOverlay(ResultMesh.Attributes()->PrimaryNormals(), false);
 	}
 
+	AppendPrimitiveMesh(TargetMesh, ResultMesh, Transform, PrimitiveOptions);
 	return TargetMesh;
 }
 
diff --git a/Engine/Plugins/Experimental/GeometryScripting/Source/GeometryScriptingCore/Public/GeometryScript/MeshPrimitiveFunctions.h b/Engine/Plugins/Experimental/GeometryScripting/Source/GeometryScriptingCore/Public/GeometryScript/MeshPrimitiveFunctions.h
--- a/Engine/Plugins/Experimental/GeometryScripting/Source/GeometryScriptingCore/Public/GeometryScript/MeshPrimitiveFunctions.h
+++ b/Engine/Plugins/Experimental/GeometryScripting/Source/GeometryScriptingCore/Public/GeometryScript/MeshPrimitiveFunctions.h
@@ -49,7 +49,7 @@ public:
 	float DegreeOffset = 0.0f;
 
 	UPROPERTY(BlueprintReadWrite, Category = Options)
-	bool ReverseDirection = false;
+	bool bReverseDirection = false;
 
 	UPROPERTY(BlueprintReadWrite, Category = Options)
 	bool bHardNormals = false;
@@ -59,6 +59,9 @@ public:
 
 	UPROPERTY(BlueprintReadWrite, Category = Options)
 	bool bProfileAtMidpoint = false;
+
+	UPROPERTY(BlueprintReadWrite, Category = Options)
+	bool bFillPartialRevolveEndcaps = true;
 };
 
 
@@ -163,6 +166,7 @@ public:
 		UDynamicMesh* TargetMesh, 
 		FGeometryScriptPrimitiveOptions PrimitiveOptions,
 		FTransform Transform,
+		FGeometryScriptRevolveOptions RevolveOptions,
 		float MajorRadius = 50,
 		float MinorRadius = 25,
 		int32 MajorSteps = 16,
@@ -177,16 +181,31 @@ public:
 	 */
 	UFUNCTION(BlueprintCallable, Category = "GeometryScript|Primitives", meta=(ScriptMethod))
 	static UPARAM(DisplayName = "Target Mesh") UDynamicMesh* 
-	AppendSimpleRevolvePolygon( 
+	AppendRevolvePolygon( 
 		UDynamicMesh* TargetMesh, 
 		FGeometryScriptPrimitiveOptions PrimitiveOptions,
 		FTransform Transform,
 		const TArray<FVector2D>& PolygonVertices,
+		FGeometryScriptRevolveOptions RevolveOptions,
 		float Radius = 100,
 		int32 Steps = 8,
 		UGeometryScriptDebug* Debug = nullptr);
 
 
+	UFUNCTION(BlueprintCallable, Category = "GeometryScript|Primitives", meta = (ScriptMethod))
+	static UPARAM(DisplayName = "Target Mesh") UDynamicMesh*
+	AppendSpiralRevolvePolygon(
+		UDynamicMesh* TargetMesh,
+		FGeometryScriptPrimitiveOptions PrimitiveOptions,
+		FTransform Transform,
+		const TArray<FVector2D>& PolygonVertices,
+		FGeometryScriptRevolveOptions RevolveOptions,
+		float Radius = 100,
+		int Steps = 18,
+		float RisePerRevolution = 50,
+		UGeometryScriptDebug* Debug = nullptr);
+
+
 	UFUNCTION(BlueprintCallable, Category = "GeometryScript|Primitives", meta=(ScriptMethod))
 	static UPARAM(DisplayName = "Target Mesh") UDynamicMesh* 
 	AppendRevolvePath( 
diff --git a/Engine/Source/Runtime/GeometryCore/Private/Generators/RevolveGenerator.cpp b/Engine/Source/Runtime/GeometryCore/Private/Generators/RevolveGenerator.cpp
new file mode 100644
--- /dev/null
+++ b/Engine/Source/Runtime/GeometryCore/Private/Generators/RevolveGenerator.cpp
@@ -0,0 +1,698 @@
+// Copyright Epic Games, Inc. All Rights Reserved.
+
+#include "Generators/RevolveGenerator.h"
+#include "Util/RevolveUtil.h"
+#include "Async/ParallelFor.h"
+#include "DynamicMesh/DynamicMeshAttributeSet.h"
+#include "DynamicMesh/MeshNormals.h"
+
+using namespace UE::Geometry;
+
+
+namespace UELocal
+{
+	//
+	// These are local copies of FDynamicMeshEditor functions, which is currenlty not available in GeometryCore.
+	// However simpler implementations would be sufficient for the sweep code below, which only uses them
+	// to handle planar caps.
+	//
+
+	static void SetTriangleNormals(FDynamicMesh3* Mesh, const TArray<int>& Triangles)
+	{
+		check(Mesh->HasAttributes());
+		FDynamicMeshNormalOverlay* Normals = Mesh->Attributes()->PrimaryNormals();
+
+		TSet<int32> TriangleSet(Triangles);
+		TUniqueFunction<bool(int32)> TrianglePredicate = [&](int32 TriangleID) { return TriangleSet.Contains(TriangleID); };
+
+		TMap<int, int> Vertices;
+
+		for (int tid : Triangles)
+		{
+			if (Normals->IsSetTriangle(tid))
+			{
+				Normals->UnsetTriangle(tid);
+			}
+
+			FIndex3i BaseTri = Mesh->GetTriangle(tid);
+			FIndex3i ElemTri;
+			for (int j = 0; j < 3; ++j)
+			{
+				const int* FoundElementID = Vertices.Find(BaseTri[j]);
+				if (FoundElementID == nullptr)
+				{
+					FVector3d VtxROINormal = FMeshNormals::ComputeVertexNormal(*Mesh, BaseTri[j], TrianglePredicate);
+					ElemTri[j] = Normals->AppendElement( (FVector3f)VtxROINormal);
+					Vertices.Add(BaseTri[j], ElemTri[j]);
+				}
+				else
+				{
+					ElemTri[j] = *FoundElementID;
+				}
+			}
+			Normals->SetTriangle(tid, ElemTri);
+		}
+	}
+
+
+
+	static void SetTriangleUVsFromProjection(FDynamicMesh3* Mesh, const TArray<int>& Triangles, const FFrame3d& ProjectionFrame, const FVector2f& UVScale, 
+		const FVector2f& UVTranslation, int UVLayerIndex, bool bShiftToOrigin, bool bNormalizeBeforeScaling)
+	{
+		if (!Triangles.Num())
+		{
+			return;
+		}
+
+		check(Mesh->HasAttributes() && Mesh->Attributes()->NumUVLayers() > UVLayerIndex);
+		FDynamicMeshUVOverlay* UVs = Mesh->Attributes()->GetUVLayer(UVLayerIndex);
+
+		TMap<int, int> BaseToOverlayVIDMap;
+		TArray<int> AllUVIndices;
+
+		FAxisAlignedBox2f UVBounds(FAxisAlignedBox2f::Empty());
+
+		for (int TID : Triangles)
+		{
+			if (UVs->IsSetTriangle(TID))
+			{
+				UVs->UnsetTriangle(TID);
+			}
+
+			FIndex3i BaseTri = Mesh->GetTriangle(TID);
+			FIndex3i ElemTri;
+			for (int j = 0; j < 3; ++j)
+			{
+				const int* FoundElementID = BaseToOverlayVIDMap.Find(BaseTri[j]);
+				if (FoundElementID == nullptr)
+				{
+					FVector2f UV = (FVector2f)ProjectionFrame.ToPlaneUV(Mesh->GetVertex(BaseTri[j]), 2);
+					UVBounds.Contain(UV);
+					ElemTri[j] = UVs->AppendElement(UV);
+					AllUVIndices.Add(ElemTri[j]);
+					BaseToOverlayVIDMap.Add(BaseTri[j], ElemTri[j]);
+				}
+				else
+				{
+					ElemTri[j] = *FoundElementID;
+				}
+			}
+			UVs->SetTriangle(TID, ElemTri);
+		}
+
+		FVector2f UvScaleToUse = bNormalizeBeforeScaling ? FVector2f(UVScale[0] / UVBounds.Width(), UVScale[1] / UVBounds.Height())
+			: UVScale;
+
+		// shift UVs so that their bbox min-corner is at origin and scaled by external scale factor
+		for (int UVID : AllUVIndices)
+		{
+			FVector2f UV = UVs->GetElement(UVID);
+			FVector2f TransformedUV = (bShiftToOrigin) ? ((UV - UVBounds.Min) * UvScaleToUse) : (UV * UvScaleToUse);
+			TransformedUV += UVTranslation;
+			UVs->SetElement(UVID, TransformedUV);
+		}
+	}
+
+}
+
+
+
+FDynamicMesh3 FRevolvePlanarPathGenerator::GenerateMesh()
+{
+	// revolve around +Z axis, around origin
+	FVector3d AxisDirection(0, 0, 1);
+	FVector3d AxisOrigin(0, 0, 0);
+
+	TArray<FVector3d> ProfileCurve;
+	TSet<int32> ProfileVerticesToWeld;
+	bool bProfileCurveIsClosed = false;
+
+	for (FVector2d Point : PathVertices)
+	{
+		ProfileCurve.Add(FVector3d(Point.X, 0, Point.Y));
+	}
+	// unclear why but the sweep code seems to be written for a clockwise ordering...
+	Algo::Reverse(ProfileCurve);
+
+	// Project first and last points onto the revolution axis to cap it
+	if (bCapped)
+	{
+		FVector3d FirstPoint = ProfileCurve[0];
+		FVector3d LastPoint = ProfileCurve.Last();
+
+		double DistanceAlongAxis = AxisDirection.Dot(LastPoint - AxisOrigin);
+		FVector3d ProjectedPoint = AxisOrigin + (AxisDirection * DistanceAlongAxis);
+		ProfileCurve.Add(ProjectedPoint);
+
+		DistanceAlongAxis = AxisDirection.Dot(FirstPoint - AxisOrigin);
+		ProjectedPoint = AxisOrigin + (AxisDirection * DistanceAlongAxis);
+		ProfileCurve.Add(ProjectedPoint);
+
+		bProfileCurveIsClosed = true;
+	}
+	else
+	{
+		bProfileCurveIsClosed = false;
+	}
+
+	//double TotalRevolutionDegrees = RevolveDegrees;// (AlongAxisOffsetPerDegree == 0) ? ClampedRevolutionDegrees : RevolutionDegrees;
+	double TotalRevolutionDegrees = FMath::Clamp(RevolveDegrees, 0.1f, 360.0f);
+
+	double DegreesPerStep = TotalRevolutionDegrees / (double)Steps;
+	double DegreesOffset = DegreeOffset;
+	if (bReverseDirection)
+	{
+		DegreesPerStep *= -1;
+		DegreesOffset *= -1;
+	}
+	//double DownAxisOffsetPerStep = TotalRevolutionDegrees * AlongAxisOffsetPerDegree / Steps;
+
+	if (bProfileAtMidpoint && DegreesPerStep != 0 && FMathd::Abs(DegreesPerStep) < 180)
+	{
+		RevolveUtil::MakeProfileCurveMidpointOfFirstStep(ProfileCurve, DegreesPerStep, AxisOrigin, AxisDirection);
+	}
+
+	// Generate the sweep curve
+	
+	TArray<FFrame3d> SweepCurve;
+	bool bSweepCurveIsClosed = (TotalRevolutionDegrees == 360);
+
+	//CurveSweepOpOut.bSweepCurveIsClosed = bWeldFullRevolution && AlongAxisOffsetPerDegree == 0 && TotalRevolutionDegrees == 360;
+	int32 NumSweepFrames = bSweepCurveIsClosed ? Steps : Steps + 1; // If closed, last sweep frame is also first
+	SweepCurve.Reserve(NumSweepFrames);
+	RevolveUtil::GenerateSweepCurve(AxisOrigin, AxisDirection, DegreesOffset,
+		DegreesPerStep, 0.0, NumSweepFrames, SweepCurve);
+
+	// Weld any vertices that are on the axis
+	RevolveUtil::WeldPointsOnAxis(ProfileCurve, AxisOrigin,
+		AxisDirection, 0.1, ProfileVerticesToWeld);
+
+	//CurveSweepOpOut.DiagonalTolerance = DiagonalProportionTolerance;
+	//double UVScale = MaterialProperties.UVScale;
+
+	//double UVScale = 1.0;
+	//CurveSweepOp.UVScale = FVector2d(UVScale, UVScale);
+
+	// what for?
+	//if (bReverseProfileCurve ^ bFlipVs)
+	//{
+	//	CurveSweepOpOut.UVScale[1] *= -1;
+	//	CurveSweepOpOut.UVOffset = FVector2d(0, UVScale);
+	//}
+
+	bool bUVsSkipFullyWeldedEdges = true;
+	ECapFillMode CapFillMode = (bFillPartialRevolveEndcaps) ?
+		ECapFillMode::EarClipping : ECapFillMode::None;		// delaunay? hits ensure...
+
+	// Use a mesh generator to do the work.
+	FProfileSweepGenerator CurveSweeper;
+	CurveSweeper.ProfileCurve = ProfileCurve;
+	CurveSweeper.SweepCurve = SweepCurve;
+	CurveSweeper.WeldedVertices = ProfileVerticesToWeld;
+	CurveSweeper.bProfileCurveIsClosed = bProfileCurveIsClosed;
+	CurveSweeper.bSweepCurveIsClosed = bSweepCurveIsClosed;
+
+	CurveSweeper.bSharpNormals = false; // bSharpNormals;
+	CurveSweeper.UVScale = UVScale;
+	CurveSweeper.bUVScaleRelativeWorld = bUVScaleRelativeWorld;
+	CurveSweeper.UnitUVInWorldCoordinates = UnitUVInWorldCoordinates;
+	CurveSweeper.UVOffset = UVOffset;
+	CurveSweeper.bUVsSkipFullyWeldedEdges = bUVsSkipFullyWeldedEdges;
+	CurveSweeper.PolygonGroupingMode = PolygonGroupingMode;
+	CurveSweeper.QuadSplitMethod = QuadSplitMethod;
+	CurveSweeper.DiagonalTolerance = DiagonalTolerance;
+	
+	FDynamicMesh3 ResultMesh(&CurveSweeper.Generate());
+
+	// Cap the ends if needed
+	if (bProfileCurveIsClosed && !bSweepCurveIsClosed && CapFillMode != ECapFillMode::None)
+	{
+		for (int CapIndex = 0; CapIndex < 2; ++CapIndex)
+		{
+			// Note that we rely on the sweep generator to tell us the start and end profile curves for
+			// capping rather than looking for a boundary ourselves. This is because in cases where the
+			// side of the profile curve is welded, there will actually be one open boundary rather than
+			// two separate ones, yet we want to cap the two sides separately.
+			// This does rely on the fact that the hole fillers we use only care about vertices, and don't
+			// rely on there being edges between them. It also relies on the vertex ids in the DynamicMesh
+			// matching the ones used in the generator. Without those two facts, the work of filling two
+			// halves of an open boundary separately would be harder to do...
+
+			// TODO: curve is planar we don't have to figure this out...
+
+			// Put the profile curve into an edge loop
+			TArray<int> VertexLoop = CurveSweeper.EndProfiles[CapIndex];
+			int NV = VertexLoop.Num();
+			TArray<int32> Edges;
+			Edges.SetNum(NV);
+			for (int i = 0; i < NV; ++i)
+			{
+				Edges[i] = ResultMesh.FindEdge(VertexLoop[i], VertexLoop[(i + 1) % NV]);
+			}
+
+			//FEdgeLoop::VertexLoopToEdgeLoop(ResultMesh.Get(), , Edges);
+			//FEdgeLoop Loop(ResultMesh.Get(), CurveSweeper.EndProfiles[CapIndex], Edges);
+
+			// The first cap actually needs to face in the opposite direction of the sweep to face outward.
+			if (CapIndex == 0)
+			{
+				Algo::Reverse(VertexLoop);
+				Algo::Reverse(Edges);
+			}
+
+			// Compute a best-fit plane of the vertices
+			TArray<FVector3<double>> VertexPositions;
+			for (int32 vid : VertexLoop)
+			{
+				VertexPositions.Add( ResultMesh.GetVertex(vid) );
+			}
+
+			FVector3<double> PlaneOrigin;
+			FVector3<double> PlaneNormal;
+			PolygonTriangulation::ComputePolygonPlane<double>(VertexPositions, PlaneNormal, PlaneOrigin);
+			PlaneNormal *= -1.0;	// Previous function seems to orient the normal opposite to what's expected elsewhere
+
+			int32 NewGroupID = ResultMesh.AllocateTriangleGroup();
+
+			TArray<FIndex3i> Triangles;
+			TArray<int32> NewTriangles;
+			PolygonTriangulation::TriangulateSimplePolygon(VertexPositions, Triangles);
+			for (FIndex3i PolyTriangle : Triangles)
+			{
+				FIndex3i MeshTriangle( VertexLoop[PolyTriangle.A], VertexLoop[PolyTriangle.C], VertexLoop[PolyTriangle.B]);  // Reversing orientation here!!
+				int32 NewTriangle = ResultMesh.AppendTriangle(MeshTriangle, NewGroupID);
+				if (NewTriangle >= 0)
+				{
+					NewTriangles.Add(NewTriangle);
+				}
+			}
+
+			// Compute normals and UVs
+			if (ResultMesh.HasAttributes())
+			{
+				UELocal::SetTriangleNormals(&ResultMesh, NewTriangles);
+
+				// Set UVs for the cap
+				FFrame3d ProjectionFrame(PlaneOrigin, PlaneNormal);
+				if (bUVScaleRelativeWorld)
+				{
+					FVector2f ScaleToUse(UVScale);
+					ScaleToUse /= CurveSweeper.UnitUVInWorldCoordinates;
+					UELocal::SetTriangleUVsFromProjection(&ResultMesh, NewTriangles, ProjectionFrame, ScaleToUse, (FVector2f)UVOffset,
+						0, // layer
+						true,  // shift to origin
+						false); // don't normalize
+				}
+				else
+				{
+					UELocal::SetTriangleUVsFromProjection(&ResultMesh, NewTriangles, ProjectionFrame, (FVector2f)UVScale, (FVector2f)UVOffset,
+						0, // layer
+						true,  // shift to origin
+						true); // normalize first
+				}
+
+			}
+		}// end for each loop
+	}//end if capping
+
+
+	return ResultMesh;
+}
+
+
+
+
+
+
+
+FDynamicMesh3 FRevolvePlanarPolygonGenerator::GenerateMesh()
+{
+	// revolve around +Z axis, around origin
+	FVector3d AxisDirection(0, 0, 1);
+	FVector3d AxisOrigin(0, 0, 0);
+
+	TArray<FVector3d> ProfileCurve;
+	TSet<int32> ProfileVerticesToWeld;
+	bool bProfileCurveIsClosed = true;
+
+	for (FVector2d Point : PolygonVertices)
+	{
+		ProfileCurve.Add(FVector3d(Point.X, 0, Point.Y));
+	}
+	// unclear why but the sweep code seems to be written for a clockwise ordering...
+	Algo::Reverse(ProfileCurve);
+
+	//double TotalRevolutionDegrees = RevolveDegrees;// (AlongAxisOffsetPerDegree == 0) ? ClampedRevolutionDegrees : RevolutionDegrees;
+	double TotalRevolutionDegrees = FMath::Clamp(RevolveDegrees, 0.1f, 360.0f);
+
+	double DegreesPerStep = TotalRevolutionDegrees / (double)Steps;
+	double DegreesOffset = DegreeOffset;
+	if (bReverseDirection)
+	{
+		DegreesPerStep *= -1;
+		DegreesOffset *= -1;
+	}
+	//double DownAxisOffsetPerStep = TotalRevolutionDegrees * AlongAxisOffsetPerDegree / Steps;
+
+	if (bProfileAtMidpoint && DegreesPerStep != 0 && FMathd::Abs(DegreesPerStep) < 180)
+	{
+		RevolveUtil::MakeProfileCurveMidpointOfFirstStep(ProfileCurve, DegreesPerStep, AxisOrigin, AxisDirection);
+	}
+
+	// Generate the sweep curve
+
+	TArray<FFrame3d> SweepCurve;
+	bool bSweepCurveIsClosed = (TotalRevolutionDegrees == 360);
+
+	//CurveSweepOpOut.bSweepCurveIsClosed = bWeldFullRevolution && AlongAxisOffsetPerDegree == 0 && TotalRevolutionDegrees == 360;
+	int32 NumSweepFrames = bSweepCurveIsClosed ? Steps : Steps + 1; // If closed, last sweep frame is also first
+	SweepCurve.Reserve(NumSweepFrames);
+	RevolveUtil::GenerateSweepCurve(AxisOrigin, AxisDirection, DegreesOffset,
+		DegreesPerStep, 0.0, NumSweepFrames, SweepCurve);
+
+	// Weld any vertices that are on the axis
+	if (bWeldVertsOnAxis)
+	{
+		RevolveUtil::WeldPointsOnAxis(ProfileCurve, AxisOrigin,
+			AxisDirection, 0.1, ProfileVerticesToWeld);
+	}
+
+	//CurveSweepOpOut.DiagonalTolerance = DiagonalProportionTolerance;
+	//double UVScale = MaterialProperties.UVScale;
+
+	//double UVScale = 1.0;
+	//CurveSweepOp.UVScale = FVector2d(UVScale, UVScale);
+
+	// what for?
+	//if (bReverseProfileCurve ^ bFlipVs)
+	//{
+	//	CurveSweepOpOut.UVScale[1] *= -1;
+	//	CurveSweepOpOut.UVOffset = FVector2d(0, UVScale);
+	//}
+
+	bool bUVsSkipFullyWeldedEdges = true;
+	ECapFillMode CapFillMode = (bFillPartialRevolveEndcaps) ?
+		ECapFillMode::EarClipping : ECapFillMode::None;		// delaunay? hits ensure...
+
+															// Use a mesh generator to do the work.
+	FProfileSweepGenerator CurveSweeper;
+	CurveSweeper.ProfileCurve = ProfileCurve;
+	CurveSweeper.SweepCurve = SweepCurve;
+	CurveSweeper.WeldedVertices = ProfileVerticesToWeld;
+	CurveSweeper.bProfileCurveIsClosed = bProfileCurveIsClosed;
+	CurveSweeper.bSweepCurveIsClosed = bSweepCurveIsClosed;
+
+	CurveSweeper.bSharpNormals = false; // bSharpNormals;
+	CurveSweeper.UVScale = UVScale;
+	CurveSweeper.bUVScaleRelativeWorld = bUVScaleRelativeWorld;
+	CurveSweeper.UnitUVInWorldCoordinates = UnitUVInWorldCoordinates;
+	CurveSweeper.UVOffset = UVOffset;
+	CurveSweeper.bUVsSkipFullyWeldedEdges = bUVsSkipFullyWeldedEdges;
+	CurveSweeper.PolygonGroupingMode = PolygonGroupingMode;
+	CurveSweeper.QuadSplitMethod = QuadSplitMethod;
+	CurveSweeper.DiagonalTolerance = DiagonalTolerance;
+
+	FDynamicMesh3 ResultMesh(&CurveSweeper.Generate());
+
+	// Cap the ends if needed
+	if (bProfileCurveIsClosed && !bSweepCurveIsClosed && CapFillMode != ECapFillMode::None)
+	{
+		for (int CapIndex = 0; CapIndex < 2; ++CapIndex)
+		{
+			// Note that we rely on the sweep generator to tell us the start and end profile curves for
+			// capping rather than looking for a boundary ourselves. This is because in cases where the
+			// side of the profile curve is welded, there will actually be one open boundary rather than
+			// two separate ones, yet we want to cap the two sides separately.
+			// This does rely on the fact that the hole fillers we use only care about vertices, and don't
+			// rely on there being edges between them. It also relies on the vertex ids in the DynamicMesh
+			// matching the ones used in the generator. Without those two facts, the work of filling two
+			// halves of an open boundary separately would be harder to do...
+
+			// TODO: curve is planar we don't have to figure this out...
+
+			// Put the profile curve into an edge loop
+			TArray<int> VertexLoop = CurveSweeper.EndProfiles[CapIndex];
+			int NV = VertexLoop.Num();
+			TArray<int32> Edges;
+			Edges.SetNum(NV);
+			for (int i = 0; i < NV; ++i)
+			{
+				Edges[i] = ResultMesh.FindEdge(VertexLoop[i], VertexLoop[(i + 1) % NV]);
+			}
+
+			//FEdgeLoop::VertexLoopToEdgeLoop(ResultMesh.Get(), , Edges);
+			//FEdgeLoop Loop(ResultMesh.Get(), CurveSweeper.EndProfiles[CapIndex], Edges);
+
+			// The first cap actually needs to face in the opposite direction of the sweep to face outward.
+			if (CapIndex == 0)
+			{
+				Algo::Reverse(VertexLoop);
+				Algo::Reverse(Edges);
+			}
+
+			// Compute a best-fit plane of the vertices
+			TArray<FVector3<double>> VertexPositions;
+			for (int32 vid : VertexLoop)
+			{
+				VertexPositions.Add( ResultMesh.GetVertex(vid) );
+			}
+
+			FVector3<double> PlaneOrigin;
+			FVector3<double> PlaneNormal;
+			PolygonTriangulation::ComputePolygonPlane<double>(VertexPositions, PlaneNormal, PlaneOrigin);
+			PlaneNormal *= -1.0;	// Previous function seems to orient the normal opposite to what's expected elsewhere
+
+			int32 NewGroupID = ResultMesh.AllocateTriangleGroup();
+
+			TArray<FIndex3i> Triangles;
+			TArray<int32> NewTriangles;
+			PolygonTriangulation::TriangulateSimplePolygon(VertexPositions, Triangles);
+			for (FIndex3i PolyTriangle : Triangles)
+			{
+				FIndex3i MeshTriangle( VertexLoop[PolyTriangle.A], VertexLoop[PolyTriangle.C], VertexLoop[PolyTriangle.B]);  // Reversing orientation here!!
+				int32 NewTriangle = ResultMesh.AppendTriangle(MeshTriangle, NewGroupID);
+				if (NewTriangle >= 0)
+				{
+					NewTriangles.Add(NewTriangle);
+				}
+			}
+
+			// Compute normals and UVs
+			if (ResultMesh.HasAttributes())
+			{
+				UELocal::SetTriangleNormals(&ResultMesh, NewTriangles);
+
+				// Set UVs for the cap
+				FFrame3d ProjectionFrame(PlaneOrigin, PlaneNormal);
+				if (bUVScaleRelativeWorld)
+				{
+					FVector2f ScaleToUse(UVScale);
+					ScaleToUse /= CurveSweeper.UnitUVInWorldCoordinates;
+					UELocal::SetTriangleUVsFromProjection(&ResultMesh, NewTriangles, ProjectionFrame, ScaleToUse, (FVector2f)UVOffset,
+						0, // layer
+						true,  // shift to origin
+						false); // don't normalize
+				}
+				else
+				{
+					UELocal::SetTriangleUVsFromProjection(&ResultMesh, NewTriangles, ProjectionFrame, (FVector2f)UVScale, (FVector2f)UVOffset,
+						0, // layer
+						true,  // shift to origin
+						true); // normalize first
+				}
+
+			}
+		}// end for each loop
+	}//end if capping
+
+	return ResultMesh;
+}
+
+
+
+
+
+
+
+FDynamicMesh3 FSpiralRevolvePlanarPolygonGenerator::GenerateMesh()
+{
+	// revolve around +Z axis, around origin
+	FVector3d AxisDirection(0, 0, 1);
+	FVector3d AxisOrigin(0, 0, 0);
+
+	TArray<FVector3d> ProfileCurve;
+	TSet<int32> ProfileVerticesToWeld;
+	bool bProfileCurveIsClosed = true;
+
+	for (FVector2d Point : PolygonVertices)
+	{
+		ProfileCurve.Add(FVector3d(Point.X, 0, Point.Y));
+	}
+	// unclear why but the sweep code seems to be written for a clockwise ordering...
+	Algo::Reverse(ProfileCurve);
+
+	double RisePerDegree = RisePerFullRevolution / 360.0;
+	double TotalRevolutionDegrees = FMath::Max(0.01, RevolveDegrees);
+	double DegreesPerStep = TotalRevolutionDegrees / (double)Steps;
+	double DegreesOffset = DegreeOffset;
+	if (bReverseDirection)
+	{
+		DegreesPerStep *= -1;
+		DegreesOffset *= -1;
+	}
+	double DownAxisOffsetPerStep = TotalRevolutionDegrees * RisePerDegree / Steps;
+
+	if (bProfileAtMidpoint && DegreesPerStep != 0 && FMathd::Abs(DegreesPerStep) < 180)
+	{
+		RevolveUtil::MakeProfileCurveMidpointOfFirstStep(ProfileCurve, DegreesPerStep, AxisOrigin, AxisDirection);
+	}
+
+	// Generate the sweep curve
+
+	TArray<FFrame3d> SweepCurve;
+	bool bSweepCurveIsClosed = false;
+
+	//CurveSweepOpOut.bSweepCurveIsClosed = bWeldFullRevolution && AlongAxisOffsetPerDegree == 0 && TotalRevolutionDegrees == 360;
+	int32 NumSweepFrames = bSweepCurveIsClosed ? Steps : Steps + 1; // If closed, last sweep frame is also first
+	SweepCurve.Reserve(NumSweepFrames);
+	RevolveUtil::GenerateSweepCurve(AxisOrigin, AxisDirection, DegreesOffset,
+		DegreesPerStep, DownAxisOffsetPerStep, NumSweepFrames, SweepCurve);
+
+	// Weld any vertices that are on the axis
+	if (bWeldVertsOnAxis && DownAxisOffsetPerStep == 0)
+	{
+		RevolveUtil::WeldPointsOnAxis(ProfileCurve, AxisOrigin,
+			AxisDirection, 0.1, ProfileVerticesToWeld);
+	}
+
+	//CurveSweepOpOut.DiagonalTolerance = DiagonalProportionTolerance;
+	//double UVScale = MaterialProperties.UVScale;
+
+	//double UVScale = 1.0;
+	//CurveSweepOp.UVScale = FVector2d(UVScale, UVScale);
+
+	// what for?
+	//if (bReverseProfileCurve ^ bFlipVs)
+	//{
+	//	CurveSweepOpOut.UVScale[1] *= -1;
+	//	CurveSweepOpOut.UVOffset = FVector2d(0, UVScale);
+	//}
+
+	bool bUVsSkipFullyWeldedEdges = true;
+	ECapFillMode CapFillMode = (bFillPartialRevolveEndcaps) ?
+		ECapFillMode::EarClipping : ECapFillMode::None;		// delaunay? hits ensure...
+
+															// Use a mesh generator to do the work.
+	FProfileSweepGenerator CurveSweeper;
+	CurveSweeper.ProfileCurve = ProfileCurve;
+	CurveSweeper.SweepCurve = SweepCurve;
+	CurveSweeper.WeldedVertices = ProfileVerticesToWeld;
+	CurveSweeper.bProfileCurveIsClosed = bProfileCurveIsClosed;
+	CurveSweeper.bSweepCurveIsClosed = bSweepCurveIsClosed;
+
+	CurveSweeper.bSharpNormals = false; // bSharpNormals;
+	CurveSweeper.UVScale = UVScale;
+	CurveSweeper.bUVScaleRelativeWorld = bUVScaleRelativeWorld;
+	CurveSweeper.UnitUVInWorldCoordinates = UnitUVInWorldCoordinates;
+	CurveSweeper.UVOffset = UVOffset;
+	CurveSweeper.bUVsSkipFullyWeldedEdges = bUVsSkipFullyWeldedEdges;
+	CurveSweeper.PolygonGroupingMode = PolygonGroupingMode;
+	CurveSweeper.QuadSplitMethod = QuadSplitMethod;
+	CurveSweeper.DiagonalTolerance = DiagonalTolerance;
+
+	FDynamicMesh3 ResultMesh(&CurveSweeper.Generate());
+
+	// Cap the ends if needed
+	if (bProfileCurveIsClosed && !bSweepCurveIsClosed && CapFillMode != ECapFillMode::None)
+	{
+		for (int CapIndex = 0; CapIndex < 2; ++CapIndex)
+		{
+			// Note that we rely on the sweep generator to tell us the start and end profile curves for
+			// capping rather than looking for a boundary ourselves. This is because in cases where the
+			// side of the profile curve is welded, there will actually be one open boundary rather than
+			// two separate ones, yet we want to cap the two sides separately.
+			// This does rely on the fact that the hole fillers we use only care about vertices, and don't
+			// rely on there being edges between them. It also relies on the vertex ids in the DynamicMesh
+			// matching the ones used in the generator. Without those two facts, the work of filling two
+			// halves of an open boundary separately would be harder to do...
+
+			// TODO: curve is planar we don't have to figure this out...
+
+			// Put the profile curve into an edge loop
+			TArray<int> VertexLoop = CurveSweeper.EndProfiles[CapIndex];
+			int NV = VertexLoop.Num();
+			TArray<int32> Edges;
+			Edges.SetNum(NV);
+			for (int i = 0; i < NV; ++i)
+			{
+				Edges[i] = ResultMesh.FindEdge(VertexLoop[i], VertexLoop[(i + 1) % NV]);
+			}
+
+			//FEdgeLoop::VertexLoopToEdgeLoop(ResultMesh.Get(), , Edges);
+			//FEdgeLoop Loop(ResultMesh.Get(), CurveSweeper.EndProfiles[CapIndex], Edges);
+
+			// The first cap actually needs to face in the opposite direction of the sweep to face outward.
+			if (CapIndex == 0)
+			{
+				Algo::Reverse(VertexLoop);
+				Algo::Reverse(Edges);
+			}
+
+			// Compute a best-fit plane of the vertices
+			TArray<FVector3<double>> VertexPositions;
+			for (int32 vid : VertexLoop)
+			{
+				VertexPositions.Add( ResultMesh.GetVertex(vid) );
+			}
+
+			FVector3<double> PlaneOrigin;
+			FVector3<double> PlaneNormal;
+			PolygonTriangulation::ComputePolygonPlane<double>(VertexPositions, PlaneNormal, PlaneOrigin);
+			PlaneNormal *= -1.0;	// Previous function seems to orient the normal opposite to what's expected elsewhere
+
+			int32 NewGroupID = ResultMesh.AllocateTriangleGroup();
+
+			TArray<FIndex3i> Triangles;
+			TArray<int32> NewTriangles;
+			PolygonTriangulation::TriangulateSimplePolygon(VertexPositions, Triangles);
+			for (FIndex3i PolyTriangle : Triangles)
+			{
+				FIndex3i MeshTriangle( VertexLoop[PolyTriangle.A], VertexLoop[PolyTriangle.C], VertexLoop[PolyTriangle.B]);  // Reversing orientation here!!
+				int32 NewTriangle = ResultMesh.AppendTriangle(MeshTriangle, NewGroupID);
+				if (NewTriangle >= 0)
+				{
+					NewTriangles.Add(NewTriangle);
+				}
+			}
+
+			// Compute normals and UVs
+			if (ResultMesh.HasAttributes())
+			{
+				UELocal::SetTriangleNormals(&ResultMesh, NewTriangles);
+
+				// Set UVs for the cap
+				FFrame3d ProjectionFrame(PlaneOrigin, PlaneNormal);
+				if (bUVScaleRelativeWorld)
+				{
+					FVector2f ScaleToUse(UVScale);
+					ScaleToUse /= CurveSweeper.UnitUVInWorldCoordinates;
+					UELocal::SetTriangleUVsFromProjection(&ResultMesh, NewTriangles, ProjectionFrame, ScaleToUse, (FVector2f)UVOffset,
+						0, // layer
+						true,  // shift to origin
+						false); // don't normalize
+				}
+				else
+				{
+					UELocal::SetTriangleUVsFromProjection(&ResultMesh, NewTriangles, ProjectionFrame, (FVector2f)UVScale, (FVector2f)UVOffset,
+						0, // layer
+						true,  // shift to origin
+						true); // normalize first
+				}
+
+			}
+		}// end for each loop
+	}//end if capping
+
+
+	return ResultMesh;
+}
diff --git a/Engine/Source/Runtime/GeometryCore/Public/Generators/RevolveGenerator.h b/Engine/Source/Runtime/GeometryCore/Public/Generators/RevolveGenerator.h
new file mode 100644
--- /dev/null
+++ b/Engine/Source/Runtime/GeometryCore/Public/Generators/RevolveGenerator.h
@@ -0,0 +1,146 @@
+// Copyright Epic Games, Inc. All Rights Reserved.
+
+#pragma once
+
+#include "CoreMinimal.h"
+
+#include "MeshShapeGenerator.h"
+#include "Generators/SweepGenerator.h"
+#include "DynamicMesh/DynamicMesh3.h"
+
+namespace UE
+{
+namespace Geometry
+{
+
+
+/**
+ * FBaseRevolveGenerator is a base class for simple surface-of-revolution generators.
+ * Internally these are implemented via FSweepGenerator, but the Revolve generators provide
+ * simpler-to-use interfaces and handle some other aspects like endcapping/etc.
+ * 
+ * Revolve generators currently depend on the ability to use DynamicMesh operations to (eg) triangulate endcaps/etc,
+ * and so they are not currently FMeshShapeGenerators (this may be resolved in future)
+ */
+class GEOMETRYCORE_API FBaseRevolveGenerator
+{
+public:
+
+	enum class ECapFillMode
+	{
+		None,
+		EarClipping,
+	};
+
+	/** Number of discrete steps used to sample the revolution path */
+	int32 Steps = 16;
+
+	/** Degrees of the revolution path (eg circle for a regular revolution). Maybe clamped in some cases. */
+	float RevolveDegrees = 360.0f;
+
+	/** offset for revolution degrees, eg to shift partial revolution around the revolve circle */
+	float DegreeOffset = 0.0f;
+
+	/** Reverse the direction of the revolution, for non-closed revolutions */
+	bool bReverseDirection = false;
+
+	/**  */
+	bool bProfileAtMidpoint = false;
+
+	/** Whether or not to fill the open ends of a partial revolve, if possible */
+	bool bFillPartialRevolveEndcaps = true;
+
+	// Generated UV coordinates will be multiplied by these values.
+	FVector2d UVScale = FVector2d(1,1);
+
+	// These values will be added to the generated UV coordinates after applying UVScale.
+	FVector2d UVOffset = FVector2d(0, 0);
+
+	// When true, the generator attempts to scale UV's in a way that preserves scaling across different mesh
+	// results, aiming for 1.0 in UV space to be equal to UnitUVInWorldCoordinates in world space. This is 
+	// generally speaking unrealistic because UV's are going to be variably stretched no matter what, but 
+	// in practice it means adjusting the V scale relative to the profile curve length and U scale relative
+	// to a very crude measurement of movement across sweep frames.
+	bool bUVScaleRelativeWorld = false;
+
+	// Only relevant if bUVScaleRelativeWorld is true (see that description)
+	float UnitUVInWorldCoordinates = 100;
+
+	/** QuadSplitMethod controls how quads are split into triangles (see enum comments) */
+	EProfileSweepQuadSplit QuadSplitMethod = EProfileSweepQuadSplit::Uniform;
+
+	// When QuadSplitMode is ShortestDiagonal, biases one of the diagonals so that symmetric
+	// quads are split uniformly. The tolerance is a proportion allowable difference.
+	double DiagonalTolerance = 0.01;
+
+	/** PolygonGroupingMode controls how polygroups are assigned on the output mesh */
+	EProfileSweepPolygonGrouping PolygonGroupingMode = EProfileSweepPolygonGrouping::PerFace;
+
+
+};
+
+
+
+/**
+ * Revolve a planar polyline (in the XY plane) around the +Z axis along a circular path
+ * (+X, +Y) in the 2D polyline is mapped to (+X, 0, +Z)
+ */
+class GEOMETRYCORE_API FRevolvePlanarPathGenerator : public FBaseRevolveGenerator
+{
+public:
+	/** Vertices of the planar polyline/path that will be revolved */
+	TArray<FVector2d> PathVertices;
+
+	/** If true, open tops/bottoms of the revolved shape are capped */
+	bool bCapped = true;
+
+	/** Generate the mesh */
+	FDynamicMesh3 GenerateMesh();
+
+};
+
+/**
+* Revolve a planar polygon (in the XY plane) around the +Z axis along a circular path
+* (+X, +Y) in the 2D polygon is mapped to (+X, 0, +Z)
+*/
+class GEOMETRYCORE_API FRevolvePlanarPolygonGenerator : public FBaseRevolveGenerator
+{
+public:
+	/** Vertices of the planar polygon that will be revolved */
+	TArray<FVector2d> PolygonVertices;
+
+	/** If true, vertices that lie on the +Z axis will be welded */
+	bool bWeldVertsOnAxis = false;
+
+	/** Generate the mesh */
+	FDynamicMesh3 GenerateMesh();
+};
+
+
+/**
+* Revolve a planar polygon (in the XY plane) around the +Z axis along a spiral path.
+* (+X, +Y) in the 2D polygon is mapped to (+X, 0, +Z)
+*/
+class GEOMETRYCORE_API FSpiralRevolvePlanarPolygonGenerator : public FBaseRevolveGenerator
+{
+public:
+	/** Vertices of the planar polygon that will be revolved */
+	TArray<FVector2d> PolygonVertices;
+
+	/** If true, vertices that lie on the +Z axis will be welded */
+	bool bWeldVertsOnAxis = false;
+
+	/** Vertical rise of the spiral path per full 360-degree revolution */
+	double RisePerFullRevolution = 50;
+
+	/** Generate the mesh */
+	FDynamicMesh3 GenerateMesh();
+};
+
+
+
+
+
+} // end namespace UE::Geometry
+}
+
