GeometryCore: Improved serialization for dynamic mesh data

This CL adds improved serialization for FDynamicMesh3 and its data structures:
- FDynamicMesh3, TDynamicVector, FRefCountVector, FSmallListSet, TDynamicMeshOverlay, ...
- Versioning using the UE5Main stream version
- Discarding redundant data
- Oodle compression
- Serialization of vertex skin weights
- Additional unit tests and benchmark

diff --git a/Engine/Source/Runtime/Core/Public/UObject/UE5MainStreamObjectVersion.h b/Engine/Source/Runtime/Core/Public/UObject/UE5MainStreamObjectVersion.h
--- a/Engine/Source/Runtime/Core/Public/UObject/UE5MainStreamObjectVersion.h
+++ b/Engine/Source/Runtime/Core/Public/UObject/UE5MainStreamObjectVersion.h
@@ -178,6 +178,9 @@ struct CORE_API FUE5MainStreamObjectVersion
 		// Add generic tagging of all anim graph nodes in anim blueprints
 		AnimGraphNodeTaggingAdded,
 		
+		// Add custom version to FDynamicMesh3
+		DynamicMeshCompactedSerialization,
+		
 		// -----<new versions can be added above this line>-------------------------------------------------
 		VersionPlusOne,
 		LatestVersion = VersionPlusOne - 1
diff --git a/Engine/Source/Runtime/GeometryCore/Private/DynamicMesh/DynamicMesh3.cpp b/Engine/Source/Runtime/GeometryCore/Private/DynamicMesh/DynamicMesh3.cpp
--- a/Engine/Source/Runtime/GeometryCore/Private/DynamicMesh/DynamicMesh3.cpp
+++ b/Engine/Source/Runtime/GeometryCore/Private/DynamicMesh/DynamicMesh3.cpp
@@ -607,41 +607,140 @@ FString FDynamicMesh3::MeshInfoString()
 
 bool FDynamicMesh3::IsSameAs(const FDynamicMesh3& m2, const FSameAsOptions& Options) const
 {
+	auto SameVertex = [this, &m2, &Options](const int Vid, const int VidM2)
+	{
+		return VectorUtil::EpsilonEqual(GetVertexRef(Vid), m2.GetVertexRef(VidM2), static_cast<double>(Options.Epsilon));
+	};
+
 	if (VertexCount() != m2.VertexCount())
 	{
 		return false;
 	}
+
 	if (TriangleCount() != m2.TriangleCount())
 	{
 		return false;
 	}
-	for (int vid : VertexIndicesItr())
+
+	TOptional<TDynamicVector<int>> VidMapping;
+	TOptional<TDynamicVector<int>> TidMapping;
+
+	if (!Options.bIgnoreDataLayout || (VertexRefCounts.IsDense() && m2.VertexRefCounts.IsDense() && TriangleRefCounts.IsDense() && m2.TriangleRefCounts.IsDense()))
 	{
-		if (m2.IsVertex(vid) == false || VectorUtil::EpsilonEqual(GetVertex(vid), m2.GetVertex(vid), (double)Options.Epsilon) == false)
+		for (const int Vid : VertexIndicesItr())
 		{
-			return false;
+			if (m2.IsVertex(Vid) == false || !SameVertex(Vid, Vid))
+			{
+				return false;
+			}
+		}
+
+		for (const int Tid : TriangleIndicesItr())
+		{
+			if (m2.IsTriangle(Tid) == false || (GetTriangle(Tid) != m2.GetTriangle(Tid)))
+			{
+				return false;
+			}
 		}
 	}
-	for (int tid : TriangleIndicesItr())
+	else
 	{
-		if (m2.IsTriangle(tid) == false || (GetTriangle(tid) != m2.GetTriangle(tid)))
+		// Ignore holes in vertex and triangle data, i.e. don't use specific IDs but still make sure they are still in the same order.
+
+		// Vertices
 		{
-			return false;
+			FRefCountVector::IndexIterator ItVid = VertexRefCounts.BeginIndices();
+			const FRefCountVector::IndexIterator ItEndVid = VertexRefCounts.EndIndices();
+			FRefCountVector::IndexIterator ItVidM2 = m2.VertexRefCounts.BeginIndices();
+			const FRefCountVector::IndexIterator ItEndVidM2 = m2.VertexRefCounts.EndIndices();
+
+			VidMapping.Emplace();
+			VidMapping->Resize(Vertices.Num(), InvalidID);
+
+			while (ItVid != ItEndVid && ItVidM2 != ItEndVidM2)
+			{
+				(*VidMapping)[*ItVid] = *ItVidM2;
+
+				if (!SameVertex(*ItVid, *ItVidM2))
+				{
+					// Vertices are not the same.
+					return false;
+				}
+			
+				++ItVid;
+				++ItVidM2;
+			}
+		
+			if (ItVid != ItEndVid || ItVidM2 != ItEndVidM2)
+			{
+				// Number of vertices is not the same.
+				return false;
+			}
+		}
+
+		// Triangles
+		{
+			FRefCountVector::IndexIterator ItTid = TriangleRefCounts.BeginIndices();
+			const FRefCountVector::IndexIterator ItEndTid = TriangleRefCounts.EndIndices();
+			FRefCountVector::IndexIterator ItTidM2 = m2.TriangleRefCounts.BeginIndices();
+			const FRefCountVector::IndexIterator ItEndTidM2 = m2.TriangleRefCounts.EndIndices();
+
+			TidMapping.Emplace();
+			TidMapping->Resize(Triangles.Num(), InvalidID);
+		
+			while (ItTid != ItEndTid && ItTidM2 != ItEndTidM2)
+			{
+				const FIndex3i Tri = GetTriangle(*ItTid);
+				const FIndex3i TriM2 = m2.GetTriangle(*ItTidM2);
+
+				(*TidMapping)[*ItTid] = *ItTidM2;
+
+				for (int i = 0; i < 3; ++i)
+				{
+					if (TriM2[i] != (VidMapping ? (*VidMapping)[Tri[i]] : Tri[i]))
+					{
+						// Triangle vertices are not the same.
+						return false;
+					}
+				}
+
+				++ItTid;
+				++ItTidM2;
+			}
+
+			if (ItTid != ItEndTid || ItTidM2 != ItEndTidM2)
+			{
+				// Number of triangles is not the same.
+				return false;
+			}
 		}
 	}
+
 	if (Options.bCheckConnectivity)
 	{
+		if (EdgeCount() != m2.EdgeCount())
+		{
+			return false;
+		}
 		for (int eid : EdgeIndicesItr())
 		{
-			FEdge e = GetEdge(eid);
-			int other_eid = m2.FindEdge(e.Vert[0], e.Vert[1]);
+			const FEdge e = GetEdge(eid);
+			const int Vert0 = VidMapping ? (*VidMapping)[e.Vert[0]] : e.Vert[0];
+			const int Vert1 = VidMapping ? (*VidMapping)[e.Vert[1]] : e.Vert[1];
+			const int other_eid = m2.FindEdge(Vert0, Vert1);
 			if (other_eid == InvalidID)
 			{
 				return false;
 			}
-			FEdge oe = m2.GetEdge(other_eid);
-			if (FMath::Min(e.Tri[0], e.Tri[1]) != FMath::Min(oe.Tri[0], oe.Tri[1]) ||
-			    FMath::Max(e.Tri[0], e.Tri[1]) != FMath::Max(oe.Tri[0], oe.Tri[1]))
+			const FEdge oe = m2.GetEdge(other_eid);
+			const FIndex2i eTri = [&e, &TidMapping]
+			{
+				return !TidMapping
+					       ? e.Tri
+					       : FIndex2i{e.Tri[0] != InvalidID ? (*TidMapping)[e.Tri[0]] : InvalidID, e.Tri[1] != InvalidID ? (*TidMapping)[e.Tri[1]] : InvalidID};
+			}();
+			if (FMath::Min(eTri[0], eTri[1]) != FMath::Min(oe.Tri[0], oe.Tri[1]) ||
+			    FMath::Max(eTri[0], eTri[1]) != FMath::Max(oe.Tri[0], oe.Tri[1]))
 			{
 				return false;
 			}
@@ -653,9 +752,19 @@ bool FDynamicMesh3::IsSameAs(const FDynamicMesh3& m2, const FSameAsOptions& Opti
 		{
 			return false;
 		}
-		for (int eid : EdgeIndicesItr())
+		for (const int Eid : EdgeIndicesItr())
 		{
-			if (m2.IsEdge(eid) == false || GetEdge(eid) != m2.GetEdge(eid))
+			if (m2.IsEdge(Eid) == false)
+			{
+				return false;
+			}
+			FEdge Edge = GetEdge(Eid);
+			if (TidMapping.IsSet())
+			{
+				Edge.Tri[0] = Edge.Tri[0] != InvalidID ? (*TidMapping)[Edge.Tri[0]] : InvalidID;
+				Edge.Tri[1] = Edge.Tri[1] != InvalidID ? (*TidMapping)[Edge.Tri[1]] : InvalidID;
+			}
+			if (Edge != m2.GetEdge(Eid))
 			{
 				return false;
 			}
@@ -671,7 +780,7 @@ bool FDynamicMesh3::IsSameAs(const FDynamicMesh3& m2, const FSameAsOptions& Opti
 		{
 			for (int vid : VertexIndicesItr())
 			{
-				if (VectorUtil::EpsilonEqual(GetVertexNormal(vid), m2.GetVertexNormal(vid), Options.Epsilon) == false)
+				if (VectorUtil::EpsilonEqual(GetVertexNormal(vid), m2.GetVertexNormal(VidMapping ? (*VidMapping)[vid] : vid), Options.Epsilon) == false)
 				{
 					return false;
 				}
@@ -688,7 +797,7 @@ bool FDynamicMesh3::IsSameAs(const FDynamicMesh3& m2, const FSameAsOptions& Opti
 		{
 			for (int vid : VertexIndicesItr())
 			{
-				if (VectorUtil::EpsilonEqual(GetVertexColor(vid), m2.GetVertexColor(vid), Options.Epsilon) == false)
+				if (VectorUtil::EpsilonEqual(GetVertexColor(vid), m2.GetVertexColor(VidMapping ? (*VidMapping)[vid] : vid), Options.Epsilon) == false)
 				{
 					return false;
 				}
@@ -705,7 +814,7 @@ bool FDynamicMesh3::IsSameAs(const FDynamicMesh3& m2, const FSameAsOptions& Opti
 		{
 			for (int vid : VertexIndicesItr())
 			{
-				if (VectorUtil::EpsilonEqual(GetVertexUV(vid), m2.GetVertexUV(vid), Options.Epsilon) == false)
+				if (VectorUtil::EpsilonEqual(GetVertexUV(vid), m2.GetVertexUV(VidMapping ? (*VidMapping)[vid] : vid), Options.Epsilon) == false)
 				{
 					return false;
 				}
@@ -720,9 +829,10 @@ bool FDynamicMesh3::IsSameAs(const FDynamicMesh3& m2, const FSameAsOptions& Opti
 		}
 		if (HasTriangleGroups())
 		{
-			for (int tid : TriangleIndicesItr())
+			for (int Tid : TriangleIndicesItr())
 			{
-				if (GetTriangleGroup(tid) != m2.GetTriangleGroup(tid))
+				const int TidM2 = TidMapping ? (*TidMapping)[Tid] : Tid;
+				if (GetTriangleGroup(Tid) != m2.GetTriangleGroup(TidM2))
 				{
 					return false;
 				}
@@ -737,7 +847,7 @@ bool FDynamicMesh3::IsSameAs(const FDynamicMesh3& m2, const FSameAsOptions& Opti
 		}
 		if (HasAttributes())
 		{
-			if (!AttributeSet->IsSameAs(*m2.AttributeSet))
+			if (!AttributeSet->IsSameAs(*m2.AttributeSet, Options.bIgnoreDataLayout))
 			{
 				return false;
 			}
@@ -761,7 +871,7 @@ bool FDynamicMesh3::CheckValidity(FValidityOptions ValidityOptions, EValidityChe
 	{
 		CheckOrFailF = [&](bool b)
 		{
-			checkf(b, TEXT("FEdgeLoop::CheckValidity failed!"));
+			checkf(b, TEXT("FDynamicMesh3::CheckValidity failed!"));
 			is_ok = is_ok && b;
 		};
 	}
@@ -769,7 +879,7 @@ bool FDynamicMesh3::CheckValidity(FValidityOptions ValidityOptions, EValidityChe
 	{
 		CheckOrFailF = [&](bool b)
 		{
-			ensureMsgf(b, TEXT("FEdgeLoop::CheckValidity failed!"));
+			ensureMsgf(b, TEXT("FDynamicMesh3::CheckValidity failed!"));
 			is_ok = is_ok && b;
 		};
 	}
diff --git a/Engine/Source/Runtime/GeometryCore/Private/DynamicMesh/DynamicMesh3_Serialization.cpp b/Engine/Source/Runtime/GeometryCore/Private/DynamicMesh/DynamicMesh3_Serialization.cpp
--- a/Engine/Source/Runtime/GeometryCore/Private/DynamicMesh/DynamicMesh3_Serialization.cpp
+++ b/Engine/Source/Runtime/GeometryCore/Private/DynamicMesh/DynamicMesh3_Serialization.cpp
@@ -7,28 +7,286 @@
 #include "ExplicitUseGeometryMathTypes.h"		// using UE::Geometry::(math types)
 using namespace UE::Geometry;
 
-enum class EDynamicMeshSerializationVersion
+namespace FDynamicMesh3Serialization_Local
+{
+enum EDynamicMeshSerializationVersion
 {
 	InitialVersion = 1,
+	CompactAndCompress = 2,
 
 	// ----- new versions to be added above this line -------------------------------------------------
 	VersionPlusOne,
 	LatestVersion = VersionPlusOne - 1
 };
 
-int FDynamicMesh3::SerializeInternal_LatestVersion()
+// Encapsulates our serialization options, and selects the serialization variant for a given set of options. 
+struct FDynamicMesh3SerializationOptions
+{
+	bool bPreserveDataLayout; //< Preserve the data layout, i.e. external vertex/triangle/edge indices are still valid after roundtrip serialization.
+	bool bCompactData; //< Remove any holes or padding in the data layout, and discard/recompute any redundant data. 
+	bool bUseCompression; //< Compress all data buffers to minimize memory footprint.
+
+	explicit FDynamicMesh3SerializationOptions(const FArchive& Ar)
+		: bPreserveDataLayout(Ar.IsTransacting())
+		, bCompactData(!Ar.IsTransacting())
+		, bUseCompression(Ar.IsPersistent())
+	{
+		checkSlow(bPreserveDataLayout != bCompactData); // Preserving the data layout and compacting the data are mutually exclusive.
+	}
+
+	enum EImplementationVariant
+	{
+		Default = 0,
+		CompactData = 1 << 16
+	};
+
+	int ImplementationVariant() const
+	{
+		return bCompactData * CompactData;
+	}
+
+	friend FArchive& operator<<(FArchive& Ar, FDynamicMesh3SerializationOptions& Options)
+	{
+		Ar << Options.bPreserveDataLayout;
+		Ar << Options.bCompactData;
+		Ar << Options.bUseCompression;
+		return Ar;
+	}
+};
+
+// Serializes a generic dynamic vector containing bulk data.
+template <typename T>
+void SerializeVector(FArchive& Ar, TDynamicVector<T>& Vector, const FDynamicMesh3SerializationOptions& Options)
+{
+	if (Options.bUseCompression)
+	{
+		Vector.template Serialize<true /* bForceBulkSerialization */, true /* bUseCompression */>(Ar);
+	}
+	else
+	{
+		Vector.template Serialize<true /* bForceBulkSerialization */, false /* bUseCompression */>(Ar);
+	}
+}
+
+// Serializes an optional generic dynamic vector containing bulk data.
+// We need to do our own serialization since using the default serialization for TOptional would not allow us to customize the dynamic vector serialization.
+template <typename T>
+void SerializeOptionalVector(FArchive& Ar, TOptional<TDynamicVector<T>>& OptionalVector, const FDynamicMesh3SerializationOptions& Options)
+{
+	bool bHasOptionalVector = OptionalVector.IsSet();
+	Ar << bHasOptionalVector;
+	if (bHasOptionalVector)
+	{
+		if (Ar.IsLoading())
+		{
+			OptionalVector.Emplace();
+		}
+		SerializeVector(Ar, OptionalVector.GetValue(), Options);
+	}
+}
+
+// Serializes all unique data related to vertices.
+void SerializeUniqueVertexData(FArchive& Ar, TDynamicVector<FVector3d>& Vertices, TOptional<TDynamicVector<FVector3f>>& VertexNormals,
+							TOptional<TDynamicVector<FVector3f>>& VertexColors, TOptional<TDynamicVector<FVector2f>>& VertexUVs,
+							const FDynamicMesh3SerializationOptions& Options)
+{
+	SerializeVector(Ar, Vertices, Options);
+	SerializeOptionalVector(Ar, VertexNormals, Options);
+	SerializeOptionalVector(Ar, VertexColors, Options);
+	SerializeOptionalVector(Ar, VertexUVs, Options);
+}
+
+// Serializes all unique data related to triangles.
+void SerializeUniqueTriangleData(FArchive& Ar, TDynamicVector<FIndex3i>& Triangles, TOptional<TDynamicVector<int32>>& TriangleGroups, int&
+                                 GroupIDCounter,
+                                 const FDynamicMesh3SerializationOptions& Options)
+{
+	SerializeVector(Ar, Triangles, Options);
+	SerializeOptionalVector(Ar, TriangleGroups, Options);
+	Ar << GroupIDCounter;
+}
+
+// Serializes a ref count vector.
+void SerializeRefCounts(FArchive& Ar, FRefCountVector& RefCounts, const FDynamicMesh3SerializationOptions& Options)
+{
+	RefCounts.Serialize(Ar, Options.bCompactData, Options.bUseCompression);
+}
+
+// Resets a ref count vector for dense triangle data, i.e. all ref count values are 1.
+void ResetDenseTriangleRefCounts(size_t Num, FRefCountVector& TriangleRefCounts)
 {
-	return static_cast<int>(EDynamicMeshSerializationVersion::LatestVersion);
+	TriangleRefCounts.Trim(Num);
+	TriangleRefCounts.GetRawRefCountsUnsafe().Fill(1);
+}
+
+// Recomputes vertex ref counts for a given set of triangle indices.
+void RecomputeVertexRefCounts(size_t NumVertices, const TDynamicVector<FIndex3i>& Triangles, const FRefCountVector& TriangleRefCounts,
+                              FRefCountVector& VertexRefCounts)
+{
+	// Callable that iterates over all triangles.
+	// This gets called from within Rebuild() with another callable as parameter that we in turn call to update the ref count for each vertex in our triangles.
+	const auto IterateTriangles = [&Triangles, &TriangleRefCounts](auto&& UpdateVertexRefCount)
+	{
+		for (const int TriIndex : TriangleRefCounts.Indices())
+		{
+			const FIndex3i& Tri = Triangles[TriIndex];
+			UpdateVertexRefCount(Tri[0]);
+			UpdateVertexRefCount(Tri[1]);
+			UpdateVertexRefCount(Tri[2]);
+		}
+	};
+
+	// Callable that allocates the initial ref count value for a vertex. This gets called from within Rebuild().
+	constexpr auto AllocateRefCountFunc = [](unsigned short& RefCount)
+	{
+		RefCount = 2;
+	};
+
+	// Callable that increments the ref count value for a vertex. This gets called from within Rebuild().
+	constexpr auto IncrementRefCountFunc = [](unsigned short& RefCount)
+	{
+		++RefCount;
+	};
+
+	// Rebuild function using our callables.
+	VertexRefCounts.Rebuild(NumVertices, IterateTriangles, AllocateRefCountFunc, IncrementRefCountFunc);
+}
+
+// Serializes a small list set.
+void SerializeSmallListSet(FArchive& Ar, FSmallListSet& SmallListSet, const FDynamicMesh3SerializationOptions& Options)
+{
+	SmallListSet.Serialize(Ar, Options.bCompactData, Options.bUseCompression);
+}
+
+// Recomputes all the redundant edge data for given sets of vertices and triangles.
+template <typename FindEdgeFunc, typename AddEdgeInternalFunc>
+void RecomputeEdgeData(const TDynamicVector<FVector3d>& Vertices, const TDynamicVector<FIndex3i>& Triangles, const FRefCountVector& TriangleRefCounts,
+                       TDynamicVector<FDynamicMesh3::FEdge>& Edges, FRefCountVector& EdgeRefCounts, TDynamicVector<FIndex3i>& TriangleEdges,
+                       FSmallListSet& VertexEdgeLists, FindEdgeFunc&& FindEdge, AddEdgeInternalFunc&& AddEdgeInternal)
+{
+	Edges.Clear();
+	EdgeRefCounts.Clear();
+
+	TriangleEdges.Resize(Triangles.Num());
+	TriangleEdges.Fill({FDynamicMesh3::InvalidID, FDynamicMesh3::InvalidID, FDynamicMesh3::InvalidID});
+
+	VertexEdgeLists.Reset();
+	VertexEdgeLists.Resize(Vertices.Num());
+
+	auto AddEdge = [&Edges, &FindEdge, &AddEdgeInternal](const int VertA, const int VertB, const int Tid) -> int
+	{
+		int Eid = Forward<FindEdgeFunc>(FindEdge)(VertA, VertB);
+		if (Eid == FDynamicMesh3::InvalidID)
+		{
+			Eid = Forward<AddEdgeInternalFunc>(AddEdgeInternal)(VertA, VertB, Tid);
+		}
+		else
+		{
+			Edges[Eid].Tri[1] = Tid;
+		}
+		return Eid;
+	};
+
+	for (const int Tid : TriangleRefCounts.Indices())
+	{
+		const FIndex3i& Tri = Triangles[Tid];
+		const FIndex3i TriEdges = {
+			AddEdge(Tri[0], Tri[1], Tid),
+			AddEdge(Tri[1], Tri[2], Tid),
+			AddEdge(Tri[2], Tri[0], Tid)
+		};
+		TriangleEdges[Tid] = TriEdges;
+	}
+}
+
+// Serializes the optional attribute set for a given mesh.
+// We need to do our own serialization since using the default serialization for TUniquePtr would not allow us to customize the attribute set serialization.
+void SerializeAttributeSet(FArchive& Ar, FDynamicMesh3* Mesh, const FCompactMaps* CompactMaps, const FDynamicMesh3SerializationOptions& Options)
+{
+	bool bHasAttributes = Mesh->Attributes() != nullptr;
+	Ar << bHasAttributes;
+	if (bHasAttributes)
+	{
+		if (Ar.IsLoading())
+		{
+			Mesh->EnableAttributes();
+		}
+
+		Mesh->Attributes()->Serialize(Ar, CompactMaps, Options.bUseCompression);
+	}
 }
 
+// Creates a generic optional dynamic vector of a given size if and only if another optional vector is set.
+// This reduces code duplication for creating compacting vectors for optional mesh data such as normals or colors.
+template <typename T>
+TOptional<TDynamicVector<T>> CreateOptionalVector(const TOptional<TDynamicVector<T>>& ExistingOptionalVector, size_t Num)
+{
+	TOptional<TDynamicVector<T>> OptionalVector;
+	if (ExistingOptionalVector.IsSet())
+	{
+		OptionalVector.Emplace();
+		OptionalVector->SetNum(Num);
+	}
+	return MoveTemp(OptionalVector);
+};
+} // namespace FDynamicMesh3Serialization_Local
+
+using namespace FDynamicMesh3Serialization_Local;
+
+// Declaration of all variants of our serialization. The integer index is derived from the version an serialization options, if available.
+template <> void FDynamicMesh3::SerializeInternal<InitialVersion>(FArchive&, void*);
+template <> void FDynamicMesh3::SerializeInternal<CompactAndCompress + FDynamicMesh3SerializationOptions::Default>(FArchive&, void*);
+template <> void FDynamicMesh3::SerializeInternal<CompactAndCompress + FDynamicMesh3SerializationOptions::CompactData>(FArchive&, void*);
+
 void FDynamicMesh3::Serialize(FArchive& Ar)
 {
-	// Check validity before saving data.
+	const double TimeStart = FPlatformTime::Seconds();
+
+	// Skip serialization if the archive is only looking for UObject references.
+	if (Ar.IsObjectReferenceCollector())
+	{
+		return;
+	}
+
+	// Check validity in debug before saving data.
 	checkSlow(Ar.IsLoading() || CheckValidity(FValidityOptions(), EValidityCheckFailMode::Ensure));
 
 	Ar.UsingCustomVersion(FUE5MainStreamObjectVersion::GUID);
+	if (Ar.IsLoading() && Ar.CustomVer(FUE5MainStreamObjectVersion::GUID) < FUE5MainStreamObjectVersion::DynamicMeshCompactedSerialization)
+	{
+		SerializeInternal<InitialVersion>(Ar, nullptr);
+	}
+	else
+	{
+		FDynamicMesh3SerializationOptions Options(Ar);
+		Ar << Options;
+		switch (Options.ImplementationVariant())
+		{
+		case FDynamicMesh3SerializationOptions::Default:
+			SerializeInternal<CompactAndCompress + FDynamicMesh3SerializationOptions::Default>(Ar, &Options);
+			break;
+		case FDynamicMesh3SerializationOptions::CompactData:
+			SerializeInternal<CompactAndCompress + FDynamicMesh3SerializationOptions::CompactData>(Ar, &Options);
+			break;
+		default:
+			ensureAlwaysMsgf(false, TEXT("Unhandled FDynamicMesh3 serialization variant '%d'."), Options.ImplementationVariant());
+		}
+	}
 
-	int32 SerializationVersion = static_cast<int>(EDynamicMeshSerializationVersion::LatestVersion);
+	// Check validity in debug after loading data.
+	checkSlow(!Ar.IsLoading() || CheckValidity(FValidityOptions(), EValidityCheckFailMode::Ensure));
+
+	UE_LOG(LogGeometry, Verbose, TEXT("NEW %s of dynamic mesh took %2.3f ms"), ANSI_TO_TCHAR(Ar.IsLoading() ? "deserialization" : "serialization"),
+	       1000.0f * static_cast<float>(FPlatformTime::Seconds() - TimeStart));
+}
+
+template <>
+void FDynamicMesh3::SerializeInternal<InitialVersion>(FArchive& Ar, void*)
+{
+	// This is not a current variant; only allow loading.
+	ensure(Ar.IsLoading());
+
+	int32 SerializationVersion;
 	Ar << SerializationVersion;
 
 	Ar << Vertices;
@@ -45,17 +303,170 @@ void FDynamicMesh3::Serialize(FArchive& Ar)
 	Ar << Edges;
 	Ar << EdgeRefCounts;
 
-	bool bHasAttributes = HasAttributes();
+	bool bHasAttributes;
 	Ar << bHasAttributes;
 	if (bHasAttributes)
+	{
+		EnableAttributes();
+		Ar << *AttributeSet;
+	}
+}
+
+// This 'CompactAndCompress' variant is relatively fast, but allows for discarding some redundant data.
+// You can chose to preserve the data layout in the serialization options, which will keep any externally held vertex/triangle/edge indices valid.
+template <>
+void FDynamicMesh3::SerializeInternal<CompactAndCompress + FDynamicMesh3SerializationOptions::Default>(FArchive& Ar, void* OptionsPtr)
+{
+	checkSlow(OptionsPtr != nullptr);
+	const FDynamicMesh3SerializationOptions& Options = *static_cast<FDynamicMesh3SerializationOptions*>(OptionsPtr);
+
+	// Vertex data.
+	SerializeUniqueVertexData(Ar, Vertices, VertexNormals, VertexColors, VertexUVs, Options);
+
+	// Triangle data.
+	SerializeUniqueTriangleData(Ar, Triangles, TriangleGroups, GroupIDCounter, Options);
+
+	// Triangle ref counts.
+	// If the triangle vector has holes then we simply need to serialize the ref count data. If all allocated triangles are used, we can easily restore a ref
+	// count vector containing all ones, which is not slower than naively serializing the data but saves some memory.
+	bool bTriangleVectorHasHoles = !TriangleRefCounts.IsDense();
+	Ar << bTriangleVectorHasHoles;
+	if (bTriangleVectorHasHoles)
+	{
+		SerializeRefCounts(Ar, TriangleRefCounts, Options);
+	}
+	else if (Ar.IsLoading())
+	{
+		ResetDenseTriangleRefCounts(Triangles.Num(), TriangleRefCounts);
+	}
+	
+	// Vertex ref counts.
+	// The vertex ref counts are redundant to the triangle data, and hence we can recompute it even if we need to preserve the data layout.
+	// Considering that recomputing it is about as fast as serializing it, we always recompute it to save memory.
+	if (Ar.IsLoading())
+	{
+		RecomputeVertexRefCounts(Vertices.Num(), Triangles, TriangleRefCounts, VertexRefCounts);
+	}
+
+	// Edge data.
+	// All edge data can be derived from the triangle and vertex data. If we need to preserve the data layout though, we simply serialize everything as is.
+	if (Options.bPreserveDataLayout)
+	{
+		SerializeVector(Ar, Edges, Options);
+		SerializeRefCounts(Ar, EdgeRefCounts, Options);
+		SerializeVector(Ar, TriangleEdges, Options);
+		SerializeSmallListSet(Ar, VertexEdgeLists, Options);
+	}
+	else
 	{
 		if (Ar.IsLoading())
 		{
-			EnableAttributes();
+			RecomputeEdgeData(Vertices, Triangles, TriangleRefCounts, Edges, EdgeRefCounts, TriangleEdges, VertexEdgeLists,
+			                  [this](int VertA, int VertB) -> int { return FindEdge(VertA, VertB); },
+			                  [this](int VertA, int VertB, int Tid) -> int { return AddEdgeInternal(VertA, VertB, Tid); });
 		}
-		Ar << *AttributeSet;
 	}
 
-	// Check validity after loading data.
-	checkSlow(!Ar.IsLoading() || CheckValidity(FValidityOptions(), EValidityCheckFailMode::Ensure));
-}
\ No newline at end of file
+	// Attribute set data.
+	SerializeAttributeSet(Ar, this, nullptr, Options);
+}
+
+// This 'CompactAndCompress' variant will significantly reduce the memory footprint of the serialized data by compacting the unique data, and disregarding any
+// redundant data, e.g. ref counts and edges. Optionally, it can also compress all data using Oodle, although at a significant runtime cost. 
+// Note that this will very likely invalidate externally held vertex/triangle/edge indices.
+template <>
+void FDynamicMesh3::SerializeInternal<CompactAndCompress + FDynamicMesh3SerializationOptions::CompactData>(FArchive& Ar, void* OptionsPtr)
+{
+	checkSlow(OptionsPtr != nullptr);
+	const FDynamicMesh3SerializationOptions& Options = *static_cast<FDynamicMesh3SerializationOptions*>(OptionsPtr);
+
+	FCompactMaps CompactMaps;
+
+	// Vertex data.
+	if (Ar.IsLoading() || (!Ar.IsLoading() && VertexRefCounts.IsDense()))
+	{
+		SerializeUniqueVertexData(Ar, Vertices, VertexNormals, VertexColors, VertexUVs, Options);
+	}
+	else
+	{
+		const size_t NumValidVertices = VertexRefCounts.GetCount();
+
+		TDynamicVector<FVector3d> VerticesCompact;
+		VerticesCompact.SetNum(NumValidVertices);
+
+		TOptional<TDynamicVector<FVector3f>> VertexNormalsCompact = CreateOptionalVector(VertexNormals, NumValidVertices);
+		TOptional<TDynamicVector<FVector3f>> VertexColorsCompact = CreateOptionalVector(VertexColors, NumValidVertices);
+		TOptional<TDynamicVector<FVector2f>> VertexUVsCompact = CreateOptionalVector(VertexUVs, NumValidVertices);
+
+		CompactMaps.ResetVertexMap(Vertices.Num(), true);
+
+		size_t VidCompact = 0;
+		for (const size_t Vid : VertexRefCounts.Indices())
+		{
+			VerticesCompact[VidCompact] = Vertices[Vid];
+			if (VertexNormalsCompact) { (*VertexNormalsCompact)[VidCompact] = (*VertexNormals)[Vid]; }
+			if (VertexColorsCompact) { (*VertexColorsCompact)[VidCompact] = (*VertexColors)[Vid]; }
+			if (VertexUVsCompact) { (*VertexUVsCompact)[VidCompact] = (*VertexUVs)[Vid]; }
+			CompactMaps.SetVertexMapping(Vid, VidCompact);
+			++VidCompact;
+		}
+
+		SerializeUniqueVertexData(Ar, VerticesCompact, VertexNormalsCompact, VertexColorsCompact, VertexUVsCompact, Options);
+	}
+
+	// Triangle data.
+	if (Ar.IsLoading() || (!Ar.IsLoading() && TriangleRefCounts.IsDense() && VertexRefCounts.IsDense()))
+	{
+		SerializeUniqueTriangleData(Ar, Triangles, TriangleGroups, GroupIDCounter, Options);
+	}
+	else
+	{
+		const size_t NumValidTriangles = TriangleRefCounts.GetCount();
+
+		TDynamicVector<FIndex3i> TrianglesCompact;
+		TrianglesCompact.SetNum(NumValidTriangles);
+
+		TOptional<TDynamicVector<int>> TriangleGroupsCompact = CreateOptionalVector(TriangleGroups, NumValidTriangles);
+
+		CompactMaps.ResetTriangleMap(Triangles.Num(), true);
+
+		const bool bHasVertexMapping = CompactMaps.VertexMapIsSet();
+		
+		size_t TidCompact = 0;
+		for (const size_t Tid : TriangleRefCounts.Indices())
+		{
+			if (!bHasVertexMapping)
+			{
+				TrianglesCompact[TidCompact] = Triangles[Tid];
+			}
+			else
+			{
+				const FIndex3i& Tri = Triangles[Tid];
+				FIndex3i& TriCompact = TrianglesCompact[TidCompact];
+				TriCompact[0] = CompactMaps.GetVertexMapping(Tri[0]);
+				TriCompact[1] = CompactMaps.GetVertexMapping(Tri[1]);
+				TriCompact[2] = CompactMaps.GetVertexMapping(Tri[2]);
+			}
+			if (TriangleGroupsCompact) { (*TriangleGroupsCompact)[TidCompact] = (*TriangleGroups)[Tid]; }
+			CompactMaps.SetTriangleMapping(Tid, TidCompact);
+			++TidCompact;
+		}
+
+		SerializeUniqueTriangleData(Ar, TrianglesCompact, TriangleGroupsCompact, GroupIDCounter, Options);
+	}
+
+	// When loading, restore redundant ref count and edge data.
+	if (Ar.IsLoading())
+	{
+		ResetDenseTriangleRefCounts(Triangles.Num(), TriangleRefCounts);
+
+		RecomputeVertexRefCounts(Vertices.Num(), Triangles, TriangleRefCounts, VertexRefCounts);
+
+		RecomputeEdgeData(Vertices, Triangles, TriangleRefCounts, Edges, EdgeRefCounts, TriangleEdges, VertexEdgeLists,
+		                  [this](int VertA, int VertB) -> int { return FindEdge(VertA, VertB); },
+		                  [this](int VertA, int VertB, int Tid) -> int { return AddEdgeInternal(VertA, VertB, Tid); });
+	}
+
+	// Attribute set data.
+	SerializeAttributeSet(Ar, this, CompactMaps.VertexMapIsSet() || CompactMaps.TriangleMapIsSet() ? &CompactMaps : nullptr, Options);
+}
diff --git a/Engine/Source/Runtime/GeometryCore/Private/DynamicMesh/DynamicMeshAttributeSet.cpp b/Engine/Source/Runtime/GeometryCore/Private/DynamicMesh/DynamicMeshAttributeSet.cpp
--- a/Engine/Source/Runtime/GeometryCore/Private/DynamicMesh/DynamicMeshAttributeSet.cpp
+++ b/Engine/Source/Runtime/GeometryCore/Private/DynamicMesh/DynamicMeshAttributeSet.cpp
@@ -866,7 +866,7 @@ void FDynamicMeshAttributeSet::OnSplitVertex(const DynamicMeshInfo::FVertexSplit
 	}
 }
 
-bool FDynamicMeshAttributeSet::IsSameAs(const FDynamicMeshAttributeSet& Other) const
+bool FDynamicMeshAttributeSet::IsSameAs(const FDynamicMeshAttributeSet& Other, bool bIgnoreDataLayout) const
 {
 	if (UVLayers.Num() != Other.UVLayers.Num() ||
 		NormalLayers.Num() != Other.NormalLayers.Num() ||
@@ -877,7 +877,7 @@ bool FDynamicMeshAttributeSet::IsSameAs(const FDynamicMeshAttributeSet& Other) c
 
 	for (int Idx = 0; Idx < UVLayers.Num(); Idx++)
 	{
-		if (!UVLayers[Idx].IsSameAs(Other.UVLayers[Idx]))
+		if (!UVLayers[Idx].IsSameAs(Other.UVLayers[Idx], bIgnoreDataLayout))
 		{
 			return false;
 		}
@@ -885,7 +885,7 @@ bool FDynamicMeshAttributeSet::IsSameAs(const FDynamicMeshAttributeSet& Other) c
 
 	for (int Idx = 0; Idx < NormalLayers.Num(); Idx++)
 	{
-		if (!NormalLayers[Idx].IsSameAs(Other.NormalLayers[Idx]))
+		if (!NormalLayers[Idx].IsSameAs(Other.NormalLayers[Idx], bIgnoreDataLayout))
 		{
 			return false;
 		}
@@ -893,7 +893,7 @@ bool FDynamicMeshAttributeSet::IsSameAs(const FDynamicMeshAttributeSet& Other) c
 
 	for (int Idx = 0; Idx < PolygroupLayers.Num(); Idx++)
 	{
-		if (!PolygroupLayers[Idx].IsSameAs(Other.PolygroupLayers[Idx]))
+		if (!PolygroupLayers[Idx].IsSameAs(Other.PolygroupLayers[Idx], bIgnoreDataLayout))
 		{
 			return false;
 		}
@@ -905,7 +905,7 @@ bool FDynamicMeshAttributeSet::IsSameAs(const FDynamicMeshAttributeSet& Other) c
 	}
 	if (HasPrimaryColors())
 	{
-		if (!ColorLayer->IsSameAs(*Other.ColorLayer))
+		if (!ColorLayer->IsSameAs(*Other.ColorLayer, bIgnoreDataLayout))
 		{
 			return false;
 		}
@@ -917,22 +917,111 @@ bool FDynamicMeshAttributeSet::IsSameAs(const FDynamicMeshAttributeSet& Other) c
 	}
 	if (HasMaterialID())
 	{
-		if (!MaterialIDAttrib->IsSameAs(*Other.MaterialIDAttrib))
+		if (!MaterialIDAttrib->IsSameAs(*Other.MaterialIDAttrib, bIgnoreDataLayout))
 		{
 			return false;
 		}
 	}
 
+	if (SkinWeightAttributes.Num() != Other.SkinWeightAttributes.Num())
+	{
+		return false;
+	}
+	if (!SkinWeightAttributes.IsEmpty())
+	{
+		SkinWeightAttributesMap::TConstIterator It(SkinWeightAttributes);
+		SkinWeightAttributesMap::TConstIterator ItOther(Other.SkinWeightAttributes);
+		while (It && ItOther)
+		{
+			if (It.Key() != ItOther.Key())
+			{
+				return false;
+			}
+			const FDynamicMeshVertexSkinWeightsAttribute* SkinWeights = It->Value.Get();
+			const FDynamicMeshVertexSkinWeightsAttribute* SkinWeightsOther = ItOther->Value.Get();
+
+			if (SkinWeights->VertexBoneWeights.Num() != SkinWeightsOther->VertexBoneWeights.Num())
+			{
+				return false;
+			}
+
+			for (int32 i = 0, Num = SkinWeights->VertexBoneWeights.Num(); i < Num; ++i)
+			{
+				AnimationCore::FBoneWeights BoneWeights = SkinWeights->VertexBoneWeights[i];
+				AnimationCore::FBoneWeights BoneWeightsOther = SkinWeightsOther->VertexBoneWeights[i];
+
+				if (BoneWeights.Num() != BoneWeightsOther.Num())
+				{
+					return false;
+				}
+
+				for (int32 Index = 0; Index < BoneWeights.Num(); ++Index)
+				{
+					// If the weight is the same, the order is nondeterministic. Hence, we need to "manually" look for the same values.
+					const int32 IndexOther = BoneWeightsOther.FindWeightIndexByBone(BoneWeights[Index].GetBoneIndex());
+					if (IndexOther == INDEX_NONE || BoneWeights[Index].GetRawWeight() != BoneWeightsOther[IndexOther].GetRawWeight())
+					{
+						return false;
+					}
+				}
+			}
+
+			++It;
+			++ItOther;
+		}
+	}
+	
 	// TODO: Test GenericAttributes
 
 	return true;
 }
 
-void FDynamicMeshAttributeSet::Serialize(FArchive& Ar)
+namespace FDynamicMeshAttributeSet_Local
+{
+template <typename LayerType>
+void SerializeLayers(TIndirectArray<LayerType>& Layers, FArchive& Ar, const FCompactMaps* CompactMaps, bool bUseCompression)
 {
-	Ar << UVLayers;
-	Ar << NormalLayers;
-	Ar << PolygroupLayers;
+	int32 Num = Layers.Num();
+	Ar << Num;
+	if (Ar.IsLoading())
+	{
+		Layers.Empty(Num);
+		for (int32 Index = 0; Index < Num; ++Index)
+		{
+			Layers.Add(new typename TIndirectArray<LayerType>::ElementType);
+		}
+	}
+	for (int32 Index = 0; Index < Num; ++Index)
+	{
+		Layers[Index].Serialize(Ar, CompactMaps, bUseCompression);
+	}
+}
+} // namespace FDynamicMeshAttributeSet_Local
+
+void FDynamicMeshAttributeSet::Serialize(FArchive& Ar, const FCompactMaps* CompactMaps, bool bUseCompression)
+{
+	using namespace FDynamicMeshAttributeSet_Local;
+	
+	Ar.UsingCustomVersion(FUE5MainStreamObjectVersion::GUID);
+
+	const bool bUseLegacySerialization = Ar.IsLoading() && Ar.CustomVer(FUE5MainStreamObjectVersion::GUID) < FUE5MainStreamObjectVersion::DynamicMeshCompactedSerialization;
+
+	if (bUseLegacySerialization)
+	{
+		// Use serialization in TIndirectArray.
+		Ar << UVLayers;
+		Ar << NormalLayers;
+		Ar << PolygroupLayers;
+	}
+	else
+	{
+		Ar << bUseCompression;
+
+		// We do our own serialization since using the serialization in TIndirectArray will not allow us to do compacting/compression.
+		SerializeLayers(UVLayers, Ar, CompactMaps, bUseCompression);
+		SerializeLayers(NormalLayers, Ar, CompactMaps, bUseCompression);
+		SerializeLayers(PolygroupLayers, Ar, CompactMaps, bUseCompression);
+	}
 
 	if (Ar.IsLoading())
 	{
@@ -952,16 +1041,16 @@ void FDynamicMeshAttributeSet::Serialize(FArchive& Ar)
 		}
 	}
 
-	// Use int32 here to futureproof for multiple color layers.
-	int32 bHasColorLayer = HasPrimaryColors() ? 1 : 0;
-	Ar << bHasColorLayer;
-	if (bHasColorLayer)
+	// Use int32 here to future-proof for multiple color layers.
+	int32 NumColorLayers = HasPrimaryColors() ? 1 : 0;
+	Ar << NumColorLayers;
+	if (NumColorLayers > 0)
 	{
 		if (Ar.IsLoading())
 		{
 			EnablePrimaryColors();
 		}
-		Ar << *ColorLayer;
+		ColorLayer->Serialize(Ar, CompactMaps, bUseCompression);
 	}
 
 	bool bHasMaterialID = HasMaterialID();
@@ -972,9 +1061,49 @@ void FDynamicMeshAttributeSet::Serialize(FArchive& Ar)
 		{
 			EnableMaterialID();
 		}
-		Ar << *MaterialIDAttrib;
+		MaterialIDAttrib->Serialize(Ar, CompactMaps, bUseCompression);
+	}
+
+	if (!bUseLegacySerialization)
+	{
+		int32 NumSkinWeightAttributes = SkinWeightAttributes.Num();
+		Ar << NumSkinWeightAttributes;
+
+		if (Ar.IsLoading())
+		{
+			SkinWeightAttributes.Reset();
+
+			for (int32 i = 0; i < NumSkinWeightAttributes; ++i)
+			{
+				FName Key;
+				Ar << Key;
+
+				TUniquePtr<FDynamicMeshVertexSkinWeightsAttribute>& Value = SkinWeightAttributes.Emplace(Key, nullptr);
+				bool IsValid;
+				Ar << IsValid;
+				if (IsValid)
+				{
+					Value.Reset(new FDynamicMeshVertexSkinWeightsAttribute(ParentMesh, false));
+					Value.Get()->Serialize(Ar, CompactMaps, bUseCompression);
+				}
+			}
+		}
+		else
+		{
+			for (SkinWeightAttributesMap::TIterator It(SkinWeightAttributes); It; ++It)
+			{
+				Ar << It.Key();
+
+				const TUniquePtr<FDynamicMeshVertexSkinWeightsAttribute>& Value = It.Value();
+				bool bIsValid = Value.IsValid();
+				Ar << bIsValid;
+				if (bIsValid)
+				{
+					Value.Get()->Serialize(Ar, CompactMaps, bUseCompression);
+				}
+			}
+		}
 	}
 
 	//Ar << GenericAttributes; // TODO
 }
-
diff --git a/Engine/Source/Runtime/GeometryCore/Private/DynamicMesh/DynamicMeshOverlay.cpp b/Engine/Source/Runtime/GeometryCore/Private/DynamicMesh/DynamicMeshOverlay.cpp
--- a/Engine/Source/Runtime/GeometryCore/Private/DynamicMesh/DynamicMeshOverlay.cpp
+++ b/Engine/Source/Runtime/GeometryCore/Private/DynamicMesh/DynamicMeshOverlay.cpp
@@ -1486,45 +1486,119 @@ bool TDynamicMeshOverlay<RealType, ElementSize>::CheckValidity(bool bAllowNonMan
 
 
 template<typename RealType, int ElementSize>
-bool TDynamicMeshOverlay<RealType, ElementSize>::IsSameAs(const TDynamicMeshOverlay<RealType, ElementSize>& Other) const
+bool TDynamicMeshOverlay<RealType, ElementSize>::IsSameAs(const TDynamicMeshOverlay<RealType, ElementSize>& Other, bool bIgnoreDataLayout) const
 {
-	if (ElementsRefCounts.GetCount() != Other.ElementsRefCounts.GetCount() ||
-		ElementsRefCounts.GetMaxIndex() != Other.ElementsRefCounts.GetMaxIndex() ||
-		Elements.Num() != Other.Elements.Num() ||
-		ParentVertices.Num() != Other.ParentVertices.Num() ||
-		ElementTriangles.Num() != Other.ElementTriangles.Num())
+	if (ElementsRefCounts.GetCount() != Other.ElementsRefCounts.GetCount())
 	{
 		return false;
 	}
 
-	for (int Idx = 0; Idx < ElementsRefCounts.GetMaxIndex(); Idx++)
+	if (!bIgnoreDataLayout || (ElementsRefCounts.IsDense() && Other.ElementsRefCounts.IsDense() && ParentMesh->MaxTriangleID() == ParentMesh->TriangleCount() &&
+		Other.ParentMesh->MaxTriangleID() == Other.ParentMesh->TriangleCount()))
 	{
-		if (ElementsRefCounts.GetRefCount(Idx) != Other.ElementsRefCounts.GetRefCount(Idx))
+		if (ElementsRefCounts.GetMaxIndex() != Other.ElementsRefCounts.GetMaxIndex() ||
+			Elements.Num() != Other.Elements.Num() ||
+			ParentVertices.Num() != Other.ParentVertices.Num() ||
+			ElementTriangles.Num() != Other.ElementTriangles.Num())
 		{
 			return false;
 		}
-	}
 
-	for (int Idx = 0; Idx < Elements.Num(); Idx++)
-	{
-		if (Elements[Idx] != Other.Elements[Idx])
+		if (Elements != Other.Elements)
 		{
 			return false;
 		}
-	}
 
-	for (int Idx = 0; Idx < ParentVertices.Num(); Idx++)
-	{
-		if (ParentVertices[Idx] != Other.ParentVertices[Idx])
+		for (int Idx = 0; Idx < ElementsRefCounts.GetMaxIndex(); Idx++)
+		{
+			if (ElementsRefCounts.GetRefCount(Idx) != Other.ElementsRefCounts.GetRefCount(Idx))
+			{
+				return false;
+			}
+		}
+
+		if (ParentVertices != Other.ParentVertices)
 		{
 			return false;
 		}
-	}
 
-	for (int Idx = 0; Idx < ElementTriangles.Num(); Idx++)
+		if (ElementTriangles != Other.ElementTriangles)
+		{
+			return false;
+		}
+	}
+	else
 	{
-		if (ElementTriangles[Idx] != Other.ElementTriangles[Idx])
+		FRefCountVector::IndexIterator ItEid = ElementsRefCounts.BeginIndices();
+		const FRefCountVector::IndexIterator ItEidEnd = ElementsRefCounts.EndIndices();
+		FRefCountVector::IndexIterator ItEidOther = Other.ElementsRefCounts.BeginIndices();
+		const FRefCountVector::IndexIterator ItEidEndOther = Other.ElementsRefCounts.EndIndices();
+
+		TDynamicVector<int> EidMapping;
+		EidMapping.Resize(ElementsRefCounts.GetMaxIndex(), FDynamicMesh3::InvalidID);
+
+		while (ItEid != ItEidEnd && ItEidOther != ItEidEndOther)
+		{
+			for (int32 i = 0; i < ElementSize; ++i)
+			{
+				if (Elements[*ItEid * ElementSize + i] != Other.Elements[*ItEidOther * ElementSize + i])
+				{
+					// Element values are not the same.
+					return false;
+				}
+			}
+
+			if (ElementsRefCounts.GetRawRefCount(*ItEid) != Other.ElementsRefCounts.GetRefCount(*ItEidOther))
+			{
+				// Element ref counts are not the same.
+				return false;
+			}
+
+			EidMapping[*ItEid] = *ItEidOther;
+
+			++ItEid;
+			++ItEidOther;
+		}
+		
+		if (ItEid != ItEidEnd || ItEidOther != ItEidEndOther)
 		{
+			// Number of elements is not the same.
+			return false;
+		}
+		
+		if (ParentMesh->TriangleCount() != Other.ParentMesh->TriangleCount())
+		{
+			return false;
+		}
+
+		FRefCountVector::IndexIterator ItTid = ParentMesh->GetTrianglesRefCounts().BeginIndices();
+		const FRefCountVector::IndexIterator ItTidEnd = ParentMesh->GetTrianglesRefCounts().EndIndices();
+		FRefCountVector::IndexIterator ItTidOther = Other.ParentMesh->GetTrianglesRefCounts().BeginIndices();
+		const FRefCountVector::IndexIterator ItTidEndOther = Other.ParentMesh->GetTrianglesRefCounts().EndIndices();
+
+		while (ItTid != ItTidEnd && ItTidOther != ItTidEndOther)
+		{
+			const int ElementTrianglesIdx = *ItTid * 3;
+			const int ElementTrianglesIdxOther = *ItTidOther * 3;
+
+			for (int i = 0; i < 3; ++i)
+			{
+				const int Eid = ElementTriangles[ElementTrianglesIdx + i];
+				const int EidOther = Other.ElementTriangles[ElementTrianglesIdxOther + i];
+				if (EidOther != (Eid != FDynamicMesh3::InvalidID ? EidMapping[Eid] : FDynamicMesh3::InvalidID))
+				{
+					// Triangle elements do not index the same element value.
+					return false;
+				}
+			}
+
+			++ItTid;
+			++ItTidOther;
+		}
+
+		if (ItTid != ItTidEnd || ItTidOther != ItTidEndOther)
+		{
+			// Number of element triangles is not the same.
 			return false;
 		}
 	}
@@ -1534,12 +1608,136 @@ bool TDynamicMeshOverlay<RealType, ElementSize>::IsSameAs(const TDynamicMeshOver
 
 
 template<typename RealType, int ElementSize>
-void TDynamicMeshOverlay<RealType, ElementSize>::Serialize(FArchive& Ar)
+void TDynamicMeshOverlay<RealType, ElementSize>::Serialize(FArchive& Ar, const FCompactMaps* CompactMaps, bool bUseCompression)
 {
-	Ar << ElementsRefCounts;
-	Ar << Elements;
-	Ar << ParentVertices;
-	Ar << ElementTriangles;
+	Ar.UsingCustomVersion(FUE5MainStreamObjectVersion::GUID);
+	if (Ar.IsLoading() && Ar.CustomVer(FUE5MainStreamObjectVersion::GUID) < FUE5MainStreamObjectVersion::DynamicMeshCompactedSerialization)
+	{
+		Ar << ElementsRefCounts;
+		Ar << Elements;
+		Ar << ParentVertices;
+		Ar << ElementTriangles;
+	}
+	else
+	{
+		auto SerializeVector = [](FArchive& Ar, auto& Vector, bool bUseCompression)
+		{
+			if (bUseCompression)
+			{
+				Vector.template Serialize<true, true>(Ar);
+			}
+			else
+			{
+				Vector.template Serialize<true, false>(Ar);
+			}
+		};
+		
+		TOptional<TArray<int>> ElementMap;
+
+		if (Ar.IsLoading() || !CompactMaps || !CompactMaps->VertexMapIsSet())
+		{
+			ElementsRefCounts.Serialize(Ar, false, bUseCompression);
+
+			SerializeVector(Ar, Elements, bUseCompression);
+			SerializeVector(Ar, ParentVertices, bUseCompression);
+		}
+		else
+		{
+			const size_t NumElements = ElementsRefCounts.GetCount();
+
+			FRefCountVector ElementsRefCountsCompact;
+			ElementsRefCountsCompact.Trim(NumElements);
+			TDynamicVector<unsigned short>& RawElementsRefCountsCompact = ElementsRefCountsCompact.GetRawRefCountsUnsafe();
+			const TDynamicVector<unsigned short>& RawElementsRefCounts = ElementsRefCounts.GetRawRefCountsUnsafe();
+
+			TDynamicVector<RealType> ElementsCompact;
+			ElementsCompact.SetNum(NumElements * ElementSize);
+			TDynamicVector<int> ParentVerticesCompact;
+			ParentVerticesCompact.SetNum(NumElements);
+
+			ElementMap.Emplace();
+			ElementMap->Init(FDynamicMesh3::InvalidID, ElementsRefCounts.GetMaxIndex());
+			TArray<int>& ElementMapping = *ElementMap;
+
+			int EidCompact = 0;
+			for (int Eid : ElementsRefCounts.Indices())
+			{
+				RawElementsRefCountsCompact[EidCompact] = RawElementsRefCounts[Eid];
+				for (int32 i = 0; i < ElementSize; ++i)
+				{
+					ElementsCompact[EidCompact * ElementSize + i] = Elements[Eid * ElementSize + i];
+				}
+				ParentVerticesCompact[EidCompact] = CompactMaps->GetVertexMapping(ParentVertices[Eid]);
+				ElementMapping[Eid] = EidCompact;
+				++EidCompact;
+			}
+
+			ElementsRefCountsCompact.Serialize(Ar, true, bUseCompression);
+			SerializeVector(Ar, ElementsCompact, bUseCompression);
+			SerializeVector(Ar, ParentVerticesCompact, bUseCompression);
+		}
+
+		if (Ar.IsLoading() || ((!CompactMaps || !CompactMaps->TriangleMapIsSet()) && !ElementMap))
+		{
+			SerializeVector(Ar, ElementTriangles, bUseCompression);
+		}
+		else
+		{
+			TDynamicVector<int> ElementTrianglesCompact;
+			ElementTrianglesCompact.SetNum(ParentMesh->TriangleCount() * 3);
+
+			if (CompactMaps && CompactMaps->TriangleMapIsSet() && ElementMap)
+			{
+				TArray<int>& ElementMapping = *ElementMap;
+
+				for (const int Tid : ParentMesh->TriangleIndicesItr())
+				{
+					const int TriIndex = Tid * 3;
+					const int TriIndexCompact = CompactMaps->GetTriangleMapping(Tid) * 3;
+					const int TriElement0 = ElementTriangles[TriIndex + 0];
+					const int TriElement1 = ElementTriangles[TriIndex + 1];
+					const int TriElement2 = ElementTriangles[TriIndex + 2];
+					ElementTrianglesCompact[TriIndexCompact + 0] = TriElement0 != FDynamicMesh3::InvalidID
+						                                               ? ElementMapping[TriElement0]
+						                                               : FDynamicMesh3::InvalidID;
+					ElementTrianglesCompact[TriIndexCompact + 1] = TriElement1 != FDynamicMesh3::InvalidID
+						                                               ? ElementMapping[TriElement1]
+						                                               : FDynamicMesh3::InvalidID;
+					ElementTrianglesCompact[TriIndexCompact + 2] = TriElement2 != FDynamicMesh3::InvalidID
+						                                               ? ElementMapping[TriElement2]
+						                                               : FDynamicMesh3::InvalidID;
+				}
+			}
+			else if (CompactMaps && CompactMaps->TriangleMapIsSet())
+			{
+				for (const int Tid : ParentMesh->TriangleIndicesItr())
+				{
+					const int TriIndex = Tid * 3;
+					const int TriIndexCompact = CompactMaps->GetTriangleMapping(Tid) * 3;
+					ElementTrianglesCompact[TriIndexCompact + 0] = ElementTriangles[TriIndex + 0];
+					ElementTrianglesCompact[TriIndexCompact + 1] = ElementTriangles[TriIndex + 1];
+					ElementTrianglesCompact[TriIndexCompact + 2] = ElementTriangles[TriIndex + 2];
+				}
+			}
+			else
+			{
+				TArray<int>& ElementMapping = *ElementMap;
+
+				for (const int Tid : ParentMesh->TriangleIndicesItr())
+				{
+					const int TriIndex = Tid * 3;
+					const int TriElement0 = ElementTriangles[TriIndex + 0];
+					const int TriElement1 = ElementTriangles[TriIndex + 1];
+					const int TriElement2 = ElementTriangles[TriIndex + 2];
+					ElementTrianglesCompact[TriIndex + 0] = TriElement0 != FDynamicMesh3::InvalidID ? ElementMapping[TriElement0] : FDynamicMesh3::InvalidID;
+					ElementTrianglesCompact[TriIndex + 1] = TriElement1 != FDynamicMesh3::InvalidID ? ElementMapping[TriElement1] : FDynamicMesh3::InvalidID;
+					ElementTrianglesCompact[TriIndex + 2] = TriElement2 != FDynamicMesh3::InvalidID ? ElementMapping[TriElement2] : FDynamicMesh3::InvalidID;
+				}
+			}
+
+			SerializeVector(Ar, ElementTrianglesCompact, bUseCompression);
+		}
+	}
 }
 
 
diff --git a/Engine/Source/Runtime/GeometryCore/Private/Util/SmallListSet.cpp b/Engine/Source/Runtime/GeometryCore/Private/Util/SmallListSet.cpp
--- a/Engine/Source/Runtime/GeometryCore/Private/Util/SmallListSet.cpp
+++ b/Engine/Source/Runtime/GeometryCore/Private/Util/SmallListSet.cpp
@@ -2,6 +2,8 @@
 
 #include "Util/SmallListSet.h"
 
+#include <UObject/UE5MainStreamObjectVersion.h>
+
 using namespace UE::Geometry;
 
 const int32 FSmallListSet::NullValue = -1;
@@ -45,7 +47,7 @@ void FSmallListSet::AllocateAt(int32 ListIndex)
 
 void FSmallListSet::Insert(int32 ListIndex, int32 Value)
 {
-	checkSlow(ListIndex >= 0);
+	checkSlow(0 <= ListIndex && ListIndex < ListHeads.Num());
 	int32 block_ptr = ListHeads[ListIndex];
 	if (block_ptr == NullValue)
 	{
@@ -413,3 +415,124 @@ bool FSmallListSet::RemoveFromLinkedList(int32 block_ptr, int32 val)
 	}
 	return false;
 }
+
+void FSmallListSet::Serialize(FArchive& Ar, bool bCompactData, bool bUseCompression)
+{
+	Ar.UsingCustomVersion(FUE5MainStreamObjectVersion::GUID);
+	if (Ar.IsLoading() && Ar.CustomVer(FUE5MainStreamObjectVersion::GUID) < FUE5MainStreamObjectVersion::DynamicMeshCompactedSerialization)
+	{
+		Ar << ListHeads;
+		Ar << ListBlocks;
+		Ar << FreeBlocks;
+		Ar << AllocatedCount;
+		Ar << LinkedListElements;
+		Ar << FreeHeadIndex;
+	}
+	else
+	{
+		Ar << bCompactData;
+		Ar << bUseCompression;
+
+		auto SerializeVector = [](FArchive& Ar, auto& Vector, bool bUseCompression)
+		{
+			if (bUseCompression)
+			{
+				Vector.template Serialize<true, true>(Ar);
+			}
+			else
+			{
+				Vector.template Serialize<true, false>(Ar);
+			}
+		};
+
+		// Compact the data into a flat buffer if either bCompactData or bUseCompression is enabled.
+		// Considering the significant time overhead for compression, it makes sense to just compact the data as well even though it is not requested.
+		if (bCompactData || bUseCompression)
+		{
+			// We are compacting the data by flattening the lists into one tightly packed buffer.
+			// The first value in the buffer is the number of allocated lists even if they are empty. After that each list consists of the number of values
+			// followed by the actual values.
+			// Note that due to values beyond the BLOCKSIZE are inserted into the linked list, we reverse the order for these values then loading in the data to
+			// restore the original order.
+
+			// Temporary buffer to flatten all lists during save or restore them during load.
+			TDynamicVector<int32> Buffer;
+
+			if (Ar.IsLoading())
+			{
+				Reset();
+
+				SerializeVector(Ar, Buffer, bUseCompression);
+
+				size_t BufferIndex = 0;
+				const int32 ListCount = Buffer[BufferIndex++];
+				Resize(ListCount);
+
+				for (size_t ListIndex = 0; ListIndex < ListCount; ++ListIndex)
+				{
+					const int32 ListValueCount = Buffer[BufferIndex++];
+					if (ListValueCount > 0)
+					{
+						AllocateAt(ListIndex);
+
+						// The first BLOCKSIZE values get inserted in order.
+						const int32 ListValueCountInBlock = FMath::Min(BLOCKSIZE, ListValueCount);
+						for (int32 ListValueIndex = 0; ListValueIndex < ListValueCountInBlock; ++ListValueIndex)
+						{
+							Insert(ListIndex, Buffer[BufferIndex + ListValueIndex]);
+						}
+
+						// Any values beyond BLOCKSIZE get inserted in reversed order. This is due to the fact that any values inserted into the linked list
+						// effectively get reversed in order, and we want to restore the original order before serializing out the data.
+						for (int32 ListValueIndex = ListValueCount - 1; ListValueIndex >= BLOCKSIZE; --ListValueIndex)
+						{
+							Insert(ListIndex, Buffer[BufferIndex + ListValueIndex]);
+						}
+
+						BufferIndex += ListValueCount;
+					}
+				}
+			}
+			else
+			{
+				const size_t ListCount = ListHeads.Num();
+				size_t BufferSize = 0;
+				size_t BufferIndex = 0;
+
+				// Store number of lists.
+				Buffer.Resize(++BufferSize);
+				Buffer[BufferIndex++] = ListCount;
+
+				for (size_t ListIndex = 0; ListIndex < ListCount; ++ListIndex)
+				{
+					const int32 ListValueCount = GetCount(ListIndex);
+
+					// Store number of values in the current list. 
+					Buffer.Resize(BufferSize += 1 + ListValueCount);
+					Buffer[BufferIndex++] = ListValueCount;
+
+					if (ListValueCount > 0)
+					{
+						// Store values in the current list.
+						for (const int32 Value : Values(ListIndex))
+						{
+							Buffer[BufferIndex++] = Value;
+						}
+					}
+				}
+
+				SerializeVector(Ar, Buffer, bUseCompression);
+			}
+		}
+		else
+		{
+			// Naively serialize all of the underlying data.
+			SerializeVector(Ar, ListHeads, false);
+			SerializeVector(Ar, ListBlocks, false);
+			SerializeVector(Ar, FreeBlocks, false);
+			SerializeVector(Ar, LinkedListElements, false);
+			Ar << AllocatedCount;
+			Ar << FreeHeadIndex;
+		}
+	}
+}
diff --git a/Engine/Source/Runtime/GeometryCore/Public/DynamicMesh/DynamicMesh3.h b/Engine/Source/Runtime/GeometryCore/Public/DynamicMesh/DynamicMesh3.h
--- a/Engine/Source/Runtime/GeometryCore/Public/DynamicMesh/DynamicMesh3.h
+++ b/Engine/Source/Runtime/GeometryCore/Public/DynamicMesh/DynamicMesh3.h
@@ -1270,6 +1270,10 @@ public:
 		bool bCheckUVs = false;
 		bool bCheckGroups = false;
 		bool bCheckAttributes = false;
+
+		// Ignore gaps and padding in the data layout, i.e. the exact index is irrelevant as long as valid entries are equal in value and in the same order.
+		bool bIgnoreDataLayout = false;
+
 		float Epsilon = TMathUtil<float>::Epsilon;
 	};
 
@@ -1457,12 +1461,19 @@ protected:
 
 	void ReverseTriOrientationInternal(int TriangleID);
 
+	/**
+	* Internal implementations for serialization to allow for better code separation between different versions and other implementation details.
+	* @tparam Variant Identifier for a specific implementation, e.g. depending on the version and/or specific options.
+	* @param Ptr Pointer to optional data, e.g. version specific options.
+	*/
+	template<int Variant>
+	void SerializeInternal(FArchive& Ar, void* Ptr);
 
-	/** Returns the internal version number used during serialization. This is used for testing purposes. */
-	static int SerializeInternal_LatestVersion();
 };
 
 
 
 } // end namespace UE::Geometry
 } // end namespace UE
+
+template<> struct TCanBulkSerialize<UE::Geometry::FDynamicMesh3::FEdge> { enum { Value = true }; };
diff --git a/Engine/Source/Runtime/GeometryCore/Public/DynamicMesh/DynamicMeshAttributeSet.h b/Engine/Source/Runtime/GeometryCore/Public/DynamicMesh/DynamicMeshAttributeSet.h
--- a/Engine/Source/Runtime/GeometryCore/Public/DynamicMesh/DynamicMeshAttributeSet.h
+++ b/Engine/Source/Runtime/GeometryCore/Public/DynamicMesh/DynamicMeshAttributeSet.h
@@ -373,7 +373,7 @@ public:
 	/**
 	 * Returns true if this AttributeSet is the same as Other.
 	 */
-	bool IsSameAs(const FDynamicMeshAttributeSet& Other) const;
+	bool IsSameAs(const FDynamicMeshAttributeSet& Other, bool bIgnoreDataLayout) const;
 
 	/**
 	 * Serialization operator for FDynamicMeshAttributeSet.
@@ -384,12 +384,18 @@ public:
 	 */
 	friend GEOMETRYCORE_API FArchive& operator<<(FArchive& Ar, FDynamicMeshAttributeSet& Set)
 	{
-		Set.Serialize(Ar);
+		Set.Serialize(Ar, nullptr, false);
 		return Ar;
 	}
 
-	/** Serialize the set to an archive. */
-	void Serialize(FArchive& Ar);
+	/**
+	* Serialize to and from an archive.
+	*
+	* @param Ar Archive to serialize with.
+	* @param CompactMaps If this is not a null pointer, the mesh serialization compacted the vertex and/or triangle data using the provided mapping. 
+	* @param bUseCompression Use compression for serializing bulk data.
+	*/
+	void Serialize(FArchive& Ar, const FCompactMaps* CompactMaps, bool bUseCompression);
 
 protected:
 	/** Parent mesh of this attribute set */
@@ -404,9 +410,11 @@ protected:
 
 	TIndirectArray<FDynamicMeshPolygroupAttribute> PolygroupLayers;
 
-	TMap<FName, TUniquePtr<FDynamicMeshVertexSkinWeightsAttribute>> SkinWeightAttributes;
-	
-	TMap<FName, TUniquePtr<FDynamicMeshAttributeBase>> GenericAttributes;
+	using SkinWeightAttributesMap = TMap<FName, TUniquePtr<FDynamicMeshVertexSkinWeightsAttribute>>;
+	SkinWeightAttributesMap SkinWeightAttributes;
+
+	using GenericAttributesMap = TMap<FName, TUniquePtr<FDynamicMeshAttributeBase>>;
+	GenericAttributesMap GenericAttributes;
 	
 protected:
 	friend class FDynamicMesh3;
diff --git a/Engine/Source/Runtime/GeometryCore/Public/DynamicMesh/DynamicMeshOverlay.h b/Engine/Source/Runtime/GeometryCore/Public/DynamicMesh/DynamicMeshOverlay.h
--- a/Engine/Source/Runtime/GeometryCore/Public/DynamicMesh/DynamicMeshOverlay.h
+++ b/Engine/Source/Runtime/GeometryCore/Public/DynamicMesh/DynamicMeshOverlay.h
@@ -523,7 +523,7 @@ public:
 	/**
 	 * Returns true if this overlay is the same as Other.
 	 */
-	bool IsSameAs(const TDynamicMeshOverlay<RealType, ElementSize>& Other) const;
+	bool IsSameAs(const TDynamicMeshOverlay<RealType, ElementSize>& Other, bool bIgnoreDataLayout) const;
 
 	/**
 	 * Serialization operator for FDynamicMeshOverlay.
@@ -534,12 +534,18 @@ public:
 	 */
 	friend FArchive& operator<<(FArchive& Ar, TDynamicMeshOverlay<RealType,ElementSize>& Overlay)
 	{
-		Overlay.Serialize(Ar);
+		Overlay.Serialize(Ar, nullptr, false);
 		return Ar;
 	}
 
-	/** Serialize the overlay to an archive. */
-	void Serialize(FArchive& Ar);
+	/**
+	 * Serialize to and from an archive.
+	 *
+	 * @param Ar Archive to serialize with.
+	 * @param CompactMaps If this is not a null pointer, the mesh serialization compacted the vertex and/or triangle data using the provided mapping. 
+	 * @param bUseCompression Use compression for serializing bulk data.
+	 */
+	void Serialize(FArchive& Ar, const FCompactMaps* CompactMaps, bool bUseCompression);
 
 public:
 	/** Set a triangle's element indices to InvalidID */
diff --git a/Engine/Source/Runtime/GeometryCore/Public/DynamicMesh/DynamicMeshTriangleAttribute.h b/Engine/Source/Runtime/GeometryCore/Public/DynamicMesh/DynamicMeshTriangleAttribute.h
--- a/Engine/Source/Runtime/GeometryCore/Public/DynamicMesh/DynamicMeshTriangleAttribute.h
+++ b/Engine/Source/Runtime/GeometryCore/Public/DynamicMesh/DynamicMeshTriangleAttribute.h
@@ -376,17 +376,49 @@ public:
 	/**
 	 * Returns true if this AttributeSet is the same as Other.
 	 */
-	bool IsSameAs(const TDynamicMeshTriangleAttribute<AttribValueType, AttribDimension>& Other) const
+	bool IsSameAs(const TDynamicMeshTriangleAttribute<AttribValueType, AttribDimension>& Other, bool bIgnoreDataLayout) const
 	{
-		if (AttribValues.Num() != Other.AttribValues.Num())
+		if (!bIgnoreDataLayout)
 		{
-			return false;
-		}
+			if (AttribValues.Num() != Other.AttribValues.Num())
+			{
+				return false;
+			}
 
-		for (int Idx = 0, NumValues = AttribValues.Num(); Idx < NumValues; Idx++)
+			for (int Idx = 0, NumValues = AttribValues.Num(); Idx < NumValues; Idx++)
+			{
+				if (AttribValues[Idx] != Other.AttribValues[Idx])
+				{
+					return false;
+				}
+			}
+		}
+		else
 		{
-			if (AttribValues[Idx] != Other.AttribValues[Idx])
+			FRefCountVector::IndexIterator ItTid = ParentMesh->GetTrianglesRefCounts().BeginIndices();
+			const FRefCountVector::IndexIterator ItTidEnd = ParentMesh->GetTrianglesRefCounts().EndIndices();
+			FRefCountVector::IndexIterator ItTidOther = Other.ParentMesh->GetTrianglesRefCounts().BeginIndices();
+			const FRefCountVector::IndexIterator ItTidEndOther = Other.ParentMesh->GetTrianglesRefCounts().EndIndices();
+
+			while (ItTid != ItTidEnd && ItTidOther != ItTidEndOther)
+			{
+				for (int32 i = 0; i < AttribDimension; ++i)
+				{
+					const AttribValueType AttribValue = AttribValues[*ItTid * AttribDimension + i];
+					const AttribValueType AttribValueOther = Other.AttribValues[*ItTidOther * AttribDimension + i];
+					if (AttribValue != AttribValueOther)
+					{
+						// Triangle attribute value is not the same.
+						return false;
+					}
+				}
+				++ItTid;
+				++ItTidOther;
+			}
+
+			if (ItTid != ItTidEnd || ItTidOther != ItTidEndOther)
 			{
+				// Number of triangle attribute values is not the same.
 				return false;
 			}
 		}
@@ -403,14 +435,65 @@ public:
 	 */
 	friend FArchive& operator<<(FArchive& Ar, TDynamicMeshTriangleAttribute<AttribValueType, AttribDimension>& Attr)
 	{
-		Attr.Serialize(Ar);
+		Attr.Serialize(Ar, nullptr, false);
 		return Ar;
 	}
 
-	/** Serialize the attribute to an archive. */
-	void Serialize(FArchive& Ar)
-	{
-		Ar << AttribValues;
+	/**
+	* Serialize to and from an archive.
+	*
+	* @param Ar Archive to serialize with.
+	* @param CompactMaps If this is not a null pointer, the mesh serialization compacted the vertex and/or triangle data using the provided mapping. 
+	* @param bUseCompression Use compression for serializing bulk data.
+	*/
+	void Serialize(FArchive& Ar, const FCompactMaps* CompactMaps, bool bUseCompression)
+	{
+		Ar.UsingCustomVersion(FUE5MainStreamObjectVersion::GUID);
+		if (Ar.IsLoading() && Ar.CustomVer(FUE5MainStreamObjectVersion::GUID) < FUE5MainStreamObjectVersion::DynamicMeshCompactedSerialization)
+		{
+			Ar << AttribValues;
+		}
+		else
+		{
+			auto SerializeVector = [](FArchive& Ar, auto& Vector, bool bUseCompression)
+			{
+				if (bUseCompression)
+				{
+					Vector.template Serialize<true, true>(Ar);
+				}
+				else
+				{
+					Vector.template Serialize<true, false>(Ar);
+				}
+			};
+
+			Ar << bUseCompression;
+
+			if (CompactMaps == nullptr || !CompactMaps->TriangleMapIsSet())
+			{
+				SerializeVector(Ar, AttribValues, bUseCompression);
+			}
+			else
+			{
+				TDynamicVector<AttribValueType> AttribValuesCompact;
+				AttribValuesCompact.SetNum(ParentMesh->TriangleCount() * AttribDimension);
+
+				int32 TidCompact = 0;
+				for (int32 Tid = 0, Num = AttribValues.Num() / AttribDimension; Tid < Num; ++Tid)
+				{
+					const int32 TidMapping = CompactMaps->GetTriangleMapping(Tid);
+					if (TidMapping != FCompactMaps::InvalidID)
+					{
+						for (int32 i = 0; i < AttribDimension; ++i)
+						{
+							AttribValuesCompact[TidCompact++] = AttribValues[Tid * AttribDimension + i];
+						}
+					}
+				}
+
+				SerializeVector(Ar, AttribValuesCompact, bUseCompression);
+			}
+		}
 	}
 };
 
diff --git a/Engine/Source/Runtime/GeometryCore/Public/DynamicMesh/DynamicVertexSkinWeightsAttribute.h b/Engine/Source/Runtime/GeometryCore/Public/DynamicMesh/DynamicVertexSkinWeightsAttribute.h
--- a/Engine/Source/Runtime/GeometryCore/Public/DynamicMesh/DynamicVertexSkinWeightsAttribute.h
+++ b/Engine/Source/Runtime/GeometryCore/Public/DynamicMesh/DynamicVertexSkinWeightsAttribute.h
@@ -145,7 +145,7 @@ public:
 		}
 	}
 
-	/** Initialize the attribute values to the given max triangle ID */
+	/** Initialize the attribute values to the given max vertex ID */
 	void Initialize(const FBoneWeights InitialValue = {})
 	{
 		check(Parent != nullptr);
@@ -336,6 +336,144 @@ public:
 		return true;
 	}
 
+	void Serialize(FArchive& Ar, const FCompactMaps* CompactMaps, bool bUseCompression)
+	{
+		Ar.UsingCustomVersion(FUE5MainStreamObjectVersion::GUID);
+
+		Ar << bUseCompression;
+
+		const bool bUseVertexCompactMap = CompactMaps && CompactMaps->VertexMapIsSet();
+
+		if (!bUseCompression)
+		{
+			if (Ar.IsLoading() || !bUseVertexCompactMap)
+			{
+				VertexBoneWeights.Serialize<false, false>(Ar);
+			}
+			else
+			{
+				TDynamicVector<FBoneWeights> VertexBoneWeightsCompact;
+				VertexBoneWeightsCompact.Resize(Parent->VertexCount());
+
+				for (int32 Vid = 0, Num = VertexBoneWeights.Num(); Vid < Num; ++Vid)
+				{
+					const int32 VidCompact = CompactMaps->GetVertexMapping(Vid);
+					if (VidCompact != FCompactMaps::InvalidID)
+					{
+						VertexBoneWeightsCompact[VidCompact] = VertexBoneWeights[Vid];
+					}
+				}
+
+				VertexBoneWeightsCompact.Serialize<false, false>(Ar);
+			}
+		}
+		else
+		{
+			// To achieve better compression performance both w.r.t. size and speed, we copy everything into one big flat buffer.
+			// We take advantage of the fact that we can effectively store everything as int32 for both the vector/array sizes as well as the bone weights.
+			TArray<int32> Buffer;
+
+			if (Ar.IsLoading())
+			{
+				// Serialize size of decompressed buffer to be able to hold all the compressed data in the archive, and allocate accordingly.
+				int32 BufferSize;
+				Ar << BufferSize;
+				Buffer.SetNumUninitialized(BufferSize);
+
+				// Decompress buffer from archive.
+				Ar.SerializeCompressedNew(Buffer.GetData(), Buffer.Num() * sizeof(int32), NAME_Oodle, NAME_Oodle, COMPRESS_NoFlags, false, nullptr);
+
+				// Restore bone weights arrays from decompressed buffer.
+				int32* BufferPtr = Buffer.GetData();
+				VertexBoneWeights.Resize(*BufferPtr++);
+				AnimationCore::FBoneWeightsSettings BoneWeightsSettings;
+				BoneWeightsSettings.SetNormalizeType(AnimationCore::EBoneWeightNormalizeType::None);
+				for (FBoneWeights& BoneWeights : VertexBoneWeights)
+				{
+					const int32 Num = *BufferPtr++;
+					for (int32 i = 0; i < Num; ++i)
+					{
+						BoneWeights.SetBoneWeight(reinterpret_cast<AnimationCore::FBoneWeight&>(*BufferPtr++), BoneWeightsSettings);
+					}
+				}
+				checkSlow(BufferPtr == &*Buffer.end());
+			}
+			else
+			{
+				// Determine total number of individual bone weights.
+
+				auto CountBoneWeights = [](const FBoneWeights& BoneWeights, SIZE_T& NumBoneWeights)
+				{
+					NumBoneWeights += BoneWeights.Num();
+				};
+
+				SIZE_T NumBoneWeights = 0;
+				if (!bUseVertexCompactMap)
+				{
+					for (const FBoneWeights& BoneWeights : VertexBoneWeights)
+					{
+						CountBoneWeights(BoneWeights, NumBoneWeights);
+					}
+				}
+				else
+				{
+					for (int32 Vid = 0, Num = VertexBoneWeights.Num(); Vid < Num; ++Vid)
+					{
+						const int32 VidCompact = CompactMaps->GetVertexMapping(Vid);
+						if (VidCompact != FCompactMaps::InvalidID)
+						{
+							CountBoneWeights(VertexBoneWeights[Vid], NumBoneWeights);
+						}
+					}
+				}
+
+				// Set buffer size to hold number of vertex bone weight arrays, size for vertex each bone weight array, and all individual bone weights.
+				// We also serialize out size of uncompressed buffer to allow for allocation to correct size during loading.
+				uint32 BufferSize = 1 + VertexBoneWeights.Num() + NumBoneWeights;
+				Buffer.SetNumUninitialized(BufferSize);
+				Ar << BufferSize;
+
+				// Write everything into the buffer.
+
+				auto WriteBoneWeights = [](const FBoneWeights& BoneWeights, int32*& BufferPtr)
+				{
+					const int32 Num = BoneWeights.Num();
+					*BufferPtr++ = Num;
+					if (Num > 0)
+					{
+						FMemory::Memcpy(BufferPtr, &BoneWeights[0], Num * sizeof(int32));
+						BufferPtr += Num;
+					}
+				};
+
+				int32* BufferPtr = Buffer.GetData();
+				if (!bUseVertexCompactMap)
+				{
+					*BufferPtr++ = VertexBoneWeights.Num();
+					for (const FBoneWeights& BoneWeights : VertexBoneWeights)
+					{
+						WriteBoneWeights(BoneWeights, BufferPtr);
+					}
+				}
+				else
+				{
+					*BufferPtr++ = Parent->VertexCount();
+					for (int32 Vid = 0, Num = VertexBoneWeights.Num(); Vid < Num; ++Vid)
+					{
+						const int32 VidCompact = CompactMaps->GetVertexMapping(Vid);
+						if (VidCompact != FCompactMaps::InvalidID)
+						{
+							WriteBoneWeights(VertexBoneWeights[Vid], BufferPtr);
+						}
+					}
+				}
+				checkSlow(BufferPtr == &*Buffer.end());
+
+				// Compress buffer to archive.
+				Ar.SerializeCompressedNew(Buffer.GetData(), Buffer.Num() * sizeof(int32), NAME_Oodle, NAME_Oodle, COMPRESS_NoFlags, false, nullptr);
+			}
+		}
+	}
 
 protected:
 
diff --git a/Engine/Source/Runtime/GeometryCore/Public/Util/CompactMaps.h b/Engine/Source/Runtime/GeometryCore/Public/Util/CompactMaps.h
--- a/Engine/Source/Runtime/GeometryCore/Public/Util/CompactMaps.h
+++ b/Engine/Source/Runtime/GeometryCore/Public/Util/CompactMaps.h
@@ -14,7 +14,7 @@ namespace Geometry
 {
 /**
  * Stores index remapping for vertices and triangles.
- * Should only be used for compacting, and should maintain invariant that Map*[Idx] <= Idx for all maps
+ * Should only be used for compacting, and should maintain invariant that *Map[Idx] <= Idx for all maps
  */
 class FCompactMaps
 {
diff --git a/Engine/Source/Runtime/GeometryCore/Public/Util/DynamicVector.h b/Engine/Source/Runtime/GeometryCore/Public/Util/DynamicVector.h
--- a/Engine/Source/Runtime/GeometryCore/Public/Util/DynamicVector.h
+++ b/Engine/Source/Runtime/GeometryCore/Public/Util/DynamicVector.h
@@ -7,6 +7,7 @@
 #include <CoreMinimal.h>
 #include "Containers/StaticArray.h"
 #include "Serialization/Archive.h"
+#include <UObject/UE5MainStreamObjectVersion.h>
 #include "VectorTypes.h"
 #include "IndexTypes.h"
 
@@ -104,6 +105,8 @@ public:
 
 	inline Type& operator[](unsigned int Index)
 	{
+		checkSlow(Index < Num());
+
 		int ArrayIndex, BlockIndex;
 		GetIndices( Index, ArrayIndex, BlockIndex );
 	
@@ -111,6 +114,8 @@ public:
 	}
 	inline const Type& operator[](unsigned int Index) const
 	{
+		checkSlow(Index < Num());
+
 		int ArrayIndex, BlockIndex;
 		GetIndices( Index, ArrayIndex, BlockIndex );
 	
@@ -130,31 +135,59 @@ public:
 	 */
 	friend FArchive& operator<<(FArchive& Ar, TDynamicVector& Vec)
 	{
-		Vec.Serialize(Ar);
+		Vec.Serialize<false, false>(Ar);
 		return Ar;
 	}
 
-	/** Serialize vector to an archive */
+	/**
+	 * Serialize vector to and from an archive
+	 * @tparam bForceBulkSerialization Forces serialization to consider data to be trivial and serialize it in bulk to potentially achieve better performance.
+	 * @tparam bUseCompression Use compression to serialize data; the resulting size will likely be smaller but serialization will take significantly longer.
+	 * @param Ar Archive to serialize with
+	 */
+	template <bool bForceBulkSerialization = false, bool bUseCompression = false>
 	void Serialize(FArchive& Ar)
 	{
-		Ar << CurBlock;
-		Ar << CurBlockUsed;
-		Ar << Blocks;
+		Ar.UsingCustomVersion(FUE5MainStreamObjectVersion::GUID);
+		if (Ar.IsLoading() && Ar.CustomVer(FUE5MainStreamObjectVersion::GUID) < FUE5MainStreamObjectVersion::DynamicMeshCompactedSerialization)
+		{
+			const SIZE_T CountBytes = 3 * sizeof(uint32) + Blocks.Num() * BlockSize * sizeof(Type);
+			Ar.CountBytes(CountBytes, CountBytes);
+			Ar << CurBlock;
+			Ar << CurBlockUsed;
+			Blocks.Serialize_Legacy(Ar);
 
-		// Account for a previously existing issue where one additional empty block was allocated and the resulting values for the current block can cause
-        // crashes, for example when using Back().
-        if (Ar.IsLoading() && CurBlock > 0 && CurBlockUsed == 0)
-        {
-            const SIZE_T Count = Num();
-            checkSlow(Count % BlockSize == 0);
-            checkSlow(Count / BlockSize < Blocks.Num());
+			// Account for a previously existing issue where one additional empty block was allocated and the resulting values for the current block can cause
+			// crashes, for example when using Back().
+			if (CurBlock > 0 && CurBlockUsed == 0)
+			{
+				const SIZE_T Count = Num();
+				checkSlow(Count % BlockSize == 0);
+				checkSlow(Count / BlockSize < Blocks.Num());
 
-        	// Fix CurBlock and CurBlockUsed.
-            SetCurBlock(Count);
+				// Fix CurBlock and CurBlockUsed.
+				SetCurBlock(Count);
 
-        	// Remove empty block.
-            Blocks.Truncate(Count / BlockSize, false);
-        }
+				// Remove empty block.
+				Blocks.Truncate(Count / BlockSize, false);
+			}
+		}
+		else
+		{
+			uint32 SerializeNum = Num();
+			const SIZE_T CountBytes =  sizeof(uint32) + SerializeNum * sizeof(Type);
+			Ar.CountBytes(CountBytes, CountBytes);
+			Ar << SerializeNum;
+			if (SerializeNum == 0 && Ar.IsLoading())
+			{
+				Clear();
+			}
+			else if (SerializeNum > 0)
+			{
+				SetCurBlock(SerializeNum);
+				Blocks.template Serialize<bForceBulkSerialization, bUseCompression>(Ar, SerializeNum);
+			}
+		}
 	}
 
 public:
@@ -357,68 +390,131 @@ private:
 			Elements.Empty(NewReservedSize);
 		}
 
-		/**
-		 * Count bytes needed to serialize this array.
-		 *
-		 * @param Ar Archive to count for.
-		 */
-		void CountBytes(FArchive& Ar) const
-		{
-			Elements.CountBytes(Ar);
-		}
-
-		/**
-		 * Serialization operator for TBlockVector.
-		 *
-		 * @param Ar Archive to serialize with.
-		 * @param Vec Vector to serialize.
-		 * @returns Passing down serializing archive.
-		 */
-		friend FArchive& operator<<(FArchive& Ar, TBlockVector& Vec)
-		{
-			Vec.Serialize(Ar);
-			return Ar;
-		}
-
 		/** Serialize TBlockVector to archive. */
-		void Serialize(FArchive& Ar)
+		void Serialize_Legacy(FArchive& Ar)
 		{
-			CountBytes(Ar);
+			const auto SerializeElement = [](FArchive& Ar, ArrayType* Element)
+			{
+				if (TCanBulkSerialize<Type>::Value)
+				{
+					Ar.Serialize(Element->GetData(), Element->Num() * sizeof(Type));
+				}
+				else
+				{
+					Ar << *Element;
+				}
+			};
+
+			int32 BlockNum = Num();
+			Ar << BlockNum;
 			if (Ar.IsLoading())
 			{
 				// Load array.
-				int32 NewNum;
-				Ar << NewNum;
-				Empty(NewNum);
-				for (int32 Index = 0; Index < NewNum; Index++)
+				Empty(BlockNum);
+				for (int32 Index = 0; Index < BlockNum; Index++)
 				{
 					ArrayType* NewElement = new ArrayType;
-					// TODO: Consider moving bulk serialization code into TStaticArray.
-					if (TCanBulkSerialize<Type>::Value)
+					SerializeElement(Ar, NewElement);
+					Add(NewElement);
+				}
+			}
+			else
+			{
+				// Save array.
+				for (int32 Index = 0; Index < BlockNum; Index++)
+				{
+					SerializeElement(Ar, &(*this)[Index]);
+				}
+			}
+		}
+
+		template <bool bForceBulkSerialization, bool bUseCompression>
+		void Serialize(FArchive& Ar, uint32 Num)
+		{
+			constexpr bool bUseBulkSerialization = bForceBulkSerialization || TCanBulkSerialize<Type>::Value || sizeof(Type) == 1;
+			static_assert(!bUseCompression || bUseBulkSerialization, "Compression only available when using bulk serialization");
+
+			checkSlow(Num > 0);
+
+			// Serialize compression flag, which adds flexibility when de-serializing existing data even if some implementation details change.  
+			bool bUseCompressionForBulkSerialization = bUseBulkSerialization && bUseCompression;
+			Ar << bUseCompressionForBulkSerialization;
+
+			// Determine number of blocks.
+			const bool bNumIsNotMultipleOfBlockSize = Num % BlockSize != 0;
+			const uint32 NumBlocks = Num / BlockSize + bNumIsNotMultipleOfBlockSize;
+			
+			if (bUseCompressionForBulkSerialization)
+			{
+				// When using compression, copy everything to/from a big single buffer and serialize the big buffer.
+				// This results in better compression ratios while at the same time accelerating compression. 
+
+				TArray<Type> Buffer;
+				Buffer.SetNumUninitialized(Num);
+				Type* BufferPtr = Buffer.GetData();
+				SIZE_T NumCopyRemaining = Num;
+
+				if (!Ar.IsLoading())
+				{					
+					for (uint32 Index = 0; Index < NumBlocks; ++Index)
 					{
-						Ar.Serialize(NewElement->GetData(), NewElement->Num() * sizeof(Type));
+						const SIZE_T NumCopy = FMath::template Min<uint32>(NumCopyRemaining, BlockSize);
+						FMemory::Memcpy(BufferPtr, Elements[Index]->GetData(), NumCopy * sizeof(Type));
+						BufferPtr += NumCopy;
+						NumCopyRemaining -= BlockSize;
 					}
-					else
+				}
+
+				Ar.SerializeCompressedNew(Buffer.GetData(), Num * sizeof(Type), NAME_Oodle, NAME_Oodle, COMPRESS_NoFlags, false, nullptr);
+
+				if (Ar.IsLoading())
+				{
+					Empty(NumBlocks);
+					for (uint32 Index = 0; Index < NumBlocks; ++Index)
 					{
-						Ar << *NewElement;
+						ArrayType *const NewElement = new ArrayType;
+						const SIZE_T NumCopy = FMath::template Min<uint32>(NumCopyRemaining, BlockSize);
+						FMemory::Memcpy(NewElement->GetData(), BufferPtr, NumCopy * sizeof(Type));
+						Add(NewElement);
+						BufferPtr += NumCopy;
+						NumCopyRemaining -= BlockSize;
 					}
-					Add(NewElement);
 				}
 			}
 			else
 			{
-				// Save array.
-				int32 CurNum = Num();
-				Ar << CurNum;
-				for (int32 Index = 0; Index < CurNum; Index++)
+				const auto SerializeElement = [&Ar, bUseBulkSerialization](ArrayType* Element, uint32 ElementNum)
 				{
-					if (TCanBulkSerialize<Type>::Value)
+					if (bUseBulkSerialization)
 					{
-						Ar.Serialize((*this)[Index].GetData(), (*this)[Index].Num() * sizeof(Type));
+						Ar.Serialize(Element->GetData(), ElementNum * sizeof(Type));
 					}
 					else
 					{
-						Ar << (*this)[Index];
+						for (uint32 Index = 0; Index < ElementNum; ++Index)
+						{
+							Ar << (*Element)[Index];
+						}
+					}
+				};
+
+				if (Ar.IsLoading())
+				{
+					Empty(NumBlocks);
+					for (uint32 Index = 0; Index < NumBlocks; ++Index)
+					{
+						ArrayType *const NewElement = new ArrayType;
+						SerializeElement(NewElement, FMath::template Min<uint32>(Num, BlockSize));
+						Add(NewElement);
+						Num -= BlockSize;
+					}
+				}
+				else
+				{
+					for (uint32 Index = 0; Index < NumBlocks; ++Index)
+					{
+						SerializeElement(&(*this)[Index], FMath::template Min<uint32>(Num, BlockSize));
+						Num -= BlockSize;
 					}
 				}
 			}
@@ -821,4 +917,4 @@ void TDynamicVector<Type>::Apply(const Func& applyF)
 
 
 } // end namespace UE::Geometry
-} // end namespace UE
\ No newline at end of file
+} // end namespace UE
diff --git a/Engine/Source/Runtime/GeometryCore/Public/Util/RefCountVector.h b/Engine/Source/Runtime/GeometryCore/Public/Util/RefCountVector.h
--- a/Engine/Source/Runtime/GeometryCore/Public/Util/RefCountVector.h
+++ b/Engine/Source/Runtime/GeometryCore/Public/Util/RefCountVector.h
@@ -112,6 +112,7 @@ public:
 
 	void Decrement(int Index, unsigned short DecrementCount = 1)
 	{
+		checkSlow(RefCounts[Index] != INVALID_REF_COUNT && RefCounts[Index] >= DecrementCount);
 		RefCounts[Index] -= DecrementCount;
 		if (RefCounts[Index] == 0)
 		{
@@ -223,6 +224,59 @@ public:
 	//   remove
 	//   clear
 
+	/**
+	 * Rebuilds all reference counts from external source data via callables.
+	 *
+	 * For example, we could rebuild ref counts for vertices by iterating over all triangles via the Iterate callable, and in the process any vertex we
+	 * encounter for the first time gets a ref count of 1 via the AllocateRefCount callable, and any vertex we encounter repeatedly has its ref count
+	 * incremented by one via the IncrementRefCount callable. 
+	 *
+	 * @param Num Number of items that get referenced.
+	 * @param Iterate Callable that iterates over all referenced items in the external source data.
+	 * @param AllocateRefCount Callable for allocating an item, i.e. an item that is referenced for the first time.
+	 * @param IncrementRefCount Callable for incrementing an item, i.e. an item that has been referenced before.
+	 */
+	template <typename IterateFunc, typename AllocateRefCountFunc, typename IncrementRefCountFunc>
+	void Rebuild(SIZE_T Num, IterateFunc&& Iterate, AllocateRefCountFunc&& AllocateRefCount, IncrementRefCountFunc&& IncrementRefCount)
+	{
+		// Initialize ref counts to the given number of elements.
+		RefCounts.Resize(Num);
+		RefCounts.Fill(INVALID_REF_COUNT);
+		UsedCount = 0;
+
+		// Lambda for updating ref count for a given index. This is passed to the external iterate function.
+		const auto UpdateRefCount = [this, &AllocateRefCount, &IncrementRefCount](int32 Index)
+		{
+			unsigned short& RefCount = RefCounts[Index];
+			if (RefCount == INVALID_REF_COUNT)
+			{
+				// Increase used counter and call external function to initialize value.
+				++UsedCount;
+				Forward<AllocateRefCountFunc>(AllocateRefCount)(RefCount);
+			}
+			else
+			{
+				// Call external function to initialize value.
+				Forward<IncrementRefCountFunc>(IncrementRefCount)(RefCount);
+			}
+		};
+
+		// Call external function that iterates over all external pieces of data, which will in turn call the lambda to update ref counts. 
+		Forward<IterateFunc>(Iterate)(UpdateRefCount);
+
+		// Add unused elements to free list.
+		const SIZE_T FreeIndicesNum = Num - UsedCount;
+		FreeIndices.SetNum(FreeIndicesNum);
+		SIZE_T FreeIndicesIndex = 0;
+		for (SIZE_T Index = 0; (Index < Num) & (FreeIndicesIndex < FreeIndicesNum); ++Index)
+		{
+			if (RefCounts[Index] == INVALID_REF_COUNT)
+			{
+				FreeIndices[FreeIndicesIndex++] = Index;
+			}
+		}
+	}
+
 	void RebuildFreeList()
 	{
 		FreeIndices.Clear();
@@ -444,9 +498,9 @@ public:
 		return FilteredEnumerable(Indices(), FilterFunc);
 	}
 
-	FString UsageStats()
+	FString UsageStats() const
 	{
-		return FString::Printf(TEXT("RefCountSize %d  FreeSize %d  FreeMem %dkb"),
+		return FString::Printf(TEXT("RefCountSize %llu  FreeSize %llu  FreeMem %llukb"),
 			RefCounts.GetLength(), FreeIndices.GetLength(), (FreeIndices.GetByteCount() / 1024));
 	}
 
@@ -459,16 +513,111 @@ public:
 	 */
 	friend FArchive& operator<<(FArchive& Ar, FRefCountVector& Vec)
 	{
-		Vec.Serialize(Ar);
+		Vec.Serialize(Ar, false, false);
 		return Ar;
 	}
 
-	/** Serialize FRefCountVector to an archive. */
-	void Serialize(FArchive& Ar)
+	/**
+	 * Serialize FRefCountVector to an archive.
+	 * @param Ar Archive to serialize with
+	 * @param bCompactData Only serialize unique data and/or recompute redundant data when loading.
+	 * @param bUseCompression Use compression to serialize data; the resulting size will likely be smaller but serialization will take significantly longer.
+	 */
+	void Serialize(FArchive& Ar, bool bCompactData, bool bUseCompression)
+	{
+		Ar.UsingCustomVersion(FUE5MainStreamObjectVersion::GUID);
+		if (Ar.IsLoading() && Ar.CustomVer(FUE5MainStreamObjectVersion::GUID) < FUE5MainStreamObjectVersion::DynamicMeshCompactedSerialization)
+		{
+			Ar << RefCounts;
+			Ar << FreeIndices;
+			Ar << UsedCount;
+		}
+		else
+		{
+			// Serialize options.
+			Ar << bCompactData;
+			Ar << bUseCompression;
+
+			// Serialize number of used ref counts. While this is redundant to the contents of the ref count vector, it allows us to accelerate restoring the
+			// free indices and therefore it is worth the very minor storage overhead.
+			Ar << UsedCount;
+
+			// Serialize ref counts.
+			if (bUseCompression)
+			{
+				RefCounts.Serialize<true, true>(Ar);
+			}
+			else
+			{
+				RefCounts.Serialize<true, false>(Ar);
+			}
+
+			if (UsedCount == RefCounts.Num())
+			{
+				// If all ref counts are used then we can ignore the free indices entirely during save, and just clear them during load.
+
+				if (Ar.IsLoading() && !FreeIndices.IsEmpty())
+				{
+					FreeIndices.Clear();
+				}
+			}
+			else
+			{
+				// Compact the data by omitting the free indices entirely during save and recompute them during load.
+				// Considering the significant time overhead for compression, we compact if either bCompactData or bUseCompression is enabled.
+				if (bCompactData || bUseCompression)
+				{
+					if (Ar.IsLoading())
+					{
+						// Rebuild the free indices from the invalid ref count values.
+						const size_t NumFree = RefCounts.Num() - UsedCount;
+						FreeIndices.Resize(NumFree);
+						size_t Index = 0;
+						for (size_t i = 0, Num = RefCounts.Num(); (i < Num) & (Index < NumFree); ++i)
+						{
+							if (RefCounts[i] == INVALID_REF_COUNT)
+							{
+								FreeIndices[Index++] = i;
+							}
+						}
+						ensure(Index == NumFree);
+					}
+				}
+				else
+				{
+					FreeIndices.Serialize<true, false>(Ar);
+				}
+			}
+		}
+	}
+
+	friend bool operator==(const FRefCountVector& Lhs, const FRefCountVector& Rhs)
+	{
+		if (Lhs.GetCount() != Rhs.GetCount())
+		{
+			return false;
+		}
+
+		const int32 Num = FMath::Max(Lhs.GetMaxIndex(), Rhs.GetMaxIndex());
+		for (int32 Idx = 0; Idx < Num; ++Idx)
+		{
+			const bool LhsIsValid = Lhs.IsValid(Idx);
+			if (LhsIsValid != Rhs.IsValid(Idx))
+			{
+				return false;
+			}
+			if (LhsIsValid && Lhs.GetRefCount(Idx) != Rhs.GetRefCount(Idx))
+			{
+				return false;
+			}
+		}
+
+		return true;
+	}
+
+	friend bool operator!=(const FRefCountVector& Lhs, const FRefCountVector& Rhs)
 	{
-		Ar << RefCounts;
-		Ar << FreeIndices;
-		Ar << UsedCount;
+		return !(Lhs == Rhs);
 	}
 
 private:
diff --git a/Engine/Source/Runtime/GeometryCore/Public/Util/SmallListSet.h b/Engine/Source/Runtime/GeometryCore/Public/Util/SmallListSet.h
--- a/Engine/Source/Runtime/GeometryCore/Public/Util/SmallListSet.h
+++ b/Engine/Source/Runtime/GeometryCore/Public/Util/SmallListSet.h
@@ -190,21 +190,58 @@ public:
 	 */
 	friend FArchive& operator<<(FArchive& Ar, FSmallListSet& Set)
 	{
-		Set.Serialize(Ar);
+		Set.Serialize(Ar, false, false);
 		return Ar;
 	}
 
-	/** Serialize FSmallListSet to an archive. */
-	void Serialize(FArchive& Ar)
+	/**
+	* Serialize FSmallListSet to an archive.
+	* @param Ar Archive to serialize with
+	* @param bCompactData Only serialize unique data and/or recompute redundant data when loading.
+	* @param bUseCompression Use compression to serialize data; the resulting size will likely be smaller but serialization will take significantly longer.
+	*/
+	void Serialize(FArchive& Ar, bool bCompactData, bool bUseCompression);
+
+	friend bool operator==(const FSmallListSet& Lhs, const FSmallListSet& Rhs)
 	{
-		Ar << ListHeads;
-		Ar << ListBlocks;
-		Ar << FreeBlocks;
-		Ar << AllocatedCount;
-		Ar << LinkedListElements;
-		Ar << FreeHeadIndex;
+		if (Lhs.Size() != Rhs.Size())
+		{
+			return false;
+		}
+	
+		for (int32 ListIndex = 0, ListNum = Lhs.Size(); ListIndex < ListNum; ++ListIndex)
+		{
+			if (Lhs.GetCount(ListIndex) != Rhs.GetCount(ListIndex))
+			{
+				return false;
+			}
+
+			ValueIterator ItLhs = Lhs.BeginValues(ListIndex);
+			ValueIterator ItRhs = Rhs.BeginValues(ListIndex);
+			const ValueIterator ItLhsEnd = Lhs.EndValues(ListIndex);
+			const ValueIterator ItRhsEnd = Rhs.EndValues(ListIndex);
+			while (ItLhs != ItLhsEnd && ItRhs != ItRhsEnd)
+			{
+				if (*ItLhs != *ItRhs)
+				{
+					return false;
+				}
+				++ItLhs;
+				++ItRhs;
+			}
+			if (ItLhs != ItLhsEnd || ItRhs != ItRhsEnd)
+			{
+				return false;
+			}
+		}
+		
+		return true;
 	}
 
+	friend bool operator!=(const FSmallListSet& Lhs, const FSmallListSet& Rhs)
+	{
+		return !(Lhs == Rhs);
+	}
 
 	//
 	// iterator support
@@ -479,7 +516,7 @@ public:
 protected:
 
 
-	// grab a block from the free list, or allocate a one
+	// grab a block from the free list, or allocate a new one
 	int32 AllocateBlock();
 
 	// push a link-node onto the free list
@@ -495,9 +532,9 @@ protected:
 
 
 public:
-	inline FString MemoryUsage()
+	inline FString MemoryUsage() const
 	{
-		return FString::Printf(TEXT("ListSize %d  Blocks Count %d  Free %d  Mem %dkb   Linked Mem %dkb"),
+		return FString::Printf(TEXT("ListSize %llu  Blocks Count %d  Free %llu  Mem %llukb   Linked Mem %llukb"),
 			ListHeads.GetLength(), AllocatedCount, (FreeBlocks.GetLength() * sizeof(int32) / 1024),
 			ListBlocks.GetLength(), (LinkedListElements.GetLength() * sizeof(int32) / 1024));
 	}
