GeometryCore: Fix Non-Unity build error in dynamic mesh serialization

#rnx

diff --git a/Engine/Source/Runtime/GeometryCore/Private/DynamicMesh/DynamicMesh3_Serialization.cpp b/Engine/Source/Runtime/GeometryCore/Private/DynamicMesh/DynamicMesh3_Serialization.cpp
--- a/Engine/Source/Runtime/GeometryCore/Private/DynamicMesh/DynamicMesh3_Serialization.cpp
+++ b/Engine/Source/Runtime/GeometryCore/Private/DynamicMesh/DynamicMesh3_Serialization.cpp
@@ -2,6 +2,7 @@
 
 #include "DynamicMesh/DynamicMesh3.h"
 #include "DynamicMesh/DynamicMeshAttributeSet.h"
+#include "HAL/PlatformTime.h"
 #include "UObject/UE5MainStreamObjectVersion.h"
 
 #include "ExplicitUseGeometryMathTypes.h"		// using UE::Geometry::(math types)
