GeometryCore: replace all usage of GeometryCore FVector3<T> with TVector<T>, remove FVector3<T> and GVector4<T>
#rnx

diff --git a/Engine/Plugins/Experimental/MeshModelingToolsetExp/Source/MeshModelingToolsEditorOnly/Private/Hair/GroomToMeshTool.cpp b/Engine/Plugins/Experimental/MeshModelingToolsetExp/Source/MeshModelingToolsEditorOnly/Private/Hair/GroomToMeshTool.cpp
--- a/Engine/Plugins/Experimental/MeshModelingToolsetExp/Source/MeshModelingToolsEditorOnly/Private/Hair/GroomToMeshTool.cpp
+++ b/Engine/Plugins/Experimental/MeshModelingToolsetExp/Source/MeshModelingToolsEditorOnly/Private/Hair/GroomToMeshTool.cpp
@@ -488,7 +488,7 @@ TSharedPtr<FDynamicMesh3, ESPMode::ThreadSafe> UGroomToMeshTool::UpdateVoxelizat
 	//MarchingCubes.RootMode = ERootfindingModes::Bisection;
 	MarchingCubes.RootModeSteps = 4;
 
-	TArray<FVector3<double>> Seeds;
+	TArray<FVector3d> Seeds;
 	for (int32 k = 0; k < Lines.Num(); ++k)
 	{
 		Seeds.Add( Lines[k].Segment.EndPoint() + Lines[k].GetRadius()*Lines[k].Segment.Direction );
diff --git a/Engine/Plugins/Experimental/MeshModelingToolsetExp/Source/MeshModelingToolsEditorOnly/Private/ShapeSprayTool.cpp b/Engine/Plugins/Experimental/MeshModelingToolsetExp/Source/MeshModelingToolsEditorOnly/Private/ShapeSprayTool.cpp
--- a/Engine/Plugins/Experimental/MeshModelingToolsetExp/Source/MeshModelingToolsEditorOnly/Private/ShapeSprayTool.cpp
+++ b/Engine/Plugins/Experimental/MeshModelingToolsetExp/Source/MeshModelingToolsEditorOnly/Private/ShapeSprayTool.cpp
@@ -184,7 +184,7 @@ void UShapeSprayTool::SplatShape(const FFrame3d& LocalFrame, double Scale, FDyna
 		FVector3f Normal = Rotationf * ShapeMesh.GetVertexNormal(vid);
 		TargetMesh->SetVertexNormal(VertexMap[vid], Normal);
 
-		TargetMesh->SetVertexColor(VertexMap[vid], (FVector3<float>)UseColor);
+		TargetMesh->SetVertexColor(VertexMap[vid], (FVector3f)UseColor);
 	}
 	for (int tid : ShapeMesh.TriangleIndicesItr())
 	{
diff --git a/Engine/Plugins/Experimental/MeshModelingToolsetExp/Source/MeshModelingToolsExp/Private/DrawPolyPathTool.cpp b/Engine/Plugins/Experimental/MeshModelingToolsetExp/Source/MeshModelingToolsExp/Private/DrawPolyPathTool.cpp
--- a/Engine/Plugins/Experimental/MeshModelingToolsetExp/Source/MeshModelingToolsExp/Private/DrawPolyPathTool.cpp
+++ b/Engine/Plugins/Experimental/MeshModelingToolsetExp/Source/MeshModelingToolsExp/Private/DrawPolyPathTool.cpp
@@ -487,7 +487,7 @@ void UDrawPolyPathTool::OnCompleteSurfacePath()
 			CurPathPoints[j].ConstrainedAlignAxis(0, TangentDir, PlaneNormal);
 		}
 	}
-	ArcLengths[NumPoints-1] = ArcLengths[NumPoints-2] + CurPathPoints[NumPoints-1].Origin.Distance(CurPathPoints[NumPoints - 2].Origin);
+	ArcLengths[NumPoints-1] = ArcLengths[NumPoints-2] + Distance(CurPathPoints[NumPoints-1].Origin, CurPathPoints[NumPoints - 2].Origin);
 
 	CurPolyLine.Reset();
 	for (const FFrame3d& Point : SurfacePathMechanic->HitPath)
@@ -593,7 +593,7 @@ void UDrawPolyPathTool::GeneratePathMesh(FDynamicMesh3& Mesh)
 			CurPathLength = 0;
 			for (int32 k = 1; k < NumPoints; ++k)
 			{
-				CurPathLength += CurPathPoints[k].Origin.Distance(CurPathPoints[k - 1].Origin);
+				CurPathLength += Distance(CurPathPoints[k].Origin, CurPathPoints[k - 1].Origin);
 			}
 
 			FRectangleMeshGenerator MeshGen;
diff --git a/Engine/Plugins/Experimental/MeshModelingToolsetExp/Source/MeshModelingToolsExp/Private/EditMeshPolygonsTool.cpp b/Engine/Plugins/Experimental/MeshModelingToolsetExp/Source/MeshModelingToolsExp/Private/EditMeshPolygonsTool.cpp
--- a/Engine/Plugins/Experimental/MeshModelingToolsetExp/Source/MeshModelingToolsExp/Private/EditMeshPolygonsTool.cpp
+++ b/Engine/Plugins/Experimental/MeshModelingToolsetExp/Source/MeshModelingToolsExp/Private/EditMeshPolygonsTool.cpp
@@ -1544,10 +1544,10 @@ void UEditMeshPolygonsTool::ApplyFillHole()
 				// Compute normals and UVs
 				if (Mesh->HasAttributes())
 				{
-					TArray<FVector3<double>> VertexPositions;
+					TArray<FVector3d> VertexPositions;
 					Loop.GetVertices(VertexPositions);
-					FVector3<double> PlaneOrigin;
-					FVector3<double> PlaneNormal;
+					FVector3d PlaneOrigin;
+					FVector3d PlaneNormal;
 					PolygonTriangulation::ComputePolygonPlane<double>(VertexPositions, PlaneNormal, PlaneOrigin);
 
 					FDynamicMeshEditor Editor(Mesh);
diff --git a/Engine/Plugins/Experimental/MeshModelingToolsetExp/Source/MeshModelingToolsExp/Private/Sculpting/MeshSculptToolBase.cpp b/Engine/Plugins/Experimental/MeshModelingToolsetExp/Source/MeshModelingToolsExp/Private/Sculpting/MeshSculptToolBase.cpp
--- a/Engine/Plugins/Experimental/MeshModelingToolsetExp/Source/MeshModelingToolsExp/Private/Sculpting/MeshSculptToolBase.cpp
+++ b/Engine/Plugins/Experimental/MeshModelingToolsetExp/Source/MeshModelingToolsExp/Private/Sculpting/MeshSculptToolBase.cpp
@@ -548,7 +548,7 @@ retry_frame_update:
 		LastBrushFrameWorld = NewFrame;
 	}
 
-	ActiveStrokePathArcLen += LastBrushFrameWorld.Origin.Distance(PrevBrushFrameWorld.Origin);
+	ActiveStrokePathArcLen += Distance(LastBrushFrameWorld.Origin, PrevBrushFrameWorld.Origin);
 
 	LastBrushFrameLocal = LastBrushFrameWorld;
 	LastBrushFrameLocal.Transform(CurTargetTransform.Inverse());
@@ -748,7 +748,7 @@ FFrame3d UMeshSculptToolBase::ComputeStampRegionPlane(const FFrame3d& StampFrame
 			continue;
 		}
 
-		double Distance = StampFrame.Origin.Distance(Centroid);
+		double Distance = UE::Geometry::Distance(StampFrame.Origin, Centroid);
 		double NormalizedDistance = (Distance / FalloffRadius) + 0.0001;
 
 		double Weight = Area;
diff --git a/Engine/Plugins/Experimental/MeshModelingToolsetExp/Source/ModelingComponents/Private/Components/OctreeDynamicMeshSceneProxy.h b/Engine/Plugins/Experimental/MeshModelingToolsetExp/Source/ModelingComponents/Private/Components/OctreeDynamicMeshSceneProxy.h
--- a/Engine/Plugins/Experimental/MeshModelingToolsetExp/Source/ModelingComponents/Private/Components/OctreeDynamicMeshSceneProxy.h
+++ b/Engine/Plugins/Experimental/MeshModelingToolsetExp/Source/ModelingComponents/Private/Components/OctreeDynamicMeshSceneProxy.h
@@ -69,7 +69,7 @@ public:
 	{
 		return [](int VertexID, int TriangleID, int TriVtxIdx, const FVector3f& Normal, FVector3f& TangentX, FVector3f& TangentY) -> void
 		{
-			VectorUtil::MakePerpVectors(Normal, TangentX, TangentY);
+			UE::Geometry::VectorUtil::MakePerpVectors(Normal, TangentX, TangentY);
 		};
 	}
 
diff --git a/Engine/Plugins/Experimental/MeshModelingToolsetExp/Source/ModelingComponents/Private/Drawing/UVLayoutPreview.cpp b/Engine/Plugins/Experimental/MeshModelingToolsetExp/Source/ModelingComponents/Private/Drawing/UVLayoutPreview.cpp
--- a/Engine/Plugins/Experimental/MeshModelingToolsetExp/Source/ModelingComponents/Private/Drawing/UVLayoutPreview.cpp
+++ b/Engine/Plugins/Experimental/MeshModelingToolsetExp/Source/ModelingComponents/Private/Drawing/UVLayoutPreview.cpp
@@ -149,7 +149,7 @@ void UUVLayoutPreview::UpdateUVMesh(const FDynamicMesh3* SourceMesh, int32 Sourc
 			{
 				NewUVTri[j] = NewUVOverlay->AppendElement(UVs[j]);
 				NewTri[j] = UVMesh.AppendVertex(FVector3d(UVs[j].X, UVs[j].Y, 0));
-				NewNormalTri[j] = NewNormalOverlay->AppendElement( FVector3<float>::UnitZ() );
+				NewNormalTri[j] = NewNormalOverlay->AppendElement( FVector3f::UnitZ() );
 			}
 
 			FVector2f EdgeUV1 = UVs[1] - UVs[0];
diff --git a/Engine/Plugins/Experimental/MeshModelingToolsetExp/Source/ModelingComponents/Private/Mechanics/SpaceCurveDeformationMechanic.cpp b/Engine/Plugins/Experimental/MeshModelingToolsetExp/Source/ModelingComponents/Private/Mechanics/SpaceCurveDeformationMechanic.cpp
--- a/Engine/Plugins/Experimental/MeshModelingToolsetExp/Source/ModelingComponents/Private/Mechanics/SpaceCurveDeformationMechanic.cpp
+++ b/Engine/Plugins/Experimental/MeshModelingToolsetExp/Source/ModelingComponents/Private/Mechanics/SpaceCurveDeformationMechanic.cpp
@@ -310,7 +310,7 @@ void USpaceCurveDeformationMechanic::GizmoTransformChanged(UTransformProxy* Prox
 		double AccumDist = 0;
 		for (int32 k = 1; k < NumPoints; ++k)
 		{
-			double Dist = CurveStartPositions[OrderedPoints[k]].Origin.Distance(CurveStartPositions[OrderedPoints[k-1]].Origin);
+			double Dist = Distance( CurveStartPositions[OrderedPoints[k]].Origin, CurveStartPositions[OrderedPoints[k-1]].Origin );
 			AccumDist += Dist;
 			Distances[k] = AccumDist;
 			if (OrderedPoints[k] == CenterID)
diff --git a/Engine/Plugins/Experimental/MeshModelingToolsetExp/Source/ModelingComponents/Private/Selection/GroupTopologySelector.cpp b/Engine/Plugins/Experimental/MeshModelingToolsetExp/Source/ModelingComponents/Private/Selection/GroupTopologySelector.cpp
--- a/Engine/Plugins/Experimental/MeshModelingToolsetExp/Source/ModelingComponents/Private/Selection/GroupTopologySelector.cpp
+++ b/Engine/Plugins/Experimental/MeshModelingToolsetExp/Source/ModelingComponents/Private/Selection/GroupTopologySelector.cpp
@@ -629,7 +629,7 @@ bool FGroupTopologySelector::FindSelectedElement(const FSelectionSettings& Setti
 				}
 
 				// Check whether any of the component segments intersect the rectangle
-				const TArray<FVector3<double>>& Verts = Curve.GetVertices();
+				const TArray<FVector3d>& Verts = Curve.GetVertices();
 				FVector3d CurrentVert = TargetTransform.TransformPosition(Verts[0]);
 				for (int32 i = 1; i < Verts.Num(); ++i)
 				{
diff --git a/Engine/Plugins/Experimental/MeshModelingToolsetExp/Source/ModelingComponents/Private/Selection/PolygonSelectionMechanic.cpp b/Engine/Plugins/Experimental/MeshModelingToolsetExp/Source/ModelingComponents/Private/Selection/PolygonSelectionMechanic.cpp
--- a/Engine/Plugins/Experimental/MeshModelingToolsetExp/Source/ModelingComponents/Private/Selection/PolygonSelectionMechanic.cpp
+++ b/Engine/Plugins/Experimental/MeshModelingToolsetExp/Source/ModelingComponents/Private/Selection/PolygonSelectionMechanic.cpp
@@ -119,7 +119,7 @@ void UPolygonSelectionMechanic::Initialize(
 			// dividing our field of view into 90 visual angle degrees, we divide the plane into 90 units).
 			float OrthoTolerance = ToolSceneQueriesUtil::GetDefaultVisualAngleSnapThreshD() * CameraState.OrthoWorldCoordinateWidth / 90.0;
 			OrthoTolerance *= TolScale;
-			return TargetTransform.TransformPosition(Position1).DistanceSquared(TargetTransform.TransformPosition(Position2)) < OrthoTolerance * OrthoTolerance;
+			return DistanceSquared( TargetTransform.TransformPosition(Position1), TargetTransform.TransformPosition(Position2) ) < OrthoTolerance * OrthoTolerance;
 		}
 		else
 		{
diff --git a/Engine/Plugins/Experimental/MeshModelingToolsetExp/Source/ModelingOperators/Private/CleaningOps/HoleFillOp.cpp b/Engine/Plugins/Experimental/MeshModelingToolsetExp/Source/ModelingOperators/Private/CleaningOps/HoleFillOp.cpp
--- a/Engine/Plugins/Experimental/MeshModelingToolsetExp/Source/ModelingOperators/Private/CleaningOps/HoleFillOp.cpp
+++ b/Engine/Plugins/Experimental/MeshModelingToolsetExp/Source/ModelingOperators/Private/CleaningOps/HoleFillOp.cpp
@@ -368,10 +368,10 @@ void FHoleFillOp::CalculateResult(FProgressCancel* Progress)
 		NewGroupIDs.Add(NewGroupID);
 
 		// Compute a best-fit plane of the boundary vertices
-		TArray<FVector3<double>> VertexPositions;
+		TArray<FVector3d> VertexPositions;
 		Loop.GetVertices(VertexPositions);
-		FVector3<double> PlaneOrigin;
-		FVector3<double> PlaneNormal;
+		FVector3d PlaneOrigin;
+		FVector3d PlaneNormal;
 		PolygonTriangulation::ComputePolygonPlane<double>(VertexPositions, PlaneNormal, PlaneOrigin);
 		PlaneNormal *= -1.0;	// Previous function seems to orient the normal opposite to what's expected elsewhere
 
diff --git a/Engine/Plugins/Experimental/MeshModelingToolsetExp/Source/ModelingOperators/Private/CompositionOps/CubeGridBooleanOp.cpp b/Engine/Plugins/Experimental/MeshModelingToolsetExp/Source/ModelingOperators/Private/CompositionOps/CubeGridBooleanOp.cpp
--- a/Engine/Plugins/Experimental/MeshModelingToolsetExp/Source/ModelingOperators/Private/CompositionOps/CubeGridBooleanOp.cpp
+++ b/Engine/Plugins/Experimental/MeshModelingToolsetExp/Source/ModelingOperators/Private/CompositionOps/CubeGridBooleanOp.cpp
@@ -153,8 +153,8 @@ protected:
 				CornerToVert[RightUp] = AppendVertex(Box.GetCorner(RightUp));
 			}
 
-			FaceWidth = Box.GetCorner(LeftCorner).Distance(Box.GetCorner(RightCorner));
-			FaceHeight = Box.GetCorner(LeftCorner).Distance(Box.GetCorner(LeftUp));
+			FaceWidth = Distance(Box.GetCorner(LeftCorner), Box.GetCorner(RightCorner));
+			FaceHeight = Distance(Box.GetCorner(LeftCorner), Box.GetCorner(LeftUp));
 
 			// The bottom two vertices will definitely need uv/normal elements
 			MakeUVsAndNormals(0, SideFaceBoxSpaceNormals[LeftCorner], CornerToVert[LeftCorner]);
@@ -209,8 +209,8 @@ protected:
 
 		// Triangulate the bottom:
 		// Prep for UVs
-		FaceWidth = Box.GetCorner(0).Distance(Box.GetCorner(1));
-		FaceHeight = Box.GetCorner(0).Distance(Box.GetCorner(3));
+		FaceWidth = Distance(Box.GetCorner(0), Box.GetCorner(1));
+		FaceHeight = Distance(Box.GetCorner(0), Box.GetCorner(3));
 
 		MakeUVsAndNormals(DiagFaceIdx1, -Box.Frame.Z(), CornerToVert[DiagFaceIdx1]);
 		MakeUVsAndNormals(DiagFaceIdx2, -Box.Frame.Z(), CornerToVert[DiagFaceIdx2]);
diff --git a/Engine/Plugins/Experimental/MeshModelingToolsetExp/Source/ModelingOperators/Private/CompositionOps/CurveSweepOp.cpp b/Engine/Plugins/Experimental/MeshModelingToolsetExp/Source/ModelingOperators/Private/CompositionOps/CurveSweepOp.cpp
--- a/Engine/Plugins/Experimental/MeshModelingToolsetExp/Source/ModelingOperators/Private/CompositionOps/CurveSweepOp.cpp
+++ b/Engine/Plugins/Experimental/MeshModelingToolsetExp/Source/ModelingOperators/Private/CompositionOps/CurveSweepOp.cpp
@@ -89,10 +89,10 @@ void FCurveSweepOp::CalculateResult(FProgressCancel* Progress)
 			}
 
 			// Compute a best-fit plane of the vertices
-			TArray<FVector3<double>> VertexPositions;
+			TArray<FVector3d> VertexPositions;
 			Loop.GetVertices(VertexPositions);
-			FVector3<double> PlaneOrigin;
-			FVector3<double> PlaneNormal;
+			FVector3d PlaneOrigin;
+			FVector3d PlaneNormal;
 			PolygonTriangulation::ComputePolygonPlane<double>(VertexPositions, PlaneNormal, PlaneOrigin);
 			PlaneNormal *= -1.0;	// Previous function seems to orient the normal opposite to what's expected elsewhere
 
diff --git a/Engine/Plugins/Experimental/PlanarCutPlugin/Source/PlanarCut/Private/GeometryMeshConversion.cpp b/Engine/Plugins/Experimental/PlanarCutPlugin/Source/PlanarCut/Private/GeometryMeshConversion.cpp
--- a/Engine/Plugins/Experimental/PlanarCutPlugin/Source/PlanarCut/Private/GeometryMeshConversion.cpp
+++ b/Engine/Plugins/Experimental/PlanarCutPlugin/Source/PlanarCut/Private/GeometryMeshConversion.cpp
@@ -462,8 +462,8 @@ namespace AugmentedDynamicMesh
 		FMeshTangentsf Tangents(&Mesh);
 		Tangents.ComputeTriVertexTangents(Normals, UVs, Options);
 
-		const TArray<FVector3<float>>& TanU = Tangents.GetTangents();
-		const TArray<FVector3<float>>& TanV = Tangents.GetBitangents();
+		const TArray<FVector3f>& TanU = Tangents.GetTangents();
+		const TArray<FVector3f>& TanV = Tangents.GetBitangents();
 		FDynamicMeshMaterialAttribute* MaterialIDs = Mesh.Attributes()->GetMaterialID();
 		for (int TID : Mesh.TriangleIndicesItr())
 		{
@@ -571,7 +571,7 @@ namespace AugmentedDynamicMesh
 						E2Bary[SecondLongestSecondEdgeIdx] = Along;
 
 						// Choose number of samples between the two edge points based on their distance
-						double AcrossDist = Triangle.BarycentricPoint(E1Bary).Distance(Triangle.BarycentricPoint(E2Bary));
+						double AcrossDist = Distance( Triangle.BarycentricPoint(E1Bary), Triangle.BarycentricPoint(E2Bary) );
 						int DivisionsAcross = FMathd::Ceil(AcrossDist / Spacing);
 						double FactorAcross = 1.0 / double(DivisionsAcross + 1);
 						for (int DivJ = 0; DivJ < DivisionsAcross; DivJ++)
@@ -883,7 +883,7 @@ namespace AugmentedDynamicMesh
 			check(ChainLen > 2 && EndIdx <= AllBoundaryVerts.Num());
 
 			TArray<FIndex3i> Triangles;
-			TArray<UE::Geometry::FVector3<double>> VertexPositions;
+			TArray<FVector3d> VertexPositions;
 			for (int Idx = BIdx + 1; Idx < EndIdx; Idx++)
 			{
 				VertexPositions.Add(Mesh.GetVertex(AllBoundaryVerts[Idx]));
@@ -1059,7 +1059,7 @@ namespace
 	{
 		FVector3d Origin = (FVector3d)Plane.GetOrigin();
 		FVector3d Normal = (FVector3d)Plane.GetNormal();
-		int32 MinIdx = VectorUtil::MinAbsElementIndex(Normal);
+		int32 MinIdx = UE::Geometry::MinAbsElementIndex(Normal);
 		FFrame3d Frame(Origin, Normal);
 		FVector3d Target(0, 0, 0);
 		Target[MinIdx] = 1;
diff --git a/Engine/Plugins/Runtime/GeometryProcessing/Source/DynamicMesh/Private/MeshSimplification.cpp b/Engine/Plugins/Runtime/GeometryProcessing/Source/DynamicMesh/Private/MeshSimplification.cpp
--- a/Engine/Plugins/Runtime/GeometryProcessing/Source/DynamicMesh/Private/MeshSimplification.cpp
+++ b/Engine/Plugins/Runtime/GeometryProcessing/Source/DynamicMesh/Private/MeshSimplification.cpp
@@ -1664,7 +1664,7 @@ void DYNAMICMESH_API TMeshSimplification<FAttrBasedQuadricErrord>::OnEdgeCollaps
 	FAttrBasedQuadricErrord& Quadric = EdgeQuadrics[edgeID].q;
 	FVector3d collapse_pt = EdgeQuadrics[edgeID].collapse_pt;
 
-	FVector3<double> UpdatedNormald;
+	FVector3d UpdatedNormald;
 	Quadric.ComputeAttributes(collapse_pt, UpdatedNormald);
 
 	FVector3f UpdatedNormal((float)UpdatedNormald.X, (float)UpdatedNormald.Y, (float)UpdatedNormald.Z);
diff --git a/Engine/Plugins/Runtime/GeometryProcessing/Source/DynamicMesh/Private/Operations/MeshMeshCut.cpp b/Engine/Plugins/Runtime/GeometryProcessing/Source/DynamicMesh/Private/Operations/MeshMeshCut.cpp
--- a/Engine/Plugins/Runtime/GeometryProcessing/Source/DynamicMesh/Private/Operations/MeshMeshCut.cpp
+++ b/Engine/Plugins/Runtime/GeometryProcessing/Source/DynamicMesh/Private/Operations/MeshMeshCut.cpp
@@ -476,7 +476,7 @@ namespace MeshCut
 			int BestIdx = -1;
 			for (int SubIdx = 0; SubIdx < 3; SubIdx++)
 			{
-				double DSq = Tri.V[SubIdx].DistanceSquared(V);
+				double DSq = DistanceSquared(Tri.V[SubIdx], V);
 				if (DSq < BestDSq)
 				{
 					BestIdx = SubIdx;
diff --git a/Engine/Plugins/Runtime/GeometryProcessing/Source/DynamicMesh/Private/Operations/SimpleHoleFiller.cpp b/Engine/Plugins/Runtime/GeometryProcessing/Source/DynamicMesh/Private/Operations/SimpleHoleFiller.cpp
--- a/Engine/Plugins/Runtime/GeometryProcessing/Source/DynamicMesh/Private/Operations/SimpleHoleFiller.cpp
+++ b/Engine/Plugins/Runtime/GeometryProcessing/Source/DynamicMesh/Private/Operations/SimpleHoleFiller.cpp
@@ -82,7 +82,7 @@ bool FSimpleHoleFiller::Fill_Fan(int GroupID)
 
 bool FSimpleHoleFiller::Fill_EarClip(int GroupID)
 {
-	TArray<FVector3<double>> Vertices;
+	TArray<FVector3d> Vertices;
 	int32 NumVertices = Loop.GetVertexCount();
 	for (int32 i = 0; i < NumVertices; ++i)
 	{
diff --git a/Engine/Plugins/Runtime/GeometryProcessing/Source/DynamicMesh/Private/ShapeApproximation/ShapeDetection3.cpp b/Engine/Plugins/Runtime/GeometryProcessing/Source/DynamicMesh/Private/ShapeApproximation/ShapeDetection3.cpp
--- a/Engine/Plugins/Runtime/GeometryProcessing/Source/DynamicMesh/Private/ShapeApproximation/ShapeDetection3.cpp
+++ b/Engine/Plugins/Runtime/GeometryProcessing/Source/DynamicMesh/Private/ShapeApproximation/ShapeDetection3.cpp
@@ -62,7 +62,7 @@ bool UE::Geometry::IsSphereMesh(const FDynamicMesh3& Mesh, FSphere3d& SphereOut,
 		SphereOut.Center = Centroid + AvgLen * AvgLenDeriv;
 		SphereOut.Radius = AvgLen;
 
-		if (SphereOut.Center.DistanceSquared(PrevCenter) < UseFitTolerance)
+		if (DistanceSquared(SphereOut.Center, PrevCenter) < UseFitTolerance)
 		{
 			break;
 		}
diff --git a/Engine/Plugins/Runtime/GeometryProcessing/Source/GeometryAlgorithms/Private/FitCapsule3.cpp b/Engine/Plugins/Runtime/GeometryProcessing/Source/GeometryAlgorithms/Private/FitCapsule3.cpp
--- a/Engine/Plugins/Runtime/GeometryProcessing/Source/GeometryAlgorithms/Private/FitCapsule3.cpp
+++ b/Engine/Plugins/Runtime/GeometryProcessing/Source/GeometryAlgorithms/Private/FitCapsule3.cpp
@@ -6,9 +6,10 @@
 
 #include "ExplicitUseGeometryMathTypes.h"		// using UE::Geometry::(math types)
 using namespace UE::Geometry;
+using namespace UE::Math;
 
 template<typename RealType>
-bool TFitCapsule3<RealType>::Solve(int32 NumPoints, TFunctionRef<FVector3<RealType>(int32)> GetPointFunc)
+bool TFitCapsule3<RealType>::Solve(int32 NumPoints, TFunctionRef<TVector<RealType>(int32)> GetPointFunc)
 {
 	using ComputeType = double;
 
@@ -16,7 +17,7 @@ bool TFitCapsule3<RealType>::Solve(int32 NumPoints, TFunctionRef<FVector3<RealTy
 	PointList.SetNum(NumPoints);
 	for (int32 k = 0; k < NumPoints; ++k)
 	{
-		FVector3<RealType> Point = GetPointFunc(k);
+		TVector<RealType> Point = GetPointFunc(k);
 		PointList[k] = { {(ComputeType)Point.X, (ComputeType)Point.Y, (ComputeType)Point.Z} };
 	}
 
@@ -29,8 +30,8 @@ bool TFitCapsule3<RealType>::Solve(int32 NumPoints, TFunctionRef<FVector3<RealTy
 		FitCapsule.segment.GetCenteredForm(Center, Direction, Extent);
 
 		Capsule.Segment = TSegment3<RealType>(
-			FVector3<RealType>((RealType)Center[0], (RealType)Center[1], (RealType)Center[2]),
-			FVector3<RealType>((RealType)Direction[0], (RealType)Direction[1], (RealType)Direction[2]),
+			TVector<RealType>((RealType)Center[0], (RealType)Center[1], (RealType)Center[2]),
+			TVector<RealType>((RealType)Direction[0], (RealType)Direction[1], (RealType)Direction[2]),
 			(RealType)Extent);
 
 		Capsule.Radius = (RealType)FitCapsule.radius;
diff --git a/Engine/Plugins/Runtime/GeometryProcessing/Source/GeometryAlgorithms/Private/GteUtil.h b/Engine/Plugins/Runtime/GeometryProcessing/Source/GeometryAlgorithms/Private/GteUtil.h
--- a/Engine/Plugins/Runtime/GeometryProcessing/Source/GeometryAlgorithms/Private/GteUtil.h
+++ b/Engine/Plugins/Runtime/GeometryProcessing/Source/GeometryAlgorithms/Private/GteUtil.h
@@ -31,15 +31,15 @@ UE::Math::TVector2<Real> Convert(const gte::Vector2<Real>& Vec)
 }
 
 template<typename Real>
-gte::Vector3<Real> Convert(const FVector3<Real>& Vec)
+gte::Vector3<Real> Convert(const UE::Math::TVector<Real>& Vec)
 {
 	return gte::Vector3<Real>({ Vec.X, Vec.Y, Vec.Z });
 }
 
 template<typename Real>
-FVector3<Real> Convert(const gte::Vector3<Real>& Vec)
+UE::Math::TVector<Real> Convert(const gte::Vector3<Real>& Vec)
 {
-	return FVector3<Real>(Vec[0], Vec[1], Vec[2]);
+	return UE::Math::TVector<Real>(Vec[0], Vec[1], Vec[2]);
 }
 
 template<typename Real>
diff --git a/Engine/Plugins/Runtime/GeometryProcessing/Source/GeometryAlgorithms/Private/MinVolumeBox3.cpp b/Engine/Plugins/Runtime/GeometryProcessing/Source/GeometryAlgorithms/Private/MinVolumeBox3.cpp
--- a/Engine/Plugins/Runtime/GeometryProcessing/Source/GeometryAlgorithms/Private/MinVolumeBox3.cpp
+++ b/Engine/Plugins/Runtime/GeometryProcessing/Source/GeometryAlgorithms/Private/MinVolumeBox3.cpp
@@ -26,6 +26,7 @@ THIRD_PARTY_INCLUDES_END
 
 #include "ExplicitUseGeometryMathTypes.h"		// using UE::Geometry::(math types)
 using namespace UE::Geometry;
+using namespace UE::Math;
 
 namespace UE {
 namespace Geometry {
@@ -43,7 +44,7 @@ struct TMinVolumeBox3Internal
 	TOrientedBox3<RealType> Result;
 	bool bSolutionOK;
 
-	void SetPoint(int32 Index, const FVector3<RealType>& Point)
+	void SetPoint(int32 Index, const TVector<RealType>& Point)
 	{
 		DoubleInput[Index] = DVector3{ {(double)Point.X, (double)Point.Y, (double)Point.Z} };
 	}
@@ -99,11 +100,11 @@ struct TMinVolumeBox3Internal
 		}
 
 		Result.Frame = TFrame3<RealType>(
-			FVector3<RealType>((RealType)MinimalBox.center[0], (RealType)MinimalBox.center[1], (RealType)MinimalBox.center[2]),
-			FVector3<RealType>((RealType)MinimalBox.axis[0][0], (RealType)MinimalBox.axis[0][1], (RealType)MinimalBox.axis[0][2]),
-			FVector3<RealType>((RealType)MinimalBox.axis[1][0], (RealType)MinimalBox.axis[1][1], (RealType)MinimalBox.axis[1][2]),
-			FVector3<RealType>((RealType)MinimalBox.axis[2][0], (RealType)MinimalBox.axis[2][1], (RealType)MinimalBox.axis[2][2]) );
-		Result.Extents = FVector3<RealType>((RealType)MinimalBox.extent[0], (RealType)MinimalBox.extent[1], (RealType)MinimalBox.extent[2]);
+			TVector<RealType>((RealType)MinimalBox.center[0], (RealType)MinimalBox.center[1], (RealType)MinimalBox.center[2]),
+			TVector<RealType>((RealType)MinimalBox.axis[0][0], (RealType)MinimalBox.axis[0][1], (RealType)MinimalBox.axis[0][2]),
+			TVector<RealType>((RealType)MinimalBox.axis[1][0], (RealType)MinimalBox.axis[1][1], (RealType)MinimalBox.axis[1][2]),
+			TVector<RealType>((RealType)MinimalBox.axis[2][0], (RealType)MinimalBox.axis[2][1], (RealType)MinimalBox.axis[2][2]) );
+		Result.Extents = TVector<RealType>((RealType)MinimalBox.extent[0], (RealType)MinimalBox.extent[1], (RealType)MinimalBox.extent[2]);
 
 		return true;
 	}
@@ -114,14 +115,14 @@ struct TMinVolumeBox3Internal
 } // end namespace UE
 
 template<typename RealType>
-bool TMinVolumeBox3<RealType>::Solve(int32 NumPoints, TFunctionRef<FVector3<RealType>(int32)> GetPointFunc, bool bUseExactBox, FProgressCancel* Progress )
+bool TMinVolumeBox3<RealType>::Solve(int32 NumPoints, TFunctionRef<TVector<RealType>(int32)> GetPointFunc, bool bUseExactBox, FProgressCancel* Progress )
 {
 	Initialize(NumPoints, bUseExactBox);
 	check(Internal);
 
 	for (int32 k = 0; k < NumPoints; ++k)
 	{
-		FVector3<RealType> Point = GetPointFunc(k);
+		TVector<RealType> Point = GetPointFunc(k);
 		Internal->SetPoint(k, Point);
 	}
 	
@@ -130,7 +131,7 @@ bool TMinVolumeBox3<RealType>::Solve(int32 NumPoints, TFunctionRef<FVector3<Real
 
 
 template<typename RealType>
-bool TMinVolumeBox3<RealType>::SolveSubsample(int32 NumPoints, int32 MaxPoints, TFunctionRef<FVector3<RealType>(int32)> GetPointFunc, bool bUseExactBox, FProgressCancel* Progress)
+bool TMinVolumeBox3<RealType>::SolveSubsample(int32 NumPoints, int32 MaxPoints, TFunctionRef<TVector<RealType>(int32)> GetPointFunc, bool bUseExactBox, FProgressCancel* Progress)
 {
 	if (NumPoints <= MaxPoints)
 	{
@@ -145,7 +146,7 @@ bool TMinVolumeBox3<RealType>::SolveSubsample(int32 NumPoints, int32 MaxPoints,
 	int32 Index;
 	while (Iter.GetNextIndex(Index) && k < MaxPoints)
 	{
-		FVector3<RealType> Point = GetPointFunc(Index);
+		TVector<RealType> Point = GetPointFunc(Index);
 		Internal->SetPoint(k, Point);
 		k++;
 	}
diff --git a/Engine/Plugins/Runtime/GeometryProcessing/Source/GeometryAlgorithms/Private/MinVolumeSphere3.cpp b/Engine/Plugins/Runtime/GeometryProcessing/Source/GeometryAlgorithms/Private/MinVolumeSphere3.cpp
--- a/Engine/Plugins/Runtime/GeometryProcessing/Source/GeometryAlgorithms/Private/MinVolumeSphere3.cpp
+++ b/Engine/Plugins/Runtime/GeometryProcessing/Source/GeometryAlgorithms/Private/MinVolumeSphere3.cpp
@@ -10,7 +10,7 @@
 
 #include "ExplicitUseGeometryMathTypes.h"		// using UE::Geometry::(math types)
 using namespace UE::Geometry;
-
+using namespace UE::Math;
 
 namespace UE {
 namespace Geometry {
@@ -25,12 +25,12 @@ struct TMinVolumeSphere3Internal
 	bool bUseExact;
 	TArray<DVector3> DoubleInput;
 
-	FVector3<RealType> Center;
+	TVector<RealType> Center;
 	RealType Radius;
 	bool bIsMinimalSphere;
 	bool bSolutionOK;
 
-	void SetPoint(int32 Index, const FVector3<RealType>& Point)
+	void SetPoint(int32 Index, const TVector<RealType>& Point)
 	{
 		DoubleInput[Index] = DVector3{ {(double)Point.X, (double)Point.Y, (double)Point.Z} };
 	}
@@ -53,7 +53,7 @@ struct TMinVolumeSphere3Internal
 			bSolutionOK = true;
 		}
 
-		Center = FVector3<RealType>((RealType)MinimalSphere.center[0], (RealType)MinimalSphere.center[1], (RealType)MinimalSphere.center[2]);
+		Center = TVector<RealType>((RealType)MinimalSphere.center[0], (RealType)MinimalSphere.center[1], (RealType)MinimalSphere.center[2]);
 		Radius = (RealType)MinimalSphere.radius;
 
 		return true;
@@ -67,14 +67,14 @@ struct TMinVolumeSphere3Internal
 
 
 template<typename RealType>
-bool TMinVolumeSphere3<RealType>::Solve(int32 NumPoints, TFunctionRef<FVector3<RealType>(int32)> GetPointFunc, bool bUseExactComputation)
+bool TMinVolumeSphere3<RealType>::Solve(int32 NumPoints, TFunctionRef<TVector<RealType>(int32)> GetPointFunc, bool bUseExactComputation)
 {
 	Initialize(NumPoints, bUseExactComputation);
 	check(Internal);
 
 	for (int32 k = 0; k < NumPoints; ++k)
 	{
-		FVector3<RealType> Point = GetPointFunc(k);
+		TVector<RealType> Point = GetPointFunc(k);
 		Internal->SetPoint(k, Point);
 	}
 	
diff --git a/Engine/Plugins/Runtime/GeometryProcessing/Source/GeometryAlgorithms/Public/Distance/DistLine3Circle3.h b/Engine/Plugins/Runtime/GeometryProcessing/Source/GeometryAlgorithms/Public/Distance/DistLine3Circle3.h
--- a/Engine/Plugins/Runtime/GeometryProcessing/Source/GeometryAlgorithms/Public/Distance/DistLine3Circle3.h
+++ b/Engine/Plugins/Runtime/GeometryProcessing/Source/GeometryAlgorithms/Public/Distance/DistLine3Circle3.h
@@ -12,6 +12,7 @@
 namespace UE {
 namespace Geometry {
 
+using namespace UE::Math;
 
 /**
 * Compute unsigned distance and closest-points between 3D line and 3D circle
@@ -28,8 +29,8 @@ public:
 	Real DistanceSquared = -1.0;
 
 	int NumClosestPairs = 0;
-	FVector3<Real> LineClosest[2];
-	FVector3<Real> CircleClosest[2];
+	TVector<Real> LineClosest[2];
+	TVector<Real> CircleClosest[2];
 	bool bIsEquiDistant;
 
 	TDistLine3Circle3(const TLine3<Real>& LineIn, const TCircle3<Real>& CircleIn) : Line(LineIn), Circle(CircleIn)
diff --git a/Engine/Plugins/Runtime/GeometryProcessing/Source/GeometryAlgorithms/Public/FitCapsule3.h b/Engine/Plugins/Runtime/GeometryProcessing/Source/GeometryAlgorithms/Public/FitCapsule3.h
--- a/Engine/Plugins/Runtime/GeometryProcessing/Source/GeometryAlgorithms/Public/FitCapsule3.h
+++ b/Engine/Plugins/Runtime/GeometryProcessing/Source/GeometryAlgorithms/Public/FitCapsule3.h
@@ -9,7 +9,7 @@
 namespace UE {
 namespace Geometry {
 
-
+using namespace UE::Math;
 
 /**
  * Fit a Capsule to a set of 3D points.
@@ -36,13 +36,13 @@ public:
 	  * Calculate a Capsule that contains the given Point Set and store in Output variables
 	  * @return true if capsule was found
 	  */
-	bool Solve(int32 NumPoints, TFunctionRef<FVector3<RealType>(int32)> GetPointFunc);
+	bool Solve(int32 NumPoints, TFunctionRef<TVector<RealType>(int32)> GetPointFunc);
 
 	/**
 	 * Calculate a Capsule that contains the given Point Set and return in CapsuleOut
 	 * @return true if capsule was found
 	 */
-	static bool Solve(int32 NumPoints, TFunctionRef<FVector3<RealType>(int32)> GetPointFunc,
+	static bool Solve(int32 NumPoints, TFunctionRef<TVector<RealType>(int32)> GetPointFunc,
 		TCapsule3<RealType>& CapsuleOut)
 	{
 		TFitCapsule3<RealType> Compute;
diff --git a/Engine/Plugins/Runtime/GeometryProcessing/Source/GeometryAlgorithms/Public/MinVolumeBox3.h b/Engine/Plugins/Runtime/GeometryProcessing/Source/GeometryAlgorithms/Public/MinVolumeBox3.h
--- a/Engine/Plugins/Runtime/GeometryProcessing/Source/GeometryAlgorithms/Public/MinVolumeBox3.h
+++ b/Engine/Plugins/Runtime/GeometryProcessing/Source/GeometryAlgorithms/Public/MinVolumeBox3.h
@@ -12,6 +12,8 @@ class FProgressCancel;
 namespace UE {
 namespace Geometry {
 
+using namespace UE::Math;
+
 template <typename RealType> struct TMinVolumeBox3Internal;
 
 /**
@@ -31,7 +33,7 @@ public:
 	 * @param bUseExactBox If true, high-precision number types are used for the minimal-box calculation, rather than doubles. This is *much* slower but more accurate (but not recommended).
 	 * @return true if minimal box was found
 	 */
-	bool Solve(int32 NumPoints, TFunctionRef<FVector3<RealType>(int32)> GetPointFunc, bool bUseExactBox = false, FProgressCancel* Progress = nullptr);
+	bool Solve(int32 NumPoints, TFunctionRef<TVector<RealType>(int32)> GetPointFunc, bool bUseExactBox = false, FProgressCancel* Progress = nullptr);
 
 	/**
 	 * Calculate the minimal box for a Subsampling of MaxPoints points of a point set
@@ -41,7 +43,7 @@ public:
 	 * @param bUseExactBox If true, high-precision number types are used for the minimal-box calculation, rather than doubles. This is *much* slower but more accurate (but not recommended).
 	 * @return true if minimal box was found
 	 */
-	bool SolveSubsample(int32 NumPoints, int32 NumSamplePoints, TFunctionRef<FVector3<RealType>(int32)> GetPointFunc, bool bUseExactBox = false, FProgressCancel* Progress = nullptr);
+	bool SolveSubsample(int32 NumPoints, int32 NumSamplePoints, TFunctionRef<TVector<RealType>(int32)> GetPointFunc, bool bUseExactBox = false, FProgressCancel* Progress = nullptr);
 
 
 	/** @return true if minimal box is available */
diff --git a/Engine/Plugins/Runtime/GeometryProcessing/Source/GeometryAlgorithms/Public/MinVolumeSphere3.h b/Engine/Plugins/Runtime/GeometryProcessing/Source/GeometryAlgorithms/Public/MinVolumeSphere3.h
--- a/Engine/Plugins/Runtime/GeometryProcessing/Source/GeometryAlgorithms/Public/MinVolumeSphere3.h
+++ b/Engine/Plugins/Runtime/GeometryProcessing/Source/GeometryAlgorithms/Public/MinVolumeSphere3.h
@@ -10,6 +10,8 @@
 namespace UE {
 namespace Geometry {
 
+using namespace UE::Math;
+
 template <typename RealType> struct TMinVolumeSphere3Internal;
 
 /**
@@ -26,7 +28,7 @@ public:
 	 * @param bUseExactComputation If true, high-precision Rational number types are used for the calculation, rather than doubles. This is slower but more precise.
 	 * @return true if minimal sphere was found
 	 */
-	bool Solve(int32 NumPoints, TFunctionRef<FVector3<RealType>(int32)> GetPointFunc, bool bUseExactComputation = false);
+	bool Solve(int32 NumPoints, TFunctionRef<TVector<RealType>(int32)> GetPointFunc, bool bUseExactComputation = false);
 
 	/** @return true if minimal box is available */
 	bool IsSolutionAvailable() const;
diff --git a/Engine/Source/Programs/HeadlessChaos/Private/GeometryCollection/GeometryCollectionTestFramework.cpp b/Engine/Source/Programs/HeadlessChaos/Private/GeometryCollection/GeometryCollectionTestFramework.cpp
--- a/Engine/Source/Programs/HeadlessChaos/Private/GeometryCollection/GeometryCollectionTestFramework.cpp
+++ b/Engine/Source/Programs/HeadlessChaos/Private/GeometryCollection/GeometryCollectionTestFramework.cpp
@@ -120,7 +120,7 @@ namespace GeometryCollectionTest
 		const UE::Geometry::FIndex3i& EdgeVertices = UE::Geometry::FIndex3i(4, 4, 4),
 		const int NumberOfMaterials = 2)
 	{
-		using GeomVector = UE::Geometry::FVector3<double>;
+		using GeomVector = UE::Math::TVector<double>;
 		UE::Geometry::FGridBoxMeshGenerator BoxGen;
 		BoxGen.Box = UE::Geometry::FOrientedBox3d(GeomVector::Zero(), GeomVector(Extents)); // box center, box dimensions
 		BoxGen.EdgeVertices = EdgeVertices;
diff --git a/Engine/Source/Runtime/GeometryCore/Private/CompGeom/ConvexHull3.cpp b/Engine/Source/Runtime/GeometryCore/Private/CompGeom/ConvexHull3.cpp
--- a/Engine/Source/Runtime/GeometryCore/Private/CompGeom/ConvexHull3.cpp
+++ b/Engine/Source/Runtime/GeometryCore/Private/CompGeom/ConvexHull3.cpp
@@ -29,12 +29,12 @@ struct TExtremePoints3
 	int Extreme[4]{ 0, 0, 0, 0 };
 
 	// Coordinate frame spanned by input points
-	FVector3<RealType> Origin{ 0,0,0 };
-	FVector3<RealType> Basis[3]{ {0,0,0}, {0,0,0}, {0,0,0} };
+	TVector<RealType> Origin{ 0,0,0 };
+	TVector<RealType> Basis[3]{ {0,0,0}, {0,0,0}, {0,0,0} };
 
-	TExtremePoints3(int32 NumPoints, TFunctionRef<FVector3<RealType>(int32)> GetPointFunc, TFunctionRef<bool(int32)> FilterFunc, double Epsilon = 0)
+	TExtremePoints3(int32 NumPoints, TFunctionRef<TVector<RealType>(int32)> GetPointFunc, TFunctionRef<bool(int32)> FilterFunc, double Epsilon = 0)
 	{
-		FVector3<RealType> FirstPoint;
+		TVector<RealType> FirstPoint;
 		int FirstPtIdx = -1;
 		for (FirstPtIdx = 0; FirstPtIdx < NumPoints; FirstPtIdx++)
 		{
@@ -51,7 +51,7 @@ struct TExtremePoints3
 			return;
 		}
 
-		FVector3<RealType> Min = GetPointFunc(FirstPtIdx), Max = GetPointFunc(FirstPtIdx);
+		TVector<RealType> Min = GetPointFunc(FirstPtIdx), Max = GetPointFunc(FirstPtIdx);
 		FIndex3i IndexMin(FirstPtIdx, FirstPtIdx, FirstPtIdx), IndexMax(FirstPtIdx, FirstPtIdx, FirstPtIdx);
 		for (int Idx = FirstPtIdx + 1; Idx < NumPoints; Idx++)
 		{
@@ -258,7 +258,7 @@ struct FHullConnectivity
 	/**
 	 * Set triangle positions to prep for IsVisible calls
 	 */
-	void SetTriPts(const FIndex3i& Tri, TFunctionRef<FVector3<RealType>(int32)> GetPointFunc,  FVector3d TriPtsOut[3])
+	void SetTriPts(const FIndex3i& Tri, TFunctionRef<TVector<RealType>(int32)> GetPointFunc,  FVector3d TriPtsOut[3])
 	{
 		TriPtsOut[0] = (FVector3d)GetPointFunc(Tri[0]);
 		TriPtsOut[1] = (FVector3d)GetPointFunc(Tri[1]);
@@ -268,7 +268,7 @@ struct FHullConnectivity
 	/**
 	 * Bucket all points based on which triangles can see them
 	 */
-	void InitVisibility(const TArray<FIndex3i>& Triangles, int32 NumPoints, TFunctionRef<FVector3<RealType>(int32)> GetPointFunc, TFunctionRef<bool(int32)> FilterFunc)
+	void InitVisibility(const TArray<FIndex3i>& Triangles, int32 NumPoints, TFunctionRef<TVector<RealType>(int32)> GetPointFunc, TFunctionRef<bool(int32)> FilterFunc)
 	{
 		VisiblePoints.SetNum(Triangles.Num());
 
@@ -490,7 +490,7 @@ struct FHullConnectivity
 	 * @param CrossedEdgeFirstVertex The first vertex of the edge that was 'crossed over' to traverse to this TriIdx, or -1 for the initial call
 	 * @return false if triangle is beyond horizon (so we don't need to search through it / remove it), true otherwise
 	 */
-	bool HorizonHelper(const TArray<FIndex3i>& Triangles, TFunctionRef<FVector3<RealType>(int32)> GetPointFunc, const FVector3d& Pt, TArray<int32>& NewlyUnclaimed, TSet<int32>& ToDelete, TArray<FNewTriangle>& ToAdd, int32 TriIdx, int32 CrossedEdgeFirstVertex)
+	bool HorizonHelper(const TArray<FIndex3i>& Triangles, TFunctionRef<TVector<RealType>(int32)> GetPointFunc, const FVector3d& Pt, TArray<int32>& NewlyUnclaimed, TSet<int32>& ToDelete, TArray<FNewTriangle>& ToAdd, int32 TriIdx, int32 CrossedEdgeFirstVertex)
 	{
 		// if it's not the first triangle, crossed edge should be set and we should check if the triangle is visible / actually needs to be replaced
 		if (CrossedEdgeFirstVertex != -1)
@@ -542,7 +542,7 @@ struct FHullConnectivity
 		return true;
 	}
 
-	void UpdateHullWithNewPoint(TArray<FIndex3i>& Triangles, TFunctionRef<FVector3<RealType>(int32)> GetPointFunc, int32 StartTriIdx, int32 PtIdx)
+	void UpdateHullWithNewPoint(TArray<FIndex3i>& Triangles, TFunctionRef<TVector<RealType>(int32)> GetPointFunc, int32 StartTriIdx, int32 PtIdx)
 	{
 		//ValidateConnectivity(Triangles);
 
@@ -603,7 +603,7 @@ struct FHullConnectivity
 
 
 template<class RealType>
-bool TConvexHull3<RealType>::Solve(int32 NumPoints, TFunctionRef<FVector3<RealType>(int32)> GetPointFunc, TFunctionRef<bool(int32)> FilterFunc)
+bool TConvexHull3<RealType>::Solve(int32 NumPoints, TFunctionRef<TVector<RealType>(int32)> GetPointFunc, TFunctionRef<bool(int32)> FilterFunc)
 {
 	Hull.Reset();
 	NumHullPoints = 0;
diff --git a/Engine/Source/Runtime/GeometryCore/Private/CompGeom/PolygonTriangulation.cpp b/Engine/Source/Runtime/GeometryCore/Private/CompGeom/PolygonTriangulation.cpp
--- a/Engine/Source/Runtime/GeometryCore/Private/CompGeom/PolygonTriangulation.cpp
+++ b/Engine/Source/Runtime/GeometryCore/Private/CompGeom/PolygonTriangulation.cpp
@@ -12,11 +12,11 @@ namespace PolygonTriangulation
 	template GEOMETRYCORE_API void TriangulateSimplePolygon<float>(const TArray<TVector2<float>>& VertexPositions, TArray<FIndex3i>& OutTriangles);
 	template GEOMETRYCORE_API void TriangulateSimplePolygon<double>(const TArray<TVector2<double>>& VertexPositions, TArray<FIndex3i>& OutTriangles);
 
-	template GEOMETRYCORE_API void ComputePolygonPlane<float>(const TArray<FVector3<float>>& VertexPositions, FVector3<float>& NormalOut, FVector3<float>& PlanePointOut);
-	template GEOMETRYCORE_API void ComputePolygonPlane<double>(const TArray<FVector3<double>>& VertexPositions, FVector3<double>& NormalOut, FVector3<double>& PlanePointOut);
+	template GEOMETRYCORE_API void ComputePolygonPlane<float>(const TArray<TVector<float>>& VertexPositions, TVector<float>& NormalOut, TVector<float>& PlanePointOut);
+	template GEOMETRYCORE_API void ComputePolygonPlane<double>(const TArray<TVector<double>>& VertexPositions, TVector<double>& NormalOut, TVector<double>& PlanePointOut);
 
-	template GEOMETRYCORE_API void TriangulateSimplePolygon<float>(const TArray<FVector3<float>>& VertexPositions, TArray<FIndex3i>& OutTriangles);
-	template GEOMETRYCORE_API void TriangulateSimplePolygon<double>(const TArray<FVector3<double>>& VertexPositions, TArray<FIndex3i>& OutTriangles);
+	template GEOMETRYCORE_API void TriangulateSimplePolygon<float>(const TArray<TVector<float>>& VertexPositions, TArray<FIndex3i>& OutTriangles);
+	template GEOMETRYCORE_API void TriangulateSimplePolygon<double>(const TArray<TVector<double>>& VertexPositions, TArray<FIndex3i>& OutTriangles);
 }
 
 
@@ -169,13 +169,13 @@ void PolygonTriangulation::TriangulateSimplePolygon(const TArray<TVector2<T>>& V
 
 
 template<typename T>
-void PolygonTriangulation::ComputePolygonPlane(const TArray<FVector3<T>>& VertexPositions, FVector3<T>& PlaneNormalOut, FVector3<T>& PlanePointOut)
+void PolygonTriangulation::ComputePolygonPlane(const TArray<TVector<T>>& VertexPositions, TVector<T>& PlaneNormalOut, TVector<T>& PlanePointOut)
 {
 	// NOTE: This polygon plane computation code is partially based on the implementation of "Newell's method" from Real-Time
 	//       Collision Detection by Christer Ericson, published by Morgan Kaufmann Publishers, (c) 2005 Elsevier Inc
 
-	PlaneNormalOut = FVector3<T>::Zero();
-	PlanePointOut = FVector3<T>::Zero();
+	PlaneNormalOut = TVector<T>::Zero();
+	PlanePointOut = TVector<T>::Zero();
 
 	int32 NumVertices = VertexPositions.Num();
 	if (NumVertices == 3)
@@ -187,8 +187,8 @@ void PolygonTriangulation::ComputePolygonPlane(const TArray<FVector3<T>>& Vertex
 	// Use 'Newell's Method' to compute a robust 'best fit' plane from the vertices of this polygon
 	for (int32 VertexNumberI = NumVertices - 1, VertexNumberJ = 0; VertexNumberJ < NumVertices; VertexNumberI = VertexNumberJ, VertexNumberJ++)
 	{
-		const FVector3<T>& PositionI = VertexPositions[VertexNumberI];
-		const FVector3<T>& PositionJ = VertexPositions[VertexNumberJ];
+		const TVector<T>& PositionI = VertexPositions[VertexNumberI];
+		const TVector<T>& PositionJ = VertexPositions[VertexNumberJ];
 		PlanePointOut += PositionJ;
 		PlaneNormalOut.X += (PositionJ.Y - PositionI.Y) * (PositionI.Z + PositionJ.Z);
 		PlaneNormalOut.Y += (PositionJ.Z - PositionI.Z) * (PositionI.X + PositionJ.X);
@@ -229,25 +229,25 @@ static bool PointInTriangle(const FVector& A, const FVector& B, const FVector& C
 // This is based on the 3D triangulation code from MeshDescription.cpp, simplified for 2D polygons
 // 
 template<typename T>
-void PolygonTriangulation::TriangulateSimplePolygon(const TArray<FVector3<T>>& VertexPositions, TArray<FIndex3i>& OutTriangles)
+void PolygonTriangulation::TriangulateSimplePolygon(const TArray<TVector<T>>& VertexPositions, TArray<FIndex3i>& OutTriangles)
 {
 	struct Local3
 	{
-		static bool IsTriangleFlipped(const FVector3<T>& ReferenceNormal, const FVector3<T>& VertexPositionA, const FVector3<T>& VertexPositionB, const FVector3<T>& VertexPositionC)
+		static bool IsTriangleFlipped(const TVector<T>& ReferenceNormal, const TVector<T>& VertexPositionA, const TVector<T>& VertexPositionB, const TVector<T>& VertexPositionC)
 		{
-			FVector3<T> TriangleNormal = VectorUtil::Normal(VertexPositionA, VertexPositionB, VertexPositionC);
+			TVector<T> TriangleNormal = VectorUtil::Normal(VertexPositionA, VertexPositionB, VertexPositionC);
 			return TriangleNormal.Dot(ReferenceNormal) <= (T)0;
 		}
 
-		static bool VectorsOnSameSide(const FVector3<T>& Vec, const FVector3<T>& A, const FVector3<T>& B, T SameSideDotProductEpsilon)
+		static bool VectorsOnSameSide(const TVector<T>& Vec, const TVector<T>& A, const TVector<T>& B, T SameSideDotProductEpsilon)
 		{
-			FVector3<T> CrossA = Vec.Cross(A);
-			FVector3<T> CrossB = Vec.Cross(B);
+			TVector<T> CrossA = Vec.Cross(A);
+			TVector<T> CrossB = Vec.Cross(B);
 			T DotWithEpsilon = SameSideDotProductEpsilon + CrossA.Dot(CrossB);
 			return DotWithEpsilon >= 0; // !FMath::IsNegativeFloat(DotWithEpsilon);
 		}
 
-		static bool PointInTriangle(const FVector3<T>& A, const FVector3<T>& B, const FVector3<T>& C, const FVector3<T>& P, const T InsideTriangleDotProductEpsilon)
+		static bool PointInTriangle(const TVector<T>& A, const TVector<T>& B, const TVector<T>& C, const TVector<T>& P, const T InsideTriangleDotProductEpsilon)
 		{
 			return (VectorsOnSameSide(B - A, P - A, C - A, InsideTriangleDotProductEpsilon) &&
 				VectorsOnSameSide(C - B, P - B, A - B, InsideTriangleDotProductEpsilon) &&
@@ -270,7 +270,7 @@ void PolygonTriangulation::TriangulateSimplePolygon(const TArray<FVector3<T>>& V
 
 	// First figure out the polygon normal.  We need this to determine which triangles are convex, so that
 	// we can figure out which ears to clip
-	FVector3<T> PolygonNormal, PolygonCentroid;
+	TVector<T> PolygonNormal, PolygonCentroid;
 	ComputePolygonPlane(VertexPositions, PolygonNormal, PolygonCentroid);
 
 	// Make a simple linked list array of the previous and next vertex numbers, for each vertex number
@@ -300,9 +300,9 @@ void PolygonTriangulation::TriangulateSimplePolygon(const TArray<FVector3<T>>& V
 		// vertices are collinear or other degenerate cases.
 		if (RemainingVertexCount > 3 && EarTestCount < RemainingVertexCount)
 		{
-			const FVector3<T>& PrevVertexPosition = VertexPositions[PrevVertexNumbers[EarVertexNumber]];
-			const FVector3<T>& EarVertexPosition = VertexPositions[EarVertexNumber];
-			const FVector3<T>& NextVertexPosition = VertexPositions[NextVertexNumbers[EarVertexNumber]];
+			const TVector<T>& PrevVertexPosition = VertexPositions[PrevVertexNumbers[EarVertexNumber]];
+			const TVector<T>& EarVertexPosition = VertexPositions[EarVertexNumber];
+			const TVector<T>& NextVertexPosition = VertexPositions[NextVertexNumbers[EarVertexNumber]];
 
 			// Figure out whether the potential ear triangle is facing the same direction as the polygon
 			// itself.  If it's facing the opposite direction, then we're dealing with a concave triangle
@@ -316,7 +316,7 @@ void PolygonTriangulation::TriangulateSimplePolygon(const TArray<FVector3<T>>& V
 				{
 					// Test every other remaining vertex to make sure that it doesn't lie inside our potential ear
 					// triangle.  If we find a vertex that's inside the triangle, then it cannot actually be an ear.
-					const FVector3<T>& TestVertexPosition = VertexPositions[TestVertexNumber];
+					const TVector<T>& TestVertexPosition = VertexPositions[TestVertexNumber];
 					if (Local3::PointInTriangle(PrevVertexPosition, EarVertexPosition, NextVertexPosition, TestVertexPosition, SMALL_NUMBER))
 					{
 						bIsEar = false;
diff --git a/Engine/Source/Runtime/GeometryCore/Private/DynamicMesh/DynamicMeshOverlay.cpp b/Engine/Source/Runtime/GeometryCore/Private/DynamicMesh/DynamicMeshOverlay.cpp
--- a/Engine/Source/Runtime/GeometryCore/Private/DynamicMesh/DynamicMeshOverlay.cpp
+++ b/Engine/Source/Runtime/GeometryCore/Private/DynamicMesh/DynamicMeshOverlay.cpp
@@ -1279,7 +1279,7 @@ void TDynamicMeshOverlay<RealType, ElementSize>::OnPokeTriangle(const FDynamicMe
 
 	// create new element at barycentric position
 	int CenterElemID = AppendElement((RealType)0);
-	FVector3<double> BaryCoords((double)PokeInfo.BaryCoords.X, (double)PokeInfo.BaryCoords.Y, (double)PokeInfo.BaryCoords.Z);
+	FVector3d BaryCoords((double)PokeInfo.BaryCoords.X, (double)PokeInfo.BaryCoords.Y, (double)PokeInfo.BaryCoords.Z);
 	SetElementFromBary(CenterElemID, Triangle[0], Triangle[1], Triangle[2], BaryCoords);
 
 	// update orig triangle and two new ones. Winding orders here mirror FDynamicMesh3::PokeTriangle
@@ -1379,7 +1379,7 @@ void TDynamicMeshOverlay<RealType, ElementSize>::SetElementFromLerp(int SetEleme
 }
 
 template<typename RealType, int ElementSize>
-void TDynamicMeshOverlay<RealType, ElementSize>::SetElementFromBary(int SetElement, int ElementA, int ElementB, int ElementC, const FVector3<double>& BaryCoords)
+void TDynamicMeshOverlay<RealType, ElementSize>::SetElementFromBary(int SetElement, int ElementA, int ElementB, int ElementC, const FVector3d& BaryCoords)
 {
 	int IndexSet = ElementSize * SetElement;
 	int IndexA = ElementSize * ElementA;
diff --git a/Engine/Source/Runtime/GeometryCore/Private/DynamicMesh/MeshTangents.cpp b/Engine/Source/Runtime/GeometryCore/Private/DynamicMesh/MeshTangents.cpp
--- a/Engine/Source/Runtime/GeometryCore/Private/DynamicMesh/MeshTangents.cpp
+++ b/Engine/Source/Runtime/GeometryCore/Private/DynamicMesh/MeshTangents.cpp
@@ -95,8 +95,8 @@ void TMeshTangents<RealType>::SetTangentCount(int Count, bool bClearToZero)
 	{
 		for (int i = 0; i < Count; ++i)
 		{
-			Tangents[i] = FVector3<RealType>::Zero();
-			Bitangents[i] = FVector3<RealType>::Zero();
+			Tangents[i] = TVector<RealType>::Zero();
+			Bitangents[i] = TVector<RealType>::Zero();
 		}
 	}
 }
@@ -127,8 +127,8 @@ bool TMeshTangents<RealType>::CopyTriVertexTangents(const FDynamicMesh3& SourceM
 			int32 k = 3 * tid;
 			for (int32 j = 0; j < 3; ++j )
 			{
-				Tangents[k+j] = FVector3<RealType>(TriTangents[j]);
-				Bitangents[k+j] = FVector3<RealType>(TriBiTangents[j]);
+				Tangents[k+j] = TVector<RealType>(TriTangents[j]);
+				Bitangents[k+j] = TVector<RealType>(TriBiTangents[j]);
 			}
 		}
 	}
@@ -241,8 +241,8 @@ void TMeshTangents<RealType>::ComputeTriangleTangents(const FDynamicMeshUVOverla
 			DegenerateTriLock.Lock();
 			AllDegenerateTris.Add(TriangleID);
 			DegenerateTriLock.Unlock();
-			Tangents[TriangleID] = (FVector3<RealType>)Tangent;
-			Bitangents[TriangleID] = (FVector3<RealType>)Bitangent;
+			Tangents[TriangleID] = (TVector<RealType>)Tangent;
+			Bitangents[TriangleID] = (TVector<RealType>)Bitangent;
 			return;
 		}
 
@@ -253,8 +253,8 @@ void TMeshTangents<RealType>::ComputeTriangleTangents(const FDynamicMeshUVOverla
 			Normalize(Bitangent);
 		}
 
-		Tangents[TriangleID] = (FVector3<RealType>)Tangent;
-		Bitangents[TriangleID] = (FVector3<RealType>)Bitangent;
+		Tangents[TriangleID] = (TVector<RealType>)Tangent;
+		Bitangents[TriangleID] = (TVector<RealType>)Bitangent;
 	});
 }
 
@@ -270,12 +270,12 @@ bool TMeshTangents<RealType>::CopyToOverlays(FDynamicMesh3& MeshToSet)
 
 	// Set aliases to make iterating over tangents and bitangents easier
 	FDynamicMeshNormalOverlay* TangentOverlays[2] = { MeshToSet.Attributes()->PrimaryTangents(), MeshToSet.Attributes()->PrimaryBiTangents() };
-	TArray<FVector3<RealType>>* TangentValues[2] = { &Tangents, &Bitangents };
+	TArray<TVector<RealType>>* TangentValues[2] = { &Tangents, &Bitangents };
 	
 	for (int Idx = 0; Idx < 2; Idx++)
 	{
 		// Create overlay topology
-		TArray<FVector3<RealType>>& TV = *TangentValues[Idx];
+		TArray<TVector<RealType>>& TV = *TangentValues[Idx];
 		TangentOverlays[Idx]->CreateFromPredicate([&MeshToSet, &TV](int ParentVertexIdx, int TriIDA, int TriIDB) -> bool
 		{
 			FIndex3i TriA = MeshToSet.GetTriangle(TriIDA);
@@ -283,7 +283,9 @@ bool TMeshTangents<RealType>::CopyToOverlays(FDynamicMesh3& MeshToSet)
 			int SubA = TriA.IndexOf(ParentVertexIdx);
 			int SubB = TriB.IndexOf(ParentVertexIdx);
 			checkSlow(SubA > -1 && SubB > -1);
-			return TV[TriIDA * 3 + SubA].DistanceSquared(TV[TriIDB * 3 + SubB]) < TMathUtil<RealType>::ZeroTolerance;
+			const TVector<RealType>& A = TV[TriIDA * 3 + SubA];
+			const TVector<RealType>& B = TV[TriIDB * 3 + SubB];
+			return DistanceSquared(A,B) < TMathUtil<RealType>::ZeroTolerance;
 		}, 0);
 
 		// Write tangent values out for each wedge value
@@ -345,8 +347,8 @@ void TMeshTangents<RealType>::ComputeSeparatePerTriangleTangents(const FDynamicM
 			FVector3d ReconsBitangent = VectorUtil::Bitangent(VtxNormal, ProjectedTangent, BitangentSign);
 
 			SetPerTriangleTangent(TriangleID, j, 
-				Normalized((FVector3<RealType>)ProjectedTangent),
-				Normalized((FVector3<RealType>)ReconsBitangent) );
+				Normalized((TVector<RealType>)ProjectedTangent),
+				Normalized((TVector<RealType>)ReconsBitangent) );
 		}
 
 	} /*, EParallelForFlags::ForceSingleThread */);
@@ -532,7 +534,7 @@ void TMeshTangents<RealType>::ComputeMikkTStyleTangents(const FDynamicMeshNormal
 
 			// todo: catch failures here?
 
-			Tangents[TriVertIndex] = Normalized((FVector3<RealType>)Tangent);
+			Tangents[TriVertIndex] = Normalized((TVector<RealType>)Tangent);
 
 			FIndex3i NormTriangle = NormalOverlay->GetTriangle(TriangleID);
 			FVector3d TriVertNormal = NormalOverlay->IsSetTriangle(TriangleID) ?
@@ -541,10 +543,10 @@ void TMeshTangents<RealType>::ComputeMikkTStyleTangents(const FDynamicMeshNormal
 			// convert bitangent
 			double BinormalSign = VectorUtil::BitangentSign(TriVertNormal, Tangent, Bitangent);
 			FVector3d ReconsBitangent = VectorUtil::Bitangent(TriVertNormal, Tangent, BinormalSign);
-			Bitangents[TriVertIndex] = Normalized((FVector3<RealType>)ReconsBitangent);
+			Bitangents[TriVertIndex] = Normalized((TVector<RealType>)ReconsBitangent);
 
 			//FVector3d OrigBitangent = AccumBitangent[UniqueIndex];
-			//FVector3<RealType> BTResult = Bitangents[3 * TriangleID + TriIndex];
+			//TVector<RealType> BTResult = Bitangents[3 * TriangleID + TriIndex];
 			//if (BTResult.SquaredLength() < 0.99)
 			//{
 			//	UE_LOG(LogTemp, Warning, TEXT("TriVertNormal is %f %f %f"), TriVertNormal.X, TriVertNormal.Y, TriVertNormal.Z);
diff --git a/Engine/Source/Runtime/GeometryCore/Private/Generators/RevolveGenerator.cpp b/Engine/Source/Runtime/GeometryCore/Private/Generators/RevolveGenerator.cpp
--- a/Engine/Source/Runtime/GeometryCore/Private/Generators/RevolveGenerator.cpp
+++ b/Engine/Source/Runtime/GeometryCore/Private/Generators/RevolveGenerator.cpp
@@ -261,14 +261,14 @@ FDynamicMesh3 FRevolvePlanarPathGenerator::GenerateMesh()
 			}
 
 			// Compute a best-fit plane of the vertices
-			TArray<FVector3<double>> VertexPositions;
+			TArray<FVector3d> VertexPositions;
 			for (int32 vid : VertexLoop)
 			{
 				VertexPositions.Add( ResultMesh.GetVertex(vid) );
 			}
 
-			FVector3<double> PlaneOrigin;
-			FVector3<double> PlaneNormal;
+			FVector3d PlaneOrigin;
+			FVector3d PlaneNormal;
 			PolygonTriangulation::ComputePolygonPlane<double>(VertexPositions, PlaneNormal, PlaneOrigin);
 			PlaneNormal *= -1.0;	// Previous function seems to orient the normal opposite to what's expected elsewhere
 
@@ -451,14 +451,14 @@ FDynamicMesh3 FRevolvePlanarPolygonGenerator::GenerateMesh()
 			}
 
 			// Compute a best-fit plane of the vertices
-			TArray<FVector3<double>> VertexPositions;
+			TArray<FVector3d> VertexPositions;
 			for (int32 vid : VertexLoop)
 			{
 				VertexPositions.Add( ResultMesh.GetVertex(vid) );
 			}
 
-			FVector3<double> PlaneOrigin;
-			FVector3<double> PlaneNormal;
+			FVector3d PlaneOrigin;
+			FVector3d PlaneNormal;
 			PolygonTriangulation::ComputePolygonPlane<double>(VertexPositions, PlaneNormal, PlaneOrigin);
 			PlaneNormal *= -1.0;	// Previous function seems to orient the normal opposite to what's expected elsewhere
 
@@ -639,14 +639,14 @@ FDynamicMesh3 FSpiralRevolvePlanarPolygonGenerator::GenerateMesh()
 			}
 
 			// Compute a best-fit plane of the vertices
-			TArray<FVector3<double>> VertexPositions;
+			TArray<FVector3d> VertexPositions;
 			for (int32 vid : VertexLoop)
 			{
 				VertexPositions.Add( ResultMesh.GetVertex(vid) );
 			}
 
-			FVector3<double> PlaneOrigin;
-			FVector3<double> PlaneNormal;
+			FVector3d PlaneOrigin;
+			FVector3d PlaneNormal;
 			PolygonTriangulation::ComputePolygonPlane<double>(VertexPositions, PlaneNormal, PlaneOrigin);
 			PlaneNormal *= -1.0;	// Previous function seems to orient the normal opposite to what's expected elsewhere
 
diff --git a/Engine/Source/Runtime/GeometryCore/Private/Intersection/ContainmentQueries3.cpp b/Engine/Source/Runtime/GeometryCore/Private/Intersection/ContainmentQueries3.cpp
--- a/Engine/Source/Runtime/GeometryCore/Private/Intersection/ContainmentQueries3.cpp
+++ b/Engine/Source/Runtime/GeometryCore/Private/Intersection/ContainmentQueries3.cpp
@@ -3,6 +3,7 @@
 #include "Intersection/ContainmentQueries3.h"
 
 using namespace UE::Geometry;
+using namespace UE::Math;
 
 template<typename RealType>
 bool UE::Geometry::IsInside(const TSphere3<RealType>& OuterSphere, const TSphere3<RealType>& InnerSphere)
@@ -20,7 +21,7 @@ bool UE::Geometry::IsInside(const TSphere3<RealType>& OuterSphere, const TCapsul
 template<typename RealType>
 bool UE::Geometry::IsInside(const TSphere3<RealType>& OuterSphere, const TOrientedBox3<RealType>& InnerBox)
 {
-	return InnerBox.TestCorners([&](const FVector3<RealType>& Point) {
+	return InnerBox.TestCorners([&](const TVector<RealType>& Point) {
 		return OuterSphere.Contains(Point);
 	});
 }
@@ -56,7 +57,7 @@ bool UE::Geometry::IsInside(const TCapsule3<RealType>& OuterCapsule, const TOrie
 	// todo: possibly more efficient to do this by calculating distance to box center and then adding on box radius?
 	// not sure if this would an exact test though, or just a bound
 
-	return InnerBox.TestCorners([&](const FVector3<RealType>& Point) {
+	return InnerBox.TestCorners([&](const TVector<RealType>& Point) {
 		return OuterCapsule.Contains(Point);
 	});
 }
@@ -69,8 +70,8 @@ template<typename RealType>
 bool UE::Geometry::IsInside(const TOrientedBox3<RealType>& OuterBox, const TOrientedBox3<RealType>& InnerBox)
 {
 	const TFrame3<RealType>& Frame = OuterBox.Frame;
-	const FVector3<RealType>& Extents = OuterBox.Extents;
-	FVector3<RealType> Axes[3];
+	const TVector<RealType>& Extents = OuterBox.Extents;
+	TVector<RealType> Axes[3];
 	OuterBox.Frame.GetAxes(Axes[0], Axes[1], Axes[2]);
 	for (int32 k = 0; k < 3; ++k)
 	{
@@ -91,8 +92,8 @@ template<typename RealType>
 bool UE::Geometry::IsInside(const TOrientedBox3<RealType>& OuterBox, const TSphere3<RealType>& InnerSphere)
 {
 	const TFrame3<RealType>& Frame = OuterBox.Frame;
-	const FVector3<RealType>& Extents = OuterBox.Extents;
-	FVector3<RealType> Axes[3];
+	const TVector<RealType>& Extents = OuterBox.Extents;
+	TVector<RealType> Axes[3];
 	OuterBox.Frame.GetAxes(Axes[0], Axes[1], Axes[2]);
 	for (int32 k = 0; k < 3; ++k)
 	{
@@ -114,8 +115,8 @@ template<typename RealType>
 bool UE::Geometry::IsInside(const TOrientedBox3<RealType>& OuterBox, const TCapsule3<RealType>& InnerCapsule)
 {
 	const TFrame3<RealType>& Frame = OuterBox.Frame;
-	const FVector3<RealType>& Extents = OuterBox.Extents;
-	FVector3<RealType> Axes[3];
+	const TVector<RealType>& Extents = OuterBox.Extents;
+	TVector<RealType> Axes[3];
 	OuterBox.Frame.GetAxes(Axes[0], Axes[1], Axes[2]);
 	for (int32 k = 0; k < 3; ++k)
 	{
diff --git a/Engine/Source/Runtime/GeometryCore/Private/Intersection/IntersectionQueries3.cpp b/Engine/Source/Runtime/GeometryCore/Private/Intersection/IntersectionQueries3.cpp
--- a/Engine/Source/Runtime/GeometryCore/Private/Intersection/IntersectionQueries3.cpp
+++ b/Engine/Source/Runtime/GeometryCore/Private/Intersection/IntersectionQueries3.cpp
@@ -28,7 +28,7 @@ template<typename RealType>
 bool UE::Geometry::TestIntersection(const THalfspace3<RealType>& Halfspace, const TOrientedBox3<RealType>& Box)
 {
 	RealType Center = Halfspace.Normal.Dot(Box.Frame.Origin) - Halfspace.Constant;
-	FVector3<RealType> X, Y, Z;
+	TVector<RealType> X, Y, Z;
 	Box.Frame.GetAxes(X, Y, Z);
 	RealType Radius =
 		TMathUtil<RealType>::Abs(Box.Extents.X * Halfspace.Normal.Dot(X)) +
diff --git a/Engine/Source/Runtime/GeometryCore/Private/Spatial/GeometrySet3.cpp b/Engine/Source/Runtime/GeometryCore/Private/Spatial/GeometrySet3.cpp
--- a/Engine/Source/Runtime/GeometryCore/Private/Spatial/GeometrySet3.cpp
+++ b/Engine/Source/Runtime/GeometryCore/Private/Spatial/GeometrySet3.cpp
@@ -237,7 +237,7 @@ bool FGeometrySet3::FindNearestCurveToRay(const FRay3d& Ray, FNearest& ResultOut
 			{
 				bTakePoint = true;
 			}
-			else if (Ray.PointAt(MinRayParamT).Distance(CurveNearestPosition) < NearestWithinTolDist && CurveNearestDist < NearestWithinTolDist)
+			else if ( Distance(Ray.PointAt(MinRayParamT), CurveNearestPosition) < NearestWithinTolDist && CurveNearestDist < NearestWithinTolDist)
 			{
 				bTakePoint = true;
 			}
diff --git a/Engine/Source/Runtime/GeometryCore/Public/BoxTypes.h b/Engine/Source/Runtime/GeometryCore/Public/BoxTypes.h
--- a/Engine/Source/Runtime/GeometryCore/Public/BoxTypes.h
+++ b/Engine/Source/Runtime/GeometryCore/Public/BoxTypes.h
@@ -212,29 +212,29 @@ typedef TInterval1<double> FInterval1d;
 template <typename RealType>
 struct TAxisAlignedBox3
 {
-	FVector3<RealType> Min;
-	FVector3<RealType> Max;
+	TVector<RealType> Min;
+	TVector<RealType> Max;
 
 	TAxisAlignedBox3() : 
 		TAxisAlignedBox3(TAxisAlignedBox3<RealType>::Empty())
 	{
 	}
 
-	TAxisAlignedBox3(const FVector3<RealType>& Min, const FVector3<RealType>& Max)
+	TAxisAlignedBox3(const TVector<RealType>& Min, const TVector<RealType>& Max)
 	{
 		this->Min = Min;
 		this->Max = Max;
 	}
 
-	TAxisAlignedBox3(const FVector3<RealType>& A, const FVector3<RealType>& B, const FVector3<RealType>& C)
+	TAxisAlignedBox3(const TVector<RealType>& A, const TVector<RealType>& B, const TVector<RealType>& C)
 	{
 		// TMathUtil::MinMax could be used here, but it generates worse code because the Min3's below will be
 		// turned into SSE instructions by the optimizer, while MinMax will not
-		Min = FVector3<RealType>(
+		Min = TVector<RealType>(
 			TMathUtil<RealType>::Min3(A.X, B.X, C.X),
 			TMathUtil<RealType>::Min3(A.Y, B.Y, C.Y),
 			TMathUtil<RealType>::Min3(A.Z, B.Z, C.Z));
-		Max = FVector3<RealType>(
+		Max = TVector<RealType>(
 			TMathUtil<RealType>::Max3(A.X, B.X, C.X),
 			TMathUtil<RealType>::Max3(A.Y, B.Y, C.Y),
 			TMathUtil<RealType>::Max3(A.Z, B.Z, C.Z));
@@ -245,18 +245,18 @@ struct TAxisAlignedBox3
 	template<typename OtherRealType>
 	explicit TAxisAlignedBox3(const TAxisAlignedBox3<OtherRealType>& OtherBox)
 	{
-		this->Min = FVector3<RealType>(OtherBox.Min);
-		this->Max = FVector3<RealType>(OtherBox.Max);
+		this->Min = TVector<RealType>(OtherBox.Min);
+		this->Max = TVector<RealType>(OtherBox.Max);
 	}
 
-	TAxisAlignedBox3(const FVector3<RealType>& Center, RealType HalfWidth)
+	TAxisAlignedBox3(const TVector<RealType>& Center, RealType HalfWidth)
 	{
-		this->Min = FVector3<RealType>(Center.X-HalfWidth, Center.Y-HalfWidth, Center.Z-HalfWidth);
-		this->Max = FVector3<RealType>(Center.X+HalfWidth, Center.Y+HalfWidth, Center.Z+HalfWidth);
+		this->Min = TVector<RealType>(Center.X-HalfWidth, Center.Y-HalfWidth, Center.Z-HalfWidth);
+		this->Max = TVector<RealType>(Center.X+HalfWidth, Center.Y+HalfWidth, Center.Z+HalfWidth);
 	}
 
 
-	TAxisAlignedBox3(const TAxisAlignedBox3& Box, const TFunction<FVector3<RealType>(const FVector3<RealType>&)> TransformF)
+	TAxisAlignedBox3(const TAxisAlignedBox3& Box, const TFunction<TVector<RealType>(const TVector<RealType>&)> TransformF)
 	{
 		if (TransformF == nullptr)
 		{
@@ -265,7 +265,7 @@ struct TAxisAlignedBox3
 			return;
 		}
 
-		FVector3<RealType> C0 = TransformF(Box.GetCorner(0));
+		TVector<RealType> C0 = TransformF(Box.GetCorner(0));
 		Min = C0;
 		Max = C0;
 		for (int i = 1; i < 8; ++i)
@@ -276,7 +276,7 @@ struct TAxisAlignedBox3
 
 	TAxisAlignedBox3(const TAxisAlignedBox3& Box, const FTransform3d& Transform)
 	{
-		FVector3<RealType> C0 = Transform.TransformPosition(Box.GetCorner(0));
+		TVector<RealType> C0 = Transform.TransformPosition(Box.GetCorner(0));
 		Min = C0;
 		Max = C0;
 		for (int i = 1; i < 8; ++i)
@@ -303,35 +303,35 @@ struct TAxisAlignedBox3
 	}
 	TAxisAlignedBox3(const FBox& Box)
 	{
-		Min = FVector3<RealType>(Box.Min);
-		Max = FVector3<RealType>(Box.Max);
+		Min = TVector<RealType>(Box.Min);
+		Max = TVector<RealType>(Box.Max);
 	}
 
 	/**
 	* @param Index corner index in range 0-7
 	* @return Corner point on the box identified by the given index. See diagram in OrientedBoxTypes.h for index/corner mapping.
 	*/
-	FVector3<RealType> GetCorner(int Index) const
+	TVector<RealType> GetCorner(int Index) const
 	{
 		check(Index >= 0 && Index <= 7);
 		RealType X = (((Index & 1) != 0) ^ ((Index & 2) != 0)) ? (Max.X) : (Min.X);
 		RealType Y = ((Index / 2) % 2 == 0) ? (Min.Y) : (Max.Y);
 		RealType Z = (Index < 4) ? (Min.Z) : (Max.Z);
-		return FVector3<RealType>(X, Y, Z);
+		return TVector<RealType>(X, Y, Z);
 	}
 
 	static TAxisAlignedBox3<RealType> Empty()
 	{
 		return TAxisAlignedBox3(
-			FVector3<RealType>(TNumericLimits<RealType>::Max(), TNumericLimits<RealType>::Max(), TNumericLimits<RealType>::Max()),
-			FVector3<RealType>(-TNumericLimits<RealType>::Max(), -TNumericLimits<RealType>::Max(), -TNumericLimits<RealType>::Max()));
+			TVector<RealType>(TNumericLimits<RealType>::Max(), TNumericLimits<RealType>::Max(), TNumericLimits<RealType>::Max()),
+			TVector<RealType>(-TNumericLimits<RealType>::Max(), -TNumericLimits<RealType>::Max(), -TNumericLimits<RealType>::Max()));
 	}
 
 	static TAxisAlignedBox3<RealType> Infinite()
 	{
 		return TAxisAlignedBox3(
-			FVector3<RealType>(-TNumericLimits<RealType>::Max(), -TNumericLimits<RealType>::Max(), -TNumericLimits<RealType>::Max()),
-			FVector3<RealType>(TNumericLimits<RealType>::Max(), TNumericLimits<RealType>::Max(), TNumericLimits<RealType>::Max()) );
+			TVector<RealType>(-TNumericLimits<RealType>::Max(), -TNumericLimits<RealType>::Max(), -TNumericLimits<RealType>::Max()),
+			TVector<RealType>(TNumericLimits<RealType>::Max(), TNumericLimits<RealType>::Max(), TNumericLimits<RealType>::Max()) );
 	}
 
 	/**
@@ -363,20 +363,20 @@ struct TAxisAlignedBox3
 	}
 
 
-	FVector3<RealType> Center() const
+	TVector<RealType> Center() const
 	{
-		return FVector3<RealType>(
+		return TVector<RealType>(
 			(Min.X + Max.X) * (RealType)0.5,
 			(Min.Y + Max.Y) * (RealType)0.5,
 			(Min.Z + Max.Z) * (RealType)0.5);
 	}
 
-	FVector3<RealType> Extents() const
+	TVector<RealType> Extents() const
 	{
 		return (Max - Min) * (RealType).5;
 	}
 
-	void Contain(const FVector3<RealType>& V)
+	void Contain(const TVector<RealType>& V)
 	{
 		if (V.X < Min.X)
 		{
@@ -414,7 +414,7 @@ struct TAxisAlignedBox3
 		Max.Z = Max.Z > Other.Max.Z ? Max.Z : Other.Max.Z;
 	}
 
-	bool Contains(const FVector3<RealType>& V) const
+	bool Contains(const TVector<RealType>& V) const
 	{
 		return (Min.X <= V.X) && (Min.Y <= V.Y) && (Min.Z <= V.Z) && (Max.X >= V.X) && (Max.Y >= V.Y) && (Max.Z >= V.Z);
 	}
@@ -427,8 +427,8 @@ struct TAxisAlignedBox3
 	TAxisAlignedBox3<RealType> Intersect(const TAxisAlignedBox3<RealType>& Box) const
 	{
 		TAxisAlignedBox3<RealType> Intersection(
-			FVector3<RealType>(TMathUtil<RealType>::Max(Min.X, Box.Min.X), TMathUtil<RealType>::Max(Min.Y, Box.Min.Y), TMathUtil<RealType>::Max(Min.Z, Box.Min.Z)),
-			FVector3<RealType>(TMathUtil<RealType>::Min(Max.X, Box.Max.X), TMathUtil<RealType>::Min(Max.Y, Box.Max.Y), TMathUtil<RealType>::Min(Max.Z, Box.Max.Z)));
+			TVector<RealType>(TMathUtil<RealType>::Max(Min.X, Box.Min.X), TMathUtil<RealType>::Max(Min.Y, Box.Min.Y), TMathUtil<RealType>::Max(Min.Z, Box.Min.Z)),
+			TVector<RealType>(TMathUtil<RealType>::Min(Max.X, Box.Max.X), TMathUtil<RealType>::Min(Max.Y, Box.Max.Y), TMathUtil<RealType>::Min(Max.Z, Box.Max.Z)));
 		if (Intersection.Height() <= 0 || Intersection.Width() <= 0 || Intersection.Depth() <= 0)
 		{
 			return TAxisAlignedBox3<RealType>::Empty();
@@ -444,7 +444,7 @@ struct TAxisAlignedBox3
 		return !((Box.Max.X <= Min.X) || (Box.Min.X >= Max.X) || (Box.Max.Y <= Min.Y) || (Box.Min.Y >= Max.Y) || (Box.Max.Z <= Min.Z) || (Box.Min.Z >= Max.Z));
 	}
 
-	RealType DistanceSquared(const FVector3<RealType>& V) const
+	RealType DistanceSquared(const TVector<RealType>& V) const
 	{
 		RealType dx = (V.X < Min.X) ? Min.X - V.X : (V.X > Max.X ? V.X - Max.X : 0);
 		RealType dy = (V.Y < Min.Y) ? Min.Y - V.Y : (V.Y > Max.Y ? V.Y - Max.Y : 0);
@@ -516,9 +516,9 @@ struct TAxisAlignedBox3
 		return TMathUtil<RealType>::Min(Width(), TMathUtil<RealType>::Min(Height(), Depth()));
 	}
 
-	FVector3<RealType> Diagonal() const
+	TVector<RealType> Diagonal() const
 	{
-		return FVector3<RealType>(Max.X - Min.X, Max.Y - Min.Y, Max.Z - Min.Z);
+		return TVector<RealType>(Max.X - Min.X, Max.Y - Min.Y, Max.Z - Min.Z);
 	}
 
 	inline bool IsEmpty() const
diff --git a/Engine/Source/Runtime/GeometryCore/Public/CapsuleTypes.h b/Engine/Source/Runtime/GeometryCore/Public/CapsuleTypes.h
--- a/Engine/Source/Runtime/GeometryCore/Public/CapsuleTypes.h
+++ b/Engine/Source/Runtime/GeometryCore/Public/CapsuleTypes.h
@@ -14,6 +14,8 @@ namespace UE
 namespace Geometry 
 {
 
+using namespace UE::Math;
+
 /*
  * 3D Capsule stored as Line Segment and Radius
  */
@@ -31,18 +33,18 @@ public:
 	TCapsule3(const TSegment3<T>& SegmentIn, T RadiusIn)
 		: Segment(SegmentIn), Radius(RadiusIn) {}
 
-	TCapsule3(const FVector3<T>& StartPoint, const FVector3<T>& EndPoint, T RadiusIn)
+	TCapsule3(const TVector<T>& StartPoint, const TVector<T>& EndPoint, T RadiusIn)
 		: Segment(StartPoint, EndPoint), Radius(RadiusIn) {}
 
 
 	/** @return Center of capsule line segment */
-	const FVector3<T>& Center() const
+	const TVector<T>& Center() const
 	{
 		return Segment.Center;
 	}
 
 	/** @return Direction of capsule line segment */
-	const FVector3<T>& Direction() const
+	const TVector<T>& Direction() const
 	{
 		return Segment.Direction;
 	}
@@ -76,11 +78,11 @@ public:
 	{
 		return TOrientedBox3<T>(
 			TFrame3<T>(Segment.Center, Segment.Direction),
-			FVector3<T>(Radius, Radius, Segment.Extent));
+			TVector<T>(Radius, Radius, Segment.Extent));
 	}
 
 	/** @return true if Capsule contains the given Point */
-	bool Contains(const FVector3<T>& Point) const
+	bool Contains(const TVector<T>& Point) const
 	{
 		T DistSqr = Segment.DistanceSquared(Point);
 		return DistSqr <= Radius * Radius;
@@ -90,7 +92,7 @@ public:
 	/**
 	 * @return minimum squared distance from Point to Capsule surface for points outside capsule, 0 for points inside
 	 */
-	inline T DistanceSquared(const FVector3<T>& Point) const
+	inline T DistanceSquared(const TVector<T>& Point) const
 	{
 		const T PosDistance = TMathUtil<T>::Max(SignedDistance(Point), (T)0);
 		return PosDistance * PosDistance;
@@ -99,7 +101,7 @@ public:
 	/**
 	 * @return signed distance from Point to Capsule surface. Points inside capsule return negative distance.
 	 */
-	inline T SignedDistance(const FVector3<T>& Point) const
+	inline T SignedDistance(const TVector<T>& Point) const
 	{
 		T SqrDist = Segment.DistanceSquared(Point);
 		return TMathUtil<T>::Sqrt(SqrDist) - Radius;
diff --git a/Engine/Source/Runtime/GeometryCore/Public/CircleTypes.h b/Engine/Source/Runtime/GeometryCore/Public/CircleTypes.h
--- a/Engine/Source/Runtime/GeometryCore/Public/CircleTypes.h
+++ b/Engine/Source/Runtime/GeometryCore/Public/CircleTypes.h
@@ -153,7 +153,7 @@ struct TCircle3
 		Radius = RadiusIn;
 	}
 
-	TCircle3(const FVector3<RealType>& CenterIn, const RealType& RadiusIn)
+	TCircle3(const TVector<RealType>& CenterIn, const RealType& RadiusIn)
 	{
 		Frame.Origin = CenterIn;
 		Radius = RadiusIn;
@@ -166,12 +166,12 @@ struct TCircle3
 	}
 
 
-	const FVector3<RealType>& GetCenter() const
+	const TVector<RealType>& GetCenter() const
 	{
 		return Frame.Origin;
 	}
 
-	FVector3<RealType> GetNormal() const
+	TVector<RealType> GetNormal() const
 	{
 		return Frame.Z();
 	}
@@ -217,18 +217,18 @@ struct TCircle3
 	}
 
 
-	FVector3<RealType> GetPointFromAngleD(RealType AngleDeg) const
+	TVector<RealType> GetPointFromAngleD(RealType AngleDeg) const
 	{
 		return GetPointFromAngleR(AngleDeg * TMathUtil<RealType>::DegToRad);
 	}
 
-	FVector3<RealType> GetPointFromAngleR(RealType AngleRad) const
+	TVector<RealType> GetPointFromAngleR(RealType AngleRad) const
 	{
 		RealType c = TMathUtil<RealType>::Cos(AngleRad), s = TMathUtil<RealType>::Sin(AngleRad);
 		return Frame.FromPlaneUV(TVector2<RealType>(Radius*c, Radius*s), 2);
 	}
 
-	FVector3<RealType> GetPointFromUnitParameter(RealType UnitParam) const
+	TVector<RealType> GetPointFromUnitParameter(RealType UnitParam) const
 	{
 		RealType AngleRad = ((bIsReversed) ? (-UnitParam) : (UnitParam)) * TMathUtil<RealType>::TwoPi;
 		return GetPointFromAngleR(AngleRad);
@@ -237,13 +237,13 @@ struct TCircle3
 
 
 
-	FVector3<RealType> ClosestPoint(const FVector3<RealType>& QueryPoint) const
+	TVector<RealType> ClosestPoint(const TVector<RealType>& QueryPoint) const
 	{
-		const FVector3<RealType>& Center = Frame.Origin;
-		FVector3<RealType> Normal = Frame.GetAxis(2);
+		const TVector<RealType>& Center = Frame.Origin;
+		TVector<RealType> Normal = Frame.GetAxis(2);
 
-		FVector3<RealType> PointDelta = QueryPoint - Center;
-		FVector3<RealType> DeltaInPlane = PointDelta - Normal.Dot(PointDelta)*Normal;
+		TVector<RealType> PointDelta = QueryPoint - Center;
+		TVector<RealType> DeltaInPlane = PointDelta - Normal.Dot(PointDelta)*Normal;
 		RealType OriginDist = DeltaInPlane.Length();
 		if (OriginDist > (RealType)0)
 		{
@@ -256,19 +256,19 @@ struct TCircle3
 	}
 
 
-	RealType DistanceSquared(const FVector3<RealType>& Point) const
+	RealType DistanceSquared(const TVector<RealType>& Point) const
 	{
 		return Point.DistanceSquared(ClosestPoint(Point));
 	}
 	
-	RealType Distance(const FVector3<RealType>& Point) const
+	RealType Distance(const TVector<RealType>& Point) const
 	{
 		return TMathUtil<RealType>::Sqrt(DistanceSquared(Point));
 	}
 
 	TOrientedBox3<RealType> GetBoundingBox() const
 	{
-		return TOrientedBox3<RealType>(Frame, FVector3<RealType>(Radius, Radius, 0));
+		return TOrientedBox3<RealType>(Frame, TVector<RealType>(Radius, Radius, 0));
 	}
 
 };
diff --git a/Engine/Source/Runtime/GeometryCore/Public/CompGeom/ConvexHull3.h b/Engine/Source/Runtime/GeometryCore/Public/CompGeom/ConvexHull3.h
--- a/Engine/Source/Runtime/GeometryCore/Public/CompGeom/ConvexHull3.h
+++ b/Engine/Source/Runtime/GeometryCore/Public/CompGeom/ConvexHull3.h
@@ -13,7 +13,7 @@
 namespace UE {
 namespace Geometry {
 
-
+using namespace UE::Math;
 
 template <typename RealType> struct TConvexHull3Internal;
 
@@ -35,7 +35,7 @@ public:
 	 * @param Filter Optional filter to include only a subset of the points in the output hull
 	 * @return true if hull was generated, false if points span < 2 dimensions
 	 */
-	bool Solve(int32 NumPoints, TFunctionRef<FVector3<RealType>(int32)> GetPointFunc, TFunctionRef<bool(int32)> FilterFunc = [](int32 Idx) {return true;});
+	bool Solve(int32 NumPoints, TFunctionRef<TVector<RealType>(int32)> GetPointFunc, TFunctionRef<bool(int32)> FilterFunc = [](int32 Idx) {return true;});
 
 	/**
 	 * Generate convex hull as long as input is not degenerate
diff --git a/Engine/Source/Runtime/GeometryCore/Public/CompGeom/PolygonTriangulation.h b/Engine/Source/Runtime/GeometryCore/Public/CompGeom/PolygonTriangulation.h
--- a/Engine/Source/Runtime/GeometryCore/Public/CompGeom/PolygonTriangulation.h
+++ b/Engine/Source/Runtime/GeometryCore/Public/CompGeom/PolygonTriangulation.h
@@ -20,7 +20,7 @@ namespace PolygonTriangulation
 
 
 	template<typename RealType>
-	void GEOMETRYCORE_API ComputePolygonPlane(const TArray<FVector3<RealType>>& VertexPositions, FVector3<RealType>& PlaneNormalOut, FVector3<RealType>& PlanePointOut );
+	void GEOMETRYCORE_API ComputePolygonPlane(const TArray<TVector<RealType>>& VertexPositions, TVector<RealType>& PlaneNormalOut, TVector<RealType>& PlanePointOut );
 
 
 	/**
@@ -29,6 +29,6 @@ namespace PolygonTriangulation
 	 * @param OutTriangles computed triangulation. Each triangle is a tuple of indices into VertexPositions.
 	 */
 	template<typename RealType>
-	void GEOMETRYCORE_API TriangulateSimplePolygon(const TArray<FVector3<RealType>>& VertexPositions, TArray<FIndex3i>& OutTriangles);
+	void GEOMETRYCORE_API TriangulateSimplePolygon(const TArray<TVector<RealType>>& VertexPositions, TArray<FIndex3i>& OutTriangles);
 
 }
\ No newline at end of file
diff --git a/Engine/Source/Runtime/GeometryCore/Public/Curve/CurveUtil.h b/Engine/Source/Runtime/GeometryCore/Public/Curve/CurveUtil.h
--- a/Engine/Source/Runtime/GeometryCore/Public/Curve/CurveUtil.h
+++ b/Engine/Source/Runtime/GeometryCore/Public/Curve/CurveUtil.h
@@ -16,11 +16,13 @@ namespace Geometry
 namespace CurveUtil
 {
 
+using namespace UE::Math;
+
 /**
  * Curve utility functions
  */
 	template<typename RealType, typename VectorType>
-	FVector3<RealType> Tangent(const TArrayView<const VectorType>& Vertices, int32 Idx, bool bLoop = false)
+	TVector<RealType> Tangent(const TArrayView<const VectorType>& Vertices, int32 Idx, bool bLoop = false)
 	{
 		int32 EndIdx = Idx + 1;
 		int32 StartIdx = Idx - 1;
@@ -89,8 +91,8 @@ namespace CurveUtil
 					int i = (ii % N);
 					int iPrev = (ii == 0) ? N - 1 : ii - 1;
 					int iNext = (ii + 1) % N;
-					FVector3<RealType> prev = Vertices[iPrev], next = Vertices[iNext];
-					FVector3<RealType> c = (prev + next) * 0.5f;
+					TVector<RealType> prev = Vertices[iPrev], next = Vertices[iNext];
+					TVector<RealType> c = (prev + next) * 0.5f;
 					Vertices[i] = (1 - Alpha) * Vertices[i] + (Alpha) * c;
 				}
 			}
@@ -105,8 +107,8 @@ namespace CurveUtil
 					{
 						continue;
 					}
-					FVector3<RealType> prev = Vertices[i - 1], next = Vertices[i + 1];
-					FVector3<RealType> c = (prev + next) * 0.5f;
+					TVector<RealType> prev = Vertices[i - 1], next = Vertices[i + 1];
+					TVector<RealType> c = (prev + next) * 0.5f;
 					Vertices[i] = (1 - Alpha) * Vertices[i] + (Alpha) * c;
 				}
 			}
@@ -122,7 +124,7 @@ namespace CurveUtil
 	void IterativeSmooth(TArrayView<VectorType> Vertices, int StartIdx, int EndIdx, double Alpha, int NumIterations, bool bClosed)
 	{
 		int N = Vertices.Num();
-		TArray<FVector3<RealType>> Buffer;
+		TArray<TVector<RealType>> Buffer;
 		Buffer.SetNumZeroed(N);
 
 		if (bClosed)
@@ -134,8 +136,8 @@ namespace CurveUtil
 					int i = (ii % N);
 					int iPrev = (ii == 0) ? N - 1 : ii - 1;
 					int iNext = (ii + 1) % N;
-					FVector3<RealType> prev = Vertices[iPrev], next = Vertices[iNext];
-					FVector3<RealType> c = (prev + next) * 0.5f;
+					TVector<RealType> prev = Vertices[iPrev], next = Vertices[iNext];
+					TVector<RealType> c = (prev + next) * 0.5f;
 					Buffer[i] = (1 - Alpha) * Vertices[i] + (Alpha) * c;
 				}
 				for (int ii = StartIdx; ii < EndIdx; ++ii)
@@ -155,8 +157,8 @@ namespace CurveUtil
 					{
 						continue;
 					}
-					FVector3<RealType> prev = Vertices[i - 1], next = Vertices[i + 1];
-					FVector3<RealType> c = (prev + next) * 0.5f;
+					TVector<RealType> prev = Vertices[i - 1], next = Vertices[i + 1];
+					TVector<RealType> c = (prev + next) * 0.5f;
 					Buffer[i] = (1 - Alpha) * Vertices[i] + (Alpha) * c;
 				}
 				for (int i = StartIdx; i < EndIdx && i < N; ++i)
diff --git a/Engine/Source/Runtime/GeometryCore/Public/Curve/DynamicGraph3.h b/Engine/Source/Runtime/GeometryCore/Public/Curve/DynamicGraph3.h
--- a/Engine/Source/Runtime/GeometryCore/Public/Curve/DynamicGraph3.h
+++ b/Engine/Source/Runtime/GeometryCore/Public/Curve/DynamicGraph3.h
@@ -20,23 +20,25 @@ namespace UE
 namespace Geometry
 {
 
+using namespace UE::Math;
+
 template <typename T>
 class FDynamicGraph3 : public FDynamicGraph
 {
 	TDynamicVectorN<T, 3> Vertices;
 
 public:
-	static FVector3<T> InvalidVertex()
+	static TVector<T> InvalidVertex()
 	{
-		return FVector3<T>(TNumericLimits<T>::Max(), 0, 0);
+		return TVector<T>(TNumericLimits<T>::Max(), 0, 0);
 	}
 
-	FVector3<T> GetVertex(int VID) const
+	TVector<T> GetVertex(int VID) const
 	{
 		return vertices_refcount.IsValid(VID) ? Vertices.AsVector3(VID) : InvalidVertex();
 	}
 
-	void SetVertex(int VID, FVector3<T> VNewPos)
+	void SetVertex(int VID, TVector<T> VNewPos)
 	{
 		check(VectorUtil::IsFinite(VNewPos)); // this will really catch a lot of bugs...
 		if (vertices_refcount.IsValid(VID))
@@ -47,7 +49,7 @@ public:
 	}
 
 	using FDynamicGraph::GetEdgeV;
-	bool GetEdgeV(int EID, FVector3<T>& A, FVector3<T>& B) const
+	bool GetEdgeV(int EID, TVector<T>& A, TVector<T>& B) const
 	{
 		if (edges_refcount.IsValid(EID))
 		{
@@ -67,14 +69,14 @@ public:
 			Vertices.AsVector3(e.B));
 	}
 
-	FVector3<T> GetEdgeCenter(int EID) const
+	TVector<T> GetEdgeCenter(int EID) const
 	{
 		checkfSlow(edges_refcount.IsValid(EID), TEXT("FDynamicGraph3.GetEdgeCenter: invalid segment with id %d"), EID);
 		const FEdge& e = edges[EID];
 		return 0.5 * (Vertices.AsVector3(e.A) + Vertices.AsVector3(e.B));
 	}
 
-	int AppendVertex(FVector3<T> V)
+	int AppendVertex(TVector<T> V)
 	{
 		int vid = append_vertex_internal();
 		Vertices.InsertAt({{V.X, V.Y, V.Z}}, vid);
@@ -87,9 +89,9 @@ public:
 	}
 
 	/** Enumerate positions of all Vertices in graph */
-	value_iteration<FVector3<T>> VerticesItr() const
+	value_iteration<TVector<T>> VerticesItr() const
 	{
-		return vertices_refcount.MappedIndices<FVector3<T>>(
+		return vertices_refcount.MappedIndices<TVector<T>>(
 			[=](int vid) {
 				return Vertices.template AsVector3<T>(vid);
 			});
@@ -100,7 +102,7 @@ public:
 	FAxisAlignedBox2d GetBounds() const
 	{
 		TAxisAlignedBox2<T> AABB;
-		for (const FVector3<T>& V : Vertices())
+		for (const TVector<T>& V : Vertices())
 		{
 			AABB.Contain(V);
 		}
@@ -112,7 +114,7 @@ protected:
 	// internal used in SplitEdge
 	virtual int append_new_split_vertex(int A, int B) override
 	{
-		FVector3<T> vNew = 0.5 * (GetVertex(A) + GetVertex(B));
+		TVector<T> vNew = 0.5 * (GetVertex(A) + GetVertex(B));
 		int f = AppendVertex(vNew);
 		return f;
 	}
@@ -121,7 +123,7 @@ protected:
 	{
 		for (int VID : VertexIndices())
 		{
-			FVector3<T> V = GetVertex(VID);
+			TVector<T> V = GetVertex(VID);
 			CheckOrFailF(VectorUtil::IsFinite(V));
 		}
 	}
diff --git a/Engine/Source/Runtime/GeometryCore/Public/Distance/DistLine3Line3.h b/Engine/Source/Runtime/GeometryCore/Public/Distance/DistLine3Line3.h
--- a/Engine/Source/Runtime/GeometryCore/Public/Distance/DistLine3Line3.h
+++ b/Engine/Source/Runtime/GeometryCore/Public/Distance/DistLine3Line3.h
@@ -12,6 +12,8 @@ namespace UE
 namespace Geometry
 {
 
+using namespace UE::Math;
+
 /**
  * Compute distance between two 3D lines
  */
@@ -27,9 +29,9 @@ public:
 	Real DistanceSquared = -1.0;
 
 	bool bIsParallel = false;
-	FVector3<Real> Line1ClosestPoint;
+	TVector<Real> Line1ClosestPoint;
 	Real Line1Parameter;
-	FVector3<Real> Line2ClosestPoint;
+	TVector<Real> Line2ClosestPoint;
 	Real Line2Parameter;
 
 	TDistLine3Line3(const TLine3<Real>& Line1In, const TLine3<Real>& Line2In)
@@ -54,7 +56,7 @@ public:
 			return DistanceSquared;
 		}
 
-		FVector3<Real> kDiff = Line1.Origin - Line2.Origin;
+		TVector<Real> kDiff = Line1.Origin - Line2.Origin;
 		Real a01 = -Line1.Direction.Dot(Line2.Direction);
 		Real b0 = kDiff.Dot(Line1.Direction);
 		Real c = kDiff.SquaredLength();
diff --git a/Engine/Source/Runtime/GeometryCore/Public/Distance/DistLine3Ray3.h b/Engine/Source/Runtime/GeometryCore/Public/Distance/DistLine3Ray3.h
--- a/Engine/Source/Runtime/GeometryCore/Public/Distance/DistLine3Ray3.h
+++ b/Engine/Source/Runtime/GeometryCore/Public/Distance/DistLine3Ray3.h
@@ -13,6 +13,8 @@ namespace UE
 namespace Geometry
 {
 
+using namespace UE::Math;
+
 /**
  * Compute distance between 3D line and 3D ray
  */
@@ -27,9 +29,9 @@ public:
 	// Results
 	Real DistanceSquared = -1.0;
 
-	FVector3<Real> LineClosestPoint;
+	TVector<Real> LineClosestPoint;
 	Real LineParameter;
-	FVector3<Real> RayClosestPoint;
+	TVector<Real> RayClosestPoint;
 	Real RayParameter;
 
 
@@ -55,7 +57,7 @@ public:
 			return DistanceSquared;
 		}
 
-		FVector3<Real> kDiff = Line.Origin - Ray.Origin;
+		TVector<Real> kDiff = Line.Origin - Ray.Origin;
 		Real a01 = -Line.Direction.Dot(Ray.Direction);
 		Real b0 = kDiff.Dot(Line.Direction);
 		Real c = kDiff.SquaredLength();
diff --git a/Engine/Source/Runtime/GeometryCore/Public/Distance/DistLine3Segment3.h b/Engine/Source/Runtime/GeometryCore/Public/Distance/DistLine3Segment3.h
--- a/Engine/Source/Runtime/GeometryCore/Public/Distance/DistLine3Segment3.h
+++ b/Engine/Source/Runtime/GeometryCore/Public/Distance/DistLine3Segment3.h
@@ -14,6 +14,8 @@ namespace UE
 namespace Geometry
 {
 
+using namespace UE::Math;
+
 /**
 * Compute unsigned distance between 3D line and 3D segment
 */
@@ -28,7 +30,7 @@ public:
 	// Output
 	Real DistanceSquared = -1.0;
 	Real LineParameter, SegmentParameter;
-	FVector3<Real> LineClosest, SegmentClosest;
+	TVector<Real> LineClosest, SegmentClosest;
 
 
 	TDistLine3Segment3(const TLine3<Real>& LineIn, const TSegment3<Real>& SegmentIn) : Line(LineIn), Segment(SegmentIn)
@@ -51,7 +53,7 @@ public:
 			return DistanceSquared;
 		}
 
-		FVector3<Real> diff = Line.Origin - Segment.Center;
+		TVector<Real> diff = Line.Origin - Segment.Center;
 		Real a01 = -Line.Direction.Dot(Segment.Direction);
 		Real b0 = diff.Dot(Line.Direction);
 		Real c = diff.SquaredLength();
diff --git a/Engine/Source/Runtime/GeometryCore/Public/Distance/DistLine3Triangle3.h b/Engine/Source/Runtime/GeometryCore/Public/Distance/DistLine3Triangle3.h
--- a/Engine/Source/Runtime/GeometryCore/Public/Distance/DistLine3Triangle3.h
+++ b/Engine/Source/Runtime/GeometryCore/Public/Distance/DistLine3Triangle3.h
@@ -15,6 +15,8 @@ namespace UE
 namespace Geometry
 {
 
+using namespace UE::Math;
+
 /**
 * Compute unsigned distance between 3D line and 3D triangle
 */
@@ -29,7 +31,7 @@ public:
 	// Output
 	Real DistanceSquared = -1.0;
 	Real LineParam;
-	FVector3<Real> LineClosest, TriangleClosest, TriangleBaryCoords;
+	TVector<Real> LineClosest, TriangleClosest, TriangleBaryCoords;
 
 
 	TDistLine3Triangle3(const TLine3<Real>& LineIn, const TTriangle3<Real>& TriangleIn) : Line(LineIn), Triangle(TriangleIn)
@@ -53,17 +55,17 @@ public:
 		}
 
 		// Test if Line intersects Triangle.  If so, the squared distance is zero.
-		FVector3<Real> edge0 = Triangle.V[1] - Triangle.V[0];
-		FVector3<Real> edge1 = Triangle.V[2] - Triangle.V[0];
-		FVector3<Real> normal = edge0.Cross(edge1);
+		TVector<Real> edge0 = Triangle.V[1] - Triangle.V[0];
+		TVector<Real> edge1 = Triangle.V[2] - Triangle.V[0];
+		TVector<Real> normal = edge0.Cross(edge1);
 		Normalize(normal);
 		Real NdD = normal.Dot(Line.Direction);
 		if (TMathUtil<Real>::Abs(NdD) > TMathUtil<Real>::ZeroTolerance)
 		{
 			// The line and triangle are not parallel, so the line intersects
 			// the plane of the triangle.
-			FVector3<Real> diff = Line.Origin - Triangle.V[0];
-			FVector3<Real> U, V;
+			TVector<Real> diff = Line.Origin - Triangle.V[0];
+			TVector<Real> U, V;
 			VectorUtil::MakePerpVectors(Line.Direction, U, V);
 			Real UdE0 = U.Dot(edge0);
 			Real UdE1 = U.Dot(edge1);
@@ -87,7 +89,7 @@ public:
 				LineParam = b1 * DdE0 + b2 * DdE1 - DdDiff;
 
 				// Barycentric coordinates for the point of intersection.
-				TriangleBaryCoords = FVector3<Real>(b0, b1, b2);
+				TriangleBaryCoords = TVector<Real>(b0, b1, b2);
 
 				// The intersection point is inside or on the Triangle.
 				LineClosest = Line.Origin + LineParam * Line.Direction;
diff --git a/Engine/Source/Runtime/GeometryCore/Public/Distance/DistPoint3Triangle3.h b/Engine/Source/Runtime/GeometryCore/Public/Distance/DistPoint3Triangle3.h
--- a/Engine/Source/Runtime/GeometryCore/Public/Distance/DistPoint3Triangle3.h
+++ b/Engine/Source/Runtime/GeometryCore/Public/Distance/DistPoint3Triangle3.h
@@ -23,15 +23,15 @@ class TDistPoint3Triangle3
 {
 public:
 	// Input
-	FVector3<Real> Point;
+	TVector<Real> Point;
 	TTriangle3<Real> Triangle;
 
 	// Results
-	FVector3<Real> TriangleBaryCoords;
-	FVector3<Real> ClosestTrianglePoint;  // do we need this just use Triangle.BarycentricPoint
+	TVector<Real> TriangleBaryCoords;
+	TVector<Real> ClosestTrianglePoint;  // do we need this just use Triangle.BarycentricPoint
 
 
-	TDistPoint3Triangle3(const FVector3<Real>& PointIn, const TTriangle3<Real>& TriangleIn)
+	TDistPoint3Triangle3(const TVector<Real>& PointIn, const TTriangle3<Real>& TriangleIn)
 	{
 		Point = PointIn;
 		Triangle = TriangleIn;
@@ -48,9 +48,9 @@ public:
 
 	Real ComputeResult()
 	{
-		FVector3<Real> diff = Point - Triangle.V[0];
-		FVector3<Real> edge0 = Triangle.V[1] - Triangle.V[0];
-		FVector3<Real> edge1 = Triangle.V[2] - Triangle.V[0];
+		TVector<Real> diff = Point - Triangle.V[0];
+		TVector<Real> edge0 = Triangle.V[1] - Triangle.V[0];
+		TVector<Real> edge1 = Triangle.V[2] - Triangle.V[0];
 		Real a00 = edge0.SquaredLength();
 		Real a01 = edge0.Dot(edge1);
 		Real a11 = edge1.SquaredLength();
@@ -190,9 +190,9 @@ public:
 			}
 		}
 
-		TriangleBaryCoords = FVector3<Real>((Real)1 - p[0] - p[1], p[0], p[1]);
+		TriangleBaryCoords = TVector<Real>((Real)1 - p[0] - p[1], p[0], p[1]);
 		ClosestTrianglePoint = Triangle.V[0] + p[0] * edge0 + p[1] * edge1;
-		return Point.DistanceSquared(ClosestTrianglePoint);
+		return DistanceSquared(Point, ClosestTrianglePoint);
 	}
 
 private:
diff --git a/Engine/Source/Runtime/GeometryCore/Public/Distance/DistRay3Segment3.h b/Engine/Source/Runtime/GeometryCore/Public/Distance/DistRay3Segment3.h
--- a/Engine/Source/Runtime/GeometryCore/Public/Distance/DistRay3Segment3.h
+++ b/Engine/Source/Runtime/GeometryCore/Public/Distance/DistRay3Segment3.h
@@ -13,6 +13,8 @@ namespace UE
 namespace Geometry
 {
 
+using namespace UE::Math;
+
 /**
  * Compute distance between 3D ray and 3D segment
  */
@@ -27,9 +29,9 @@ public:
 	// Results
 	Real DistanceSquared = -1.0;
 
-	FVector3<Real> RayClosestPoint;
+	TVector<Real> RayClosestPoint;
 	Real RayParameter;
-	FVector3<Real> SegmentClosestPoint;
+	TVector<Real> SegmentClosestPoint;
 	Real SegmentParameter;
 
 
@@ -224,7 +226,7 @@ public:
 	static double SquaredDistance(const TRay3<Real>& Ray, const TSegment3<Real>& Segment,
 		Real& RayParam, Real& SegParam)
 	{
-		FVector3<Real> diff = Ray.Origin - Segment.Center;
+		TVector<Real> diff = Ray.Origin - Segment.Center;
 		double a01 = -Ray.Direction.Dot(Segment.Direction);
 		double b0 = diff.Dot(Ray.Direction);
 		double b1 = -diff.Dot(Segment.Direction);
diff --git a/Engine/Source/Runtime/GeometryCore/Public/Distance/DistSegment3Triangle3.h b/Engine/Source/Runtime/GeometryCore/Public/Distance/DistSegment3Triangle3.h
--- a/Engine/Source/Runtime/GeometryCore/Public/Distance/DistSegment3Triangle3.h
+++ b/Engine/Source/Runtime/GeometryCore/Public/Distance/DistSegment3Triangle3.h
@@ -16,6 +16,8 @@ namespace UE
 namespace Geometry
 {
 
+using namespace UE::Math;
+
 /**
 * Compute unsigned distance between 3D segment and 3D triangle
 */
@@ -30,7 +32,7 @@ public:
 	// Output
 	Real DistanceSquared = -1.0;
 	Real SegmentParameter;
-	FVector3<Real> TriangleClosest, TriangleBaryCoords, SegmentClosest;
+	TVector<Real> TriangleClosest, TriangleBaryCoords, SegmentClosest;
 	
 	TDistSegment3Triangle3(const TSegment3<Real>& SegmentIn, const TTriangle3<Real>& TriangleIn) : Segment(SegmentIn), Triangle(TriangleIn)
 	{
diff --git a/Engine/Source/Runtime/GeometryCore/Public/Distance/DistTriangle3Triangle3.h b/Engine/Source/Runtime/GeometryCore/Public/Distance/DistTriangle3Triangle3.h
--- a/Engine/Source/Runtime/GeometryCore/Public/Distance/DistTriangle3Triangle3.h
+++ b/Engine/Source/Runtime/GeometryCore/Public/Distance/DistTriangle3Triangle3.h
@@ -17,6 +17,8 @@ namespace UE
 namespace Geometry
 {
 
+using namespace UE::Math;
+
 /**
 * Compute unsigned distance between 3D segment and 3D triangle
 */
@@ -26,8 +28,8 @@ class TDistTriangle3Triangle3
 public:
 	// Output
 	Real DistanceSquared = -1.0;
-	FVector3<Real> TriangleClosest[2];
-	FVector3<Real> TriangleBaryCoords[2];
+	TVector<Real> TriangleClosest[2];
+	TVector<Real> TriangleBaryCoords[2];
 
 
 	TDistTriangle3Triangle3()
diff --git a/Engine/Source/Runtime/GeometryCore/Public/DynamicMesh/DynamicMeshOverlay.h b/Engine/Source/Runtime/GeometryCore/Public/DynamicMesh/DynamicMeshOverlay.h
--- a/Engine/Source/Runtime/GeometryCore/Public/DynamicMesh/DynamicMeshOverlay.h
+++ b/Engine/Source/Runtime/GeometryCore/Public/DynamicMesh/DynamicMeshOverlay.h
@@ -571,7 +571,7 @@ protected:
 	/** Set the value at an Element to be a linear interpolation of two other Elements */
 	void SetElementFromLerp(int SetElement, int ElementA, int ElementB, double Alpha);
 	/** Set the value at an Element to be a barycentric interpolation of three other Elements */
-	void SetElementFromBary(int SetElement, int ElementA, int ElementB, int ElementC, const FVector3<double>& BaryCoords);
+	void SetElementFromBary(int SetElement, int ElementA, int ElementB, int ElementC, const FVector3d& BaryCoords);
 
 	/** updates the triangles array and optionally the element reference counts */
 	void InternalSetTriangle(int TriangleID, const FIndex3i& TriElements, bool bUpdateRefCounts, bool bAllowElementFreeing = true);
diff --git a/Engine/Source/Runtime/GeometryCore/Public/DynamicMesh/DynamicPointSet3.h b/Engine/Source/Runtime/GeometryCore/Public/DynamicMesh/DynamicPointSet3.h
--- a/Engine/Source/Runtime/GeometryCore/Public/DynamicMesh/DynamicPointSet3.h
+++ b/Engine/Source/Runtime/GeometryCore/Public/DynamicMesh/DynamicPointSet3.h
@@ -20,6 +20,8 @@ namespace UE
 namespace Geometry
 {
 
+using namespace UE::Math;
+
 /**
  * TDynamicPointSet3 implements a dynamic 3D point set, templated on real-value type (float or double).
  * The points are indexed and the class allows for gaps in the index space.
@@ -118,7 +120,7 @@ public:
 
 
 	/** Append vertex at position, returns vid */
-	int AppendVertex(const FVector3<RealType>& Position)
+	int AppendVertex(const TVector<RealType>& Position)
 	{
 		int vid = VertexRefCounts.Allocate();
 		int i = 3 * vid;
@@ -152,7 +154,7 @@ public:
 	 * If bUnsafe, we use fast id allocation that does not update free list.
 	 * You should only be using this between BeginUnsafeVerticesInsert() / EndUnsafeVerticesInsert() calls
 	 */
-	EMeshResult InsertVertex(int VertexID, const FVector3<RealType>& Position, bool bUnsafe = false)
+	EMeshResult InsertVertex(int VertexID, const TVector<RealType>& Position, bool bUnsafe = false)
 	{
 		if (VertexRefCounts.IsValid(VertexID))
 		{
@@ -187,15 +189,15 @@ public:
 public:
 
 	/** @return the vertex position */
-	inline FVector3<RealType> GetVertex(int VertexID) const
+	inline TVector<RealType> GetVertex(int VertexID) const
 	{
 		check(IsVertex(VertexID));
 		int i = 3 * VertexID;
-		return FVector3<RealType>(Vertices[i], Vertices[i + 1], Vertices[i + 2]);
+		return TVector<RealType>(Vertices[i], Vertices[i + 1], Vertices[i + 2]);
 	}
 
 	/** Set vertex position */
-	inline void SetVertex(int VertexID, const FVector3<RealType>& vNewPos)
+	inline void SetVertex(int VertexID, const TVector<RealType>& vNewPos)
 	{
 		//check(VectorUtil::IsFinite(vNewPos));
 		check(IsVertex(VertexID));
@@ -219,11 +221,11 @@ public:
 	}
 
 	/** Enumerate positions of all points */
-	FRefCountVector::MappedEnumerable<FVector3<RealType>> VerticesItr() const
+	FRefCountVector::MappedEnumerable<TVector<RealType>> VerticesItr() const
 	{
-		return VertexRefCounts.MappedIndices<FVector3<RealType>>([this](int VertexID) {
+		return VertexRefCounts.MappedIndices<TVector<RealType>>([this](int VertexID) {
 			int i = 3 * VertexID;
-			return FVector3<RealType>(Vertices[i], Vertices[i + 1], Vertices[i + 2]);
+			return TVector<RealType>(Vertices[i], Vertices[i + 1], Vertices[i + 2]);
 		});
 	}
 
@@ -260,7 +262,7 @@ public:
 		for (int vi : VertexIndicesItr())
 		{
 			int k = 3 * vi;
-			Box.Contain( FVector3<RealType>(Vertices[k], Vertices[k + 1], Vertices[k + 2]) );
+			Box.Contain( TVector<RealType>(Vertices[k], Vertices[k + 1], Vertices[k + 2]) );
 		}
 		return Box;
 	}
diff --git a/Engine/Source/Runtime/GeometryCore/Public/DynamicMesh/MeshTangents.h b/Engine/Source/Runtime/GeometryCore/Public/DynamicMesh/MeshTangents.h
--- a/Engine/Source/Runtime/GeometryCore/Public/DynamicMesh/MeshTangents.h
+++ b/Engine/Source/Runtime/GeometryCore/Public/DynamicMesh/MeshTangents.h
@@ -11,6 +11,8 @@ namespace UE
 namespace Geometry
 {
 
+using namespace UE::Math;
+
 /**
  * FDynamicMeshTangents is a helper object for accessing tangents stored in the AttributeSet of a FDynamicMesh3.
  */
@@ -65,9 +67,9 @@ protected:
 	/** Target Mesh */
 	const FDynamicMesh3* Mesh;
 	/** Set of computed tangents */
-	TArray<FVector3<RealType>> Tangents;
+	TArray<TVector<RealType>> Tangents;
 	/** Set of computed bitangents */
-	TArray<FVector3<RealType>> Bitangents;
+	TArray<TVector<RealType>> Bitangents;
 	/** Indices of degenerate triangles. This may be useful to know externally, as those tangets are often meaningless */
 	TArray<int32> AllDegenerateTris;
 
@@ -87,9 +89,9 @@ public:
 		this->Mesh = MeshIn;
 	}
 
-	const TArray<FVector3<RealType>>& GetTangents() const { return Tangents; }
+	const TArray<TVector<RealType>>& GetTangents() const { return Tangents; }
 
-	const TArray<FVector3<RealType>>& GetBitangents() const { return Bitangents; }
+	const TArray<TVector<RealType>>& GetBitangents() const { return Bitangents; }
 
 	const TArray<int32>& GetDegenerateTris() const { return AllDegenerateTris; }
 
@@ -181,8 +183,8 @@ public:
 		}
 		else
 		{
-			TangentOut = FVector3<RealType>::UnitX();
-			BitangentOut = FVector3<RealType>::UnitY();
+			TangentOut = TVector<RealType>::UnitX();
+			BitangentOut = TVector<RealType>::UnitY();
 		}
 	}
 
@@ -192,7 +194,7 @@ public:
 	 * @param TriangleID triangle index in mesh
 	 * @param TriVertIdx vertex index in range 0,1,2
 	 */
-	void SetPerTriangleTangent(int TriangleID, int TriVertIdx, const FVector3<RealType>& Tangent, const FVector3<RealType>& Bitangent)
+	void SetPerTriangleTangent(int TriangleID, int TriVertIdx, const TVector<RealType>& Tangent, const TVector<RealType>& Bitangent)
 	{
 		int k = TriangleID * 3 + TriVertIdx;
 		Tangents[k] = Tangent;
@@ -260,16 +262,16 @@ template<typename RealType>
 template<typename OtherRealType>
 void TMeshTangents<RealType>::CopyTriVertexTangents(const TMeshTangents<OtherRealType>& OtherMeshTangents)
 {
-	const TArray<FVector3<OtherRealType>>& OtherTangents = OtherMeshTangents.GetTangents();
-	const TArray<FVector3<OtherRealType>>& OtherBitangents = OtherMeshTangents.GetBitangents();
+	const TArray<TVector<OtherRealType>>& OtherTangents = OtherMeshTangents.GetTangents();
+	const TArray<TVector<OtherRealType>>& OtherBitangents = OtherMeshTangents.GetBitangents();
 
 	int32 Num = OtherTangents.Num();
 	check(Mesh->MaxTriangleID() * 3 == Num);
 	InitializeTriVertexTangents(false);
 	for (int32 k = 0; k < Num; ++k)
 	{
-		Tangents[k] = FVector3<RealType>(OtherTangents[k]);
-		Bitangents[k] = FVector3<RealType>(OtherBitangents[k]);
+		Tangents[k] = TVector<RealType>(OtherTangents[k]);
+		Bitangents[k] = TVector<RealType>(OtherBitangents[k]);
 	}
 	AllDegenerateTris = OtherMeshTangents.GetDegenerateTris();
 }
diff --git a/Engine/Source/Runtime/GeometryCore/Public/FrameTypes.h b/Engine/Source/Runtime/GeometryCore/Public/FrameTypes.h
--- a/Engine/Source/Runtime/GeometryCore/Public/FrameTypes.h
+++ b/Engine/Source/Runtime/GeometryCore/Public/FrameTypes.h
@@ -29,7 +29,7 @@ struct TFrame3
 	/**
 	 * Origin of the frame
 	 */
-	FVector3<RealType> Origin;
+	TVector<RealType> Origin;
 
 	/**
 	 * Rotation of the frame. Think of this as the rotation of the unit X/Y/Z axes to the 3D frame axes.
@@ -41,14 +41,14 @@ struct TFrame3
 	 */
 	TFrame3()
 	{
-		Origin = FVector3<RealType>::Zero();
+		Origin = TVector<RealType>::Zero();
 		Rotation = TQuaternion<RealType>::Identity();
 	}
 
 	/**
 	 * Construct a frame at the given Origin aligned to the unit axes
 	 */
-	explicit TFrame3(const FVector3<RealType>& OriginIn)
+	explicit TFrame3(const TVector<RealType>& OriginIn)
 	{
 		Origin = OriginIn;
 		Rotation = TQuaternion<RealType>::Identity();
@@ -57,7 +57,7 @@ struct TFrame3
 	/**
 	 * Construct a Frame from the given Origin and Rotation
 	 */
-	TFrame3(const FVector3<RealType>& OriginIn, const TQuaternion<RealType> RotationIn)
+	TFrame3(const TVector<RealType>& OriginIn, const TQuaternion<RealType> RotationIn)
 	{
 		Origin = OriginIn;
 		Rotation = RotationIn;
@@ -68,10 +68,10 @@ struct TFrame3
 	 * @param OriginIn origin of frame
 	 * @param SetZ target Z axis
 	 */
-	TFrame3(const FVector3<RealType>& OriginIn, const FVector3<RealType>& SetZ)
+	TFrame3(const TVector<RealType>& OriginIn, const TVector<RealType>& SetZ)
 	{
 		Origin = OriginIn;
-		Rotation.SetFromTo(FVector3<RealType>::UnitZ(), SetZ);
+		Rotation.SetFromTo(TVector<RealType>::UnitZ(), SetZ);
 	}
 
 	/**
@@ -81,7 +81,7 @@ struct TFrame3
 	 * @param Y desired Y axis of frame
 	 * @param Z desired Z axis of frame
 	 */
-	TFrame3(const FVector3<RealType>& OriginIn, const FVector3<RealType>& X, const FVector3<RealType>& Y, const FVector3<RealType>& Z)
+	TFrame3(const TVector<RealType>& OriginIn, const TVector<RealType>& X, const TVector<RealType>& Y, const TVector<RealType>& Z)
 	{
 		Origin = OriginIn;
 		Rotation = TQuaternion<RealType>( TMatrix3<RealType>(X, Y, Z, false) );
@@ -90,7 +90,7 @@ struct TFrame3
 	/** Construct a Frame from an FTransform */
 	explicit TFrame3(const FTransform& Transform)
 	{
-		Origin = FVector3<RealType>(Transform.GetTranslation());
+		Origin = TVector<RealType>(Transform.GetTranslation());
 		Rotation = TQuaternion<RealType>(Transform.GetRotation());
 	}
 
@@ -98,15 +98,15 @@ struct TFrame3
 	explicit TFrame3(const FPlane& Plane)
 	{
 		FVector Normal(Plane.X, Plane.Y, Plane.Z);
-		Origin = (RealType)Plane.W * FVector3<RealType>(Normal);
-		Rotation.SetFromTo(FVector3<RealType>::UnitZ(), (FVector3<RealType>)Normal);
+		Origin = (RealType)Plane.W * TVector<RealType>(Normal);
+		Rotation.SetFromTo(TVector<RealType>::UnitZ(), (TVector<RealType>)Normal);
 	}
 
 
 	/** Construct a Frame from an FVector and FQuat */
 	explicit TFrame3(const FVector& OriginIn, const FQuat& RotationIn)
 	{
-		Origin = FVector3<RealType>(OriginIn);
+		Origin = TVector<RealType>(OriginIn);
 		Rotation = TQuaternion<RealType>(RotationIn);
 	}
 
@@ -114,7 +114,7 @@ struct TFrame3
 	template<typename RealType2>
 	explicit TFrame3(const TFrame3<RealType2>& OtherFrame)
 	{
-		Origin = static_cast<FVector3<RealType>>(OtherFrame.Origin);
+		Origin = static_cast<TVector<RealType>>(OtherFrame.Origin);
 		Rotation = static_cast<TQuaternion<RealType>>(OtherFrame.Rotation);
 	}
 
@@ -123,7 +123,7 @@ struct TFrame3
 	 * @param AxisIndex index of axis of frame, either 0, 1, or 2
 	 * @return axis vector
 	 */
-	FVector3<RealType> GetAxis(int AxisIndex) const
+	TVector<RealType> GetAxis(int AxisIndex) const
 	{
 		switch (AxisIndex)
 		{
@@ -135,32 +135,32 @@ struct TFrame3
 			return Rotation.AxisZ();
 		default:
 			checkNoEntry();
-			return FVector3<RealType>::Zero(); // compiler demands a return value
+			return TVector<RealType>::Zero(); // compiler demands a return value
 		}
 	}
 
 	/**
 	 * @return X/Y/Z axes of frame. This is more efficient than calculating each axis separately.
 	 */
-	void GetAxes(FVector3<RealType>& X, FVector3<RealType>& Y, FVector3<RealType>& Z) const
+	void GetAxes(TVector<RealType>& X, TVector<RealType>& Y, TVector<RealType>& Z) const
 	{
 		Rotation.GetAxes(X, Y, Z);
 	}
 
 	/** @return X axis of frame (axis 0) */
-	FVector3<RealType> X() const
+	TVector<RealType> X() const
 	{
 		return Rotation.AxisX();
 	}
 
 	/** @return Y axis of frame (axis 1) */
-	FVector3<RealType> Y() const
+	TVector<RealType> Y() const
 	{
 		return Rotation.AxisY();
 	}
 
 	/** @return Z axis of frame (axis 2) */
-	FVector3<RealType> Z() const
+	TVector<RealType> Z() const
 	{
 		return Rotation.AxisZ();
 	}
@@ -184,36 +184,36 @@ struct TFrame3
 	}
 
 	/** @return point at distances along frame axes */
-	FVector3<RealType> PointAt(RealType X, RealType Y, RealType Z) const
+	TVector<RealType> PointAt(RealType X, RealType Y, RealType Z) const
 	{
-		return Rotation * FVector3<RealType>(X,Y,Z) + Origin;
+		return Rotation * TVector<RealType>(X,Y,Z) + Origin;
 	}
 
 	/** @return point at distances along frame axes */
-	FVector3<RealType> PointAt(const FVector3<RealType>& Point) const
+	TVector<RealType> PointAt(const TVector<RealType>& Point) const
 	{
-		return Rotation * FVector3<RealType>(Point.X, Point.Y, Point.Z) + Origin;
+		return Rotation * TVector<RealType>(Point.X, Point.Y, Point.Z) + Origin;
 	}
 	
 	/** @return input Point transformed into local coordinate system of Frame */
-	FVector3<RealType> ToFramePoint(const FVector3<RealType>& Point) const
+	TVector<RealType> ToFramePoint(const TVector<RealType>& Point) const
 	{
 		return Rotation.InverseMultiply((Point-Origin));
 	}
 	/** @return input Point transformed from local coordinate system of Frame into "World" coordinate system */
-	FVector3<RealType> FromFramePoint(const FVector3<RealType>& Point) const
+	TVector<RealType> FromFramePoint(const TVector<RealType>& Point) const
 	{
 		return Rotation * Point + Origin;
 	}
 
 
 	/** @return input Vector transformed into local coordinate system of Frame */
-	FVector3<RealType> ToFrameVector(const FVector3<RealType>& Vector) const
+	TVector<RealType> ToFrameVector(const TVector<RealType>& Vector) const
 	{
 		return Rotation.InverseMultiply(Vector);
 	}
 	/** @return input Vector transformed from local coordinate system of Frame into "World" coordinate system */
-	FVector3<RealType> FromFrameVector(const FVector3<RealType>& Vector) const
+	TVector<RealType> FromFrameVector(const TVector<RealType>& Vector) const
 	{
 		return Rotation * Vector;
 	}
@@ -263,7 +263,7 @@ struct TFrame3
 	 * @param PlaneNormalAxis which plane to project onto, identified by perpendicular normal. Default is 2, ie normal is Z, plane is (X,Y)
 	 * @return 2D coordinates in UV plane, relative to origin
 	 */
-	TVector2<RealType> ToPlaneUV(const FVector3<RealType>& Pos, int PlaneNormalAxis = 2) const
+	TVector2<RealType> ToPlaneUV(const TVector<RealType>& Pos, int PlaneNormalAxis = 2) const
 	{
 		int Axis0 = 0, Axis1 = 1;
 		if (PlaneNormalAxis == 0)
@@ -274,7 +274,7 @@ struct TFrame3
 		{
 			Axis1 = 2;
 		}
-		FVector3<RealType> LocalPos = Pos - Origin;
+		TVector<RealType> LocalPos = Pos - Origin;
 		RealType U = LocalPos.Dot(GetAxis(Axis0));
 		RealType V = LocalPos.Dot(GetAxis(Axis1));
 		return TVector2<RealType>(U, V);
@@ -288,9 +288,9 @@ struct TFrame3
 	 * @param PlaneNormalAxis which plane to map to, identified by perpendicular normal. Default is 2, ie normal is Z, plane is (X,Y)
 	 * @return 3D coordinates in frame's plane (including Origin translation)
 	 */
-	FVector3<RealType> FromPlaneUV(const TVector2<RealType>& PosUV, int PlaneNormalAxis = 2) const
+	TVector<RealType> FromPlaneUV(const TVector2<RealType>& PosUV, int PlaneNormalAxis = 2) const
 	{
-		FVector3<RealType> PlanePos(PosUV[0], PosUV[1], 0);
+		TVector<RealType> PlanePos(PosUV[0], PosUV[1], 0);
 		if (PlaneNormalAxis == 0)
 		{
 			PlanePos[0] = 0; PlanePos[2] = PosUV[0];
@@ -310,10 +310,10 @@ struct TFrame3
 	 * @param PlaneNormalAxis which plane to project onto, identified by perpendicular normal. Default is 2, ie normal is Z, plane is (X,Y)
 	 * @return 3D coordinate in the plane
 	 */
-	FVector3<RealType> ToPlane(const FVector3<RealType>& Pos, int PlaneNormalAxis = 2) const
+	TVector<RealType> ToPlane(const TVector<RealType>& Pos, int PlaneNormalAxis = 2) const
 	{
-		FVector3<RealType> Normal = GetAxis(PlaneNormalAxis);
-		FVector3<RealType> LocalVec = Pos - Origin;
+		TVector<RealType> Normal = GetAxis(PlaneNormalAxis);
+		TVector<RealType> LocalVec = Pos - Origin;
 		RealType SignedDist = LocalVec.Dot(Normal);
 		return Pos - SignedDist * Normal;
 	}
@@ -343,7 +343,7 @@ struct TFrame3
 	 */
 	void Transform(const FTransform& XForm)
 	{
-		Origin = FVector3<RealType>(XForm.TransformPosition((FVector)Origin));
+		Origin = TVector<RealType>(XForm.TransformPosition((FVector)Origin));
 		Rotate(TQuaternion<RealType>(XForm.GetRotation()));
 	}
 
@@ -363,7 +363,7 @@ struct TFrame3
 	 * @param AxisIndex which axis to align
 	 * @param ToDirection target direction
 	 */
-	void AlignAxis(int AxisIndex, const FVector3<RealType>& ToDirection)
+	void AlignAxis(int AxisIndex, const TVector<RealType>& ToDirection)
 	{
 		TQuaternion<RealType> RelRotation(GetAxis(AxisIndex), ToDirection);
 		Rotate(RelRotation);
@@ -376,10 +376,10 @@ struct TFrame3
 	 * @param ToDirection target direction
 	 * @param AroundVector rotation is constrained to be around this vector (ie this direction in frame stays constant)
 	 */
-	void ConstrainedAlignAxis(int AxisIndex, const FVector3<RealType>& ToDirection, const FVector3<RealType>& AroundVector)
+	void ConstrainedAlignAxis(int AxisIndex, const TVector<RealType>& ToDirection, const TVector<RealType>& AroundVector)
 	{
 		//@todo PlaneAngleSigned does acos() and then SetAxisAngleD() does cos/sin...can we optimize this?
-		FVector3<RealType> AxisVec = GetAxis(AxisIndex);
+		TVector<RealType> AxisVec = GetAxis(AxisIndex);
 		RealType AngleDeg = VectorUtil::PlaneAngleSignedD(AxisVec, ToDirection, AroundVector);
 		TQuaternion<RealType> RelRotation;
 		RelRotation.SetAxisAngleD(AroundVector, AngleDeg);
@@ -401,15 +401,15 @@ struct TFrame3
 	 * @param UpDotTolerance defaults to cos(45), ie flip between regions happens roughly half way to poles
 	 */
 	void ConstrainedAlignPerpAxes(int PerpAxis1 = 0, int PerpAxis2 = 1, int NormalAxis = 2, 
-		const FVector3<RealType>& UpAxis = FVector3<RealType>::UnitZ(),
-		const FVector3<RealType>& FallbackAxis = FVector3<RealType>::UnitX(),
+		const TVector<RealType>& UpAxis = TVector<RealType>::UnitZ(),
+		const TVector<RealType>& FallbackAxis = TVector<RealType>::UnitX(),
 		RealType UpDotTolerance = (RealType)0.707)
 	{
 		check(PerpAxis1 != PerpAxis2 && PerpAxis1 != NormalAxis && PerpAxis2 != NormalAxis);
-		const FVector3<RealType> NormalVec = GetAxis(NormalAxis);
+		const TVector<RealType> NormalVec = GetAxis(NormalAxis);
 
 		// decide if we should use Fallback (polar-cap) axis or main (equator-region) axis
-		const FVector3<RealType>& TargetAxis =
+		const TVector<RealType>& TargetAxis =
 			(TMathUtil<RealType>::Abs(NormalVec.Dot(UpAxis)) > UpDotTolerance) ?
 			FallbackAxis : UpAxis;
 
@@ -426,23 +426,23 @@ struct TFrame3
 	 * @param RayOrigin origin of ray
 	 * @param RayDirection direction of ray
 	 * @param PlaneNormalAxis which axis of frame to use as plane normal
-	 * @param HitPointOut intersection point, or FVector3::Max() if ray does not hit plane or is parallel to plane
+	 * @param HitPointOut intersection point, or invalid point if ray does not hit plane or is parallel to plane
 	 * @return true if ray intersects plane and HitPointOut is valid
 	 */
-	bool RayPlaneIntersection(const FVector3<RealType>& RayOrigin, const FVector3<RealType>& RayDirection, int PlaneNormalAxis, UE::Math::TVector<RealType>& HitPointOut) const
+	bool RayPlaneIntersection(const TVector<RealType>& RayOrigin, const TVector<RealType>& RayDirection, int PlaneNormalAxis, UE::Math::TVector<RealType>& HitPointOut) const
 	{
-		FVector3<RealType> Normal = GetAxis(PlaneNormalAxis);
+		TVector<RealType> Normal = GetAxis(PlaneNormalAxis);
 		RealType PlaneD = -Origin.Dot(Normal);
 		RealType NormalDot = RayDirection.Dot(Normal);
 		if (VectorUtil::EpsilonEqual(NormalDot, (RealType)0, TMathUtil<RealType>::ZeroTolerance))
 		{
-			HitPointOut = FVector3<RealType>::MaxVector();
+			HitPointOut = TVector<RealType>(TNumericLimits<RealType>::Max(), TNumericLimits<RealType>::Max(), TNumericLimits<RealType>::Max());
 			return false;
 		}
 		RealType t = -( RayOrigin.Dot(Normal) + PlaneD) / NormalDot;
 		if (t < 0)
 		{
-			HitPointOut = FVector3<RealType>::MaxVector();
+			HitPointOut = TVector<RealType>(TNumericLimits<RealType>::Max(), TNumericLimits<RealType>::Max(), TNumericLimits<RealType>::Max());
 			return false;
 		}
 		HitPointOut = RayOrigin + t * RayDirection;
diff --git a/Engine/Source/Runtime/GeometryCore/Public/Generators/LineSegmentGenerators.h b/Engine/Source/Runtime/GeometryCore/Public/Generators/LineSegmentGenerators.h
--- a/Engine/Source/Runtime/GeometryCore/Public/Generators/LineSegmentGenerators.h
+++ b/Engine/Source/Runtime/GeometryCore/Public/Generators/LineSegmentGenerators.h
@@ -11,6 +11,8 @@ namespace UE
 {
 	namespace Geometry
 	{
+		using namespace UE::Math;
+
 		/**
 		 * call EmitLineFunc for the 12 line segments defined by the given box parameters
 		 * @param HalfDimensions X/Y/Z dimensions of box
@@ -22,13 +24,13 @@ namespace UE
 		 */
 		template<typename RealType>
 		void GenerateBoxSegments(
-			const FVector3<RealType>& HalfDimensions, 
-			const FVector3<RealType>& Center, 
-			const FVector3<RealType>& AxisX, 
-			const FVector3<RealType>& AxisY, 
-			const FVector3<RealType>& AxisZ,
+			const TVector<RealType>& HalfDimensions, 
+			const TVector<RealType>& Center, 
+			const TVector<RealType>& AxisX, 
+			const TVector<RealType>& AxisY, 
+			const TVector<RealType>& AxisZ,
 			const TTransform3<RealType>& Transform,
-			TFunctionRef<void(const FVector3<RealType>& A, const FVector3<RealType>& B)> EmitLineFunc);
+			TFunctionRef<void(const TVector<RealType>& A, const TVector<RealType>& B)> EmitLineFunc);
 
 		/**
 		 * call EmitLineFunc for the line segments that make up the Circle defined by the given parameters and sampled with NumSteps vertices
@@ -43,11 +45,11 @@ namespace UE
 		void GenerateCircleSegments(
 			int32 NumSteps, 
 			RealType Radius, 
-			const FVector3<RealType>& Center,
-			const FVector3<RealType>& AxisX,
-			const FVector3<RealType>& AxisY,
+			const TVector<RealType>& Center,
+			const TVector<RealType>& AxisX,
+			const TVector<RealType>& AxisY,
 			const TTransform3<RealType>& Transform,
-			TFunctionRef<void(const FVector3<RealType>& A, const FVector3<RealType>& B)> EmitLineFunc);
+			TFunctionRef<void(const TVector<RealType>& A, const TVector<RealType>& B)> EmitLineFunc);
 
 		/**
 		 * call EmitLineFunc for the line segments that make up the Circular Arc defined by the given parameters and sampled with NumSteps vertices
@@ -65,11 +67,11 @@ namespace UE
 			RealType Radius,
 			RealType StartAngle, 
 			RealType EndAngle,
-			const FVector3<RealType>& Center,
-			const FVector3<RealType>& AxisX,
-			const FVector3<RealType>& AxisY,
+			const TVector<RealType>& Center,
+			const TVector<RealType>& AxisX,
+			const TVector<RealType>& AxisY,
 			const TTransform3<RealType>& Transform,
-			TFunctionRef<void(const FVector3<RealType>& A, const FVector3<RealType>& B)> EmitLineFunc);
+			TFunctionRef<void(const TVector<RealType>& A, const TVector<RealType>& B)> EmitLineFunc);
 	}
 }
 
@@ -79,16 +81,16 @@ namespace UE
 
 template<typename RealType>
 void UE::Geometry::GenerateBoxSegments(
-	const FVector3<RealType>& HalfDimensions,
-	const FVector3<RealType>& Center,
-	const FVector3<RealType>& AxisX,
-	const FVector3<RealType>& AxisY,
-	const FVector3<RealType>& AxisZ,
+	const TVector<RealType>& HalfDimensions,
+	const TVector<RealType>& Center,
+	const TVector<RealType>& AxisX,
+	const TVector<RealType>& AxisY,
+	const TVector<RealType>& AxisZ,
 	const TTransform3<RealType>& Transform,
-	TFunctionRef<void(const FVector3<RealType>& A, const FVector3<RealType>& B)> EmitLineFunc)
+	TFunctionRef<void(const TVector<RealType>& A, const TVector<RealType>& B)> EmitLineFunc)
 {
 	// B is box max/min, P and Q are used to store the start and endpoints of the segments we create
-	FVector3<RealType>	B[2], P, Q;
+	TVector<RealType>	B[2], P, Q;
 	B[0] = HalfDimensions; // max
 	B[1] = -B[0]; // min
 
@@ -120,20 +122,20 @@ void UE::Geometry::GenerateBoxSegments(
 template<typename RealType>
 void UE::Geometry::GenerateCircleSegments(
 	int32 NumSteps, RealType Radius,
-	const FVector3<RealType>& Center,
-	const FVector3<RealType>& AxisX,
-	const FVector3<RealType>& AxisY,
+	const TVector<RealType>& Center,
+	const TVector<RealType>& AxisX,
+	const TVector<RealType>& AxisY,
 	const TTransform3<RealType>& Transform,
-	TFunctionRef<void(const FVector3<RealType>& A, const FVector3<RealType>& B)> EmitLineFunc)
+	TFunctionRef<void(const TVector<RealType>& A, const TVector<RealType>& B)> EmitLineFunc)
 {
-	FVector3<RealType> PrevPos = FVector3<RealType>::Zero();
+	TVector<RealType> PrevPos = TVector<RealType>::Zero();
 	for (int32 i = 0; i <= NumSteps; ++i)
 	{
 		RealType t = (RealType)i / (RealType)NumSteps;
 		RealType Angle = (RealType)2.0 * TMathUtil<RealType>::Pi * t;
 		RealType PlaneX = Radius * TMathUtil<RealType>::Cos(Angle);
 		RealType PlaneY = Radius * TMathUtil<RealType>::Sin(Angle);
-		FVector3<RealType> CurPos = Transform.TransformPosition(Center + PlaneX * AxisX + PlaneY * AxisY);
+		TVector<RealType> CurPos = Transform.TransformPosition(Center + PlaneX * AxisX + PlaneY * AxisY);
 		if (i > 0)
 		{
 			EmitLineFunc(PrevPos, CurPos);
@@ -150,20 +152,20 @@ void UE::Geometry::GenerateArcSegments(
 	RealType Radius,
 	RealType StartAngle,
 	RealType EndAngle,
-	const FVector3<RealType>& Center,
-	const FVector3<RealType>& AxisX,
-	const FVector3<RealType>& AxisY,
+	const TVector<RealType>& Center,
+	const TVector<RealType>& AxisX,
+	const TVector<RealType>& AxisY,
 	const TTransform3<RealType>& Transform,
-	TFunctionRef<void(const FVector3<RealType>& A, const FVector3<RealType>& B)> EmitLineFunc)
+	TFunctionRef<void(const TVector<RealType>& A, const TVector<RealType>& B)> EmitLineFunc)
 {
-	FVector3<RealType> PrevPos = FVector3<RealType>::Zero();
+	TVector<RealType> PrevPos = TVector<RealType>::Zero();
 	for (int32 i = 0; i <= NumSteps; ++i)
 	{
 		RealType t = (RealType)i / (RealType)NumSteps;
 		RealType Angle = ((RealType)1 - t) * StartAngle + (t)*EndAngle;
 		RealType PlaneX = Radius * TMathUtil<RealType>::Cos(Angle);
 		RealType PlaneY = Radius * TMathUtil<RealType>::Sin(Angle);
-		FVector3<RealType> CurPos = Transform.TransformPosition(Center + PlaneX * AxisX + PlaneY * AxisY);
+		TVector<RealType> CurPos = Transform.TransformPosition(Center + PlaneX * AxisX + PlaneY * AxisY);
 		if (i > 0)
 		{
 			EmitLineFunc(PrevPos, CurPos);
diff --git a/Engine/Source/Runtime/GeometryCore/Public/Generators/MarchingCubes.h b/Engine/Source/Runtime/GeometryCore/Public/Generators/MarchingCubes.h
--- a/Engine/Source/Runtime/GeometryCore/Public/Generators/MarchingCubes.h
+++ b/Engine/Source/Runtime/GeometryCore/Public/Generators/MarchingCubes.h
@@ -22,6 +22,8 @@ namespace UE
 namespace Geometry
 {
 
+using namespace UE::Math;
+
 enum class /*GEOMETRYCORE_API*/ ERootfindingModes
 {
 	SingleLerp,
@@ -35,7 +37,7 @@ public:
 	/**
 	*  this is the function we will evaluate
 	*/
-	TFunction<double(FVector3<double>)> Implicit;
+	TFunction<double(TVector<double>)> Implicit;
 
 	/**
 	*  mesh surface will be at this isovalue. Normally 0 unless you want
@@ -97,7 +99,7 @@ public:
 
 	FMarchingCubes()
 	{
-		Bounds = TAxisAlignedBox3<double>(FVector3<double>::Zero(), 8);
+		Bounds = TAxisAlignedBox3<double>(TVector<double>::Zero(), 8);
 		CubeSize = 0.25;
 	}
 
@@ -149,7 +151,7 @@ public:
 	}
 
 
-	FMeshShapeGenerator& GenerateContinuation(TArrayView<const FVector3<double>> Seeds)
+	FMeshShapeGenerator& GenerateContinuation(TArrayView<const FVector3d> Seeds)
 	{
 		TRACE_CPUPROFILER_EVENT_SCOPE(Geometry_MCMesh_GenerateContinuation);
 
@@ -237,19 +239,19 @@ protected:
 		CellDimensions = FVector3i(NX, NY, NZ);
 	}
 
-	void corner_pos(const FVector3i& IJK, FVector3<double>& Pos)
+	void corner_pos(const FVector3i& IJK, TVector<double>& Pos)
 	{
 		Pos.X = Bounds.Min.X + CubeSize * IJK.X;
 		Pos.Y = Bounds.Min.Y + CubeSize * IJK.Y;
 		Pos.Z = Bounds.Min.Z + CubeSize * IJK.Z;
 	}
-	FVector3<double> corner_pos(const FVector3i& IJK)
+	TVector<double> corner_pos(const FVector3i& IJK)
 	{
-		return FVector3<double>(Bounds.Min.X + CubeSize * IJK.X,
+		return TVector<double>(Bounds.Min.X + CubeSize * IJK.X,
 			Bounds.Min.Y + CubeSize * IJK.Y,
 			Bounds.Min.Z + CubeSize * IJK.Z);
 	}
-	FVector3i cell_index(const FVector3<double>& Pos)
+	FVector3i cell_index(const TVector<double>& Pos)
 	{
 		checkSlow(Bounds.Contains(Pos));
 		return FVector3i(
@@ -316,7 +318,7 @@ protected:
 		return (Found != nullptr) ? *Found : IndexConstants::InvalidID;
 	}
 
-	int AppendOrFindVertexID(int64 hash, FVector3<double> Pos)
+	int AppendOrFindVertexID(int64 hash, TVector<double> Pos)
 	{
 		int32 SectionIndex = (int32)(hash % (NumEdgeVertexSections - 1));
 		FScopeLock Lock(&EdgeVertexSectionLocks[SectionIndex]);
@@ -347,10 +349,10 @@ protected:
 		// re-acquire the lock we need to check again that we have not already
 		// computed this edge, otherwise we will end up with duplicate vertices!
 
-		FVector3<double> pa = FVector3<double>::Zero(), pb = FVector3<double>::Zero();
+		TVector<double> pa = TVector<double>::Zero(), pb = TVector<double>::Zero();
 		corner_pos(Idx1, pa);
 		corner_pos(Idx2, pb);
-		FVector3<double> Pos = FVector3<double>::Zero();
+		TVector<double> Pos = TVector<double>::Zero();
 		find_iso(pa, pb, F1, F2, Pos);
 
 		return AppendOrFindVertexID(hash, Pos);
@@ -376,7 +378,7 @@ protected:
 		double* findval = corner_values.Find(hash);
 		if (findval == nullptr)
 		{
-			FVector3<double> V = corner_pos(Idx);
+			TVector<double> V = corner_pos(Idx);
 			value = Implicit(V);
 			corner_values.Add(hash, value);
 		} 
@@ -436,7 +438,7 @@ protected:
 			return (double)val;
 		}
 
-		FVector3<double> V = corner_pos(Idx);
+		TVector<double> V = corner_pos(Idx);
 		val = (float)Implicit(V);
 
 		{
@@ -459,7 +461,7 @@ protected:
 			return (double)val;
 		}
 
-		FVector3<double> V = corner_pos(Idx);
+		TVector<double> V = corner_pos(Idx);
 		val = (float)Implicit(V);
 		corner_values_grid[Idx] = val;
 		return (double)val;
@@ -492,7 +494,7 @@ protected:
 
 	double corner_value_nohash(const FVector3i& Idx) 
 	{
-		FVector3<double> V = corner_pos(Idx);
+		TVector<double> V = corner_pos(Idx);
 		return Implicit(V);
 	}
 	void initialize_cell_values_nohash(FGridCell& Cell, bool Shift)
@@ -657,7 +659,7 @@ protected:
 	/**
 	*  fully sequential version, no threading
 	*/
-	void generate_continuation(TArrayView<const FVector3<double>> Seeds)
+	void generate_continuation(TArrayView<const FVector3d> Seeds)
 	{
 		FGridCell Cell;
 		int vertTArray[12];
@@ -666,7 +668,7 @@ protected:
 
 		TArray<FVector3i> stack;
 
-		for (FVector3<double> seed : Seeds)
+		for (FVector3d seed : Seeds)
 		{
 			FVector3i seed_idx = cell_index(seed);
 			if (!done_cells.IsValidIndex(seed_idx) || done_cells[seed_idx] == 1)
@@ -708,13 +710,13 @@ protected:
 	/**
 	*  parallel seed evaluation
 	*/
-	void generate_continuation_parallel(TArrayView<const FVector3<double>> Seeds)
+	void generate_continuation_parallel(TArrayView<const FVector3d> Seeds)
 	{
 		parallel_mesh_access = true;
 
 		ParallelFor(Seeds.Num(), [this, &Seeds](int32 Index)
 		{
-			FVector3<double> Seed = Seeds[Index];
+			FVector3d Seed = Seeds[Index];
 			FVector3i seed_idx = cell_index(Seed);
 			if (!done_cells.IsValidIndex(seed_idx) || set_cell_if_not_done(seed_idx) == false)
 			{
@@ -821,7 +823,7 @@ protected:
 		// have a crossing on that edge. Look up the indices of this
 		// edge and find the intersection point along it
 		Shift = 1;
-		FVector3<double> pa = FVector3<double>::Zero(), pb = FVector3<double>::Zero();
+		TVector<double> pa = TVector<double>::Zero(), pb = TVector<double>::Zero();
 		for (int i = 0; i <= 11; i++)
 		{
 			if ((EdgeTable[cubeindex] & Shift) != 0)
@@ -878,7 +880,7 @@ protected:
 	/**
 	*  add vertex to mesh, with locking if we are computing in parallel
 	*/
-	int append_vertex(FVector3<double> V, int64 CellHash)
+	int append_vertex(TVector<double> V, int64 CellHash)
 	{
 		int SectionIndex = GetVertexSectionIndex(CellHash);
 		int32 NewIndex = VertexCounter++;
@@ -975,7 +977,7 @@ protected:
 	/**
 	*  root-find the intersection along edge from f(P1)=ValP1 to f(P2)=ValP2
 	*/
-	void find_iso(const FVector3<double>& P1, const FVector3<double>& P2, double ValP1, double ValP2, FVector3<double>& PIso)
+	void find_iso(const TVector<double>& P1, const TVector<double>& P2, double ValP1, double ValP2, TVector<double>& PIso)
 	{
 		// Ok, this is a bit hacky but seems to work? If both isovalues
 		// are the same, we just return the midpoint. If one is nearly zero, we can
@@ -1003,7 +1005,7 @@ protected:
 
 		// Note: if we don't maintain min/max order here, then numerical error means
 		//   that hashing on point x/y/z doesn't work
-		FVector3<double> a = P1, b = P2;
+		TVector<double> a = P1, b = P2;
 		double fa = ValP1, fb = ValP2;
 		if (ValP2 < ValP1)
 		{
diff --git a/Engine/Source/Runtime/GeometryCore/Public/HalfspaceTypes.h b/Engine/Source/Runtime/GeometryCore/Public/HalfspaceTypes.h
--- a/Engine/Source/Runtime/GeometryCore/Public/HalfspaceTypes.h
+++ b/Engine/Source/Runtime/GeometryCore/Public/HalfspaceTypes.h
@@ -10,6 +10,8 @@ namespace UE
 namespace Geometry
 {
 
+using namespace UE::Math;
+
 /*
  * 3D Halfspace stored as parameters to Plane Equation (Normal, Normal.Dot(PointOnPlane))
  * The Normal points "into" the halfspace, ie X is inside if (Normal.Dot(X) - Constant) >= 0
@@ -19,13 +21,13 @@ struct THalfspace3
 {
 public:
 	/** Normal vector of 3D plane that defines Halfspace */
-	FVector3<T> Normal = FVector3<T>::UnitY();
+	TVector<T> Normal = TVector<T>::UnitY();
 	/** Distance along Normal that defines position of Halfspace */
 	T Constant = (T)0;
 
 	THalfspace3() = default;
 
-	THalfspace3(const FVector3<T>& PlaneNormalIn, T ConstantIn)
+	THalfspace3(const TVector<T>& PlaneNormalIn, T ConstantIn)
 		: Normal(PlaneNormalIn), Constant(ConstantIn) {}
 
 	THalfspace3(T NormalX, T NormalY, T NormalZ, T ConstantIn)
@@ -33,11 +35,11 @@ public:
 
 
 	/** Construct a Halfspace from the plane Normal and a Point lying on the plane */
-	THalfspace3(const FVector3<T>& PlaneNormalIn, const FVector3<T>& PlanePoint)
+	THalfspace3(const TVector<T>& PlaneNormalIn, const TVector<T>& PlanePoint)
 		: Normal(PlaneNormalIn), Constant(Normal.Dot(PlanePoint)) {}
 
 	/** Construct a Halfspace from three points */
-	THalfspace3(const FVector3<T>& P0, const FVector3<T>& P1, const FVector3<T>& P2)
+	THalfspace3(const TVector<T>& P0, const TVector<T>& P1, const TVector<T>& P2)
 	{
 		Normal = VectorUtil::Normal(P0, P1, P2);
 		Constant = Normal.Dot(P0);
@@ -45,7 +47,7 @@ public:
 
 
 	/** @return true if Halfspace contains given Point */
-	bool Contains(const FVector3<T>& Point) const
+	bool Contains(const TVector<T>& Point) const
 	{
 		return (Normal.Dot(Point) - Constant) >= 0;
 	}
diff --git a/Engine/Source/Runtime/GeometryCore/Public/Implicit/Blend.h b/Engine/Source/Runtime/GeometryCore/Public/Implicit/Blend.h
--- a/Engine/Source/Runtime/GeometryCore/Public/Implicit/Blend.h
+++ b/Engine/Source/Runtime/GeometryCore/Public/Implicit/Blend.h
@@ -17,6 +17,7 @@ namespace UE
 namespace Geometry
 {
 
+using namespace UE::Math;
 
 template<typename TriangleMeshType>
 class TImplicitBlend 
@@ -215,7 +216,7 @@ protected:
 			return;
 		}
 
-		TArray<FVector3<double>> Seeds;
+		TArray<FVector3d> Seeds;
 		for (const TriangleMeshType* Source : Sources)
 		{
 			for (int VID = 0; VID < Source->MaxVertexID(); VID++)
diff --git a/Engine/Source/Runtime/GeometryCore/Public/Implicit/ImplicitFunctions.h b/Engine/Source/Runtime/GeometryCore/Public/Implicit/ImplicitFunctions.h
--- a/Engine/Source/Runtime/GeometryCore/Public/Implicit/ImplicitFunctions.h
+++ b/Engine/Source/Runtime/GeometryCore/Public/Implicit/ImplicitFunctions.h
@@ -13,16 +13,18 @@ namespace UE
 namespace Geometry
 {
 
+using namespace UE::Math;
+
 template<typename RealType>
 struct TImplicitFunction3
 {
-	TUniqueFunction<RealType(FVector3<RealType>)> Value();
+	TUniqueFunction<RealType(TVector<RealType>)> Value();
 };
 
 template<typename RealType>
 struct TBoundedImplicitFunction3
 {
-	TUniqueFunction<RealType(FVector3<RealType>)> Value();
+	TUniqueFunction<RealType(TVector<RealType>)> Value();
 	TUniqueFunction<TAxisAlignedBox3<RealType>> Bounds();
 };
 
@@ -34,10 +36,10 @@ struct TBoundedImplicitFunction3
 template<typename RealType>
 struct TImplicitPoint3
 {
-	FVector3<RealType> Position;
+	TVector<RealType> Position;
 	RealType Radius;
 
-	RealType Value(const FVector3<RealType>& Point) const
+	RealType Value(const TVector<RealType>& Point) const
 	{
 		return Position.Distance(Point) - Radius;
 	}
@@ -61,7 +63,7 @@ struct TImplicitLine3
 	TSegment3<RealType> Segment;
 	RealType Radius;
 
-	RealType Value(const FVector3<RealType>& Point) const
+	RealType Value(const TVector<RealType>& Point) const
 	{
 		RealType DistanceSqr = Segment.DistanceSquared(Point);
 		return TMathUtil<RealType>::Sqrt(DistanceSqr) - Radius;
@@ -101,7 +103,7 @@ struct TSkeletalImplicitLine3
 		return Scale * DefaultIsoValue;
 	}
 
-	RealType Value(const FVector3<RealType>& Point) const
+	RealType Value(const TVector<RealType>& Point) const
 	{
 		RealType DistanceSqr = Segment.DistanceSquared(Point);
 		DistanceSqr /= (Scale * Scale);
@@ -149,7 +151,7 @@ struct TDistanceFieldToSkeletalField
 		return B;
 	}
 
-	RealType Value(const FVector3<RealType> Pt)
+	RealType Value(const TVector<RealType> Pt)
 	{
 		checkSlow(DistanceField != nullptr);
 		RealType Dist = DistanceField->Value(Pt);
@@ -177,7 +179,7 @@ struct TSkeletalRicciNaryBlend3
 	RealType BlendPower = 2.0;
 	bool bSubtract = false;
 
-	RealType Value(const FVector3<RealType> Pt)
+	RealType Value(const TVector<RealType> Pt)
 	{
 		int N = Children.Num();
 		checkSlow(N > 0);
diff --git a/Engine/Source/Runtime/GeometryCore/Public/Implicit/Solidify.h b/Engine/Source/Runtime/GeometryCore/Public/Implicit/Solidify.h
--- a/Engine/Source/Runtime/GeometryCore/Public/Implicit/Solidify.h
+++ b/Engine/Source/Runtime/GeometryCore/Public/Implicit/Solidify.h
@@ -12,6 +12,7 @@ namespace UE
 namespace Geometry
 {
 
+using namespace UE::Math;
 
 /**
  * Use marching cubes to remesh a triangle mesh to a solid surface
@@ -130,7 +131,7 @@ public:
 			};
 		}
 
-		TArray<FVector3<double>> MCSeeds;
+		TArray<FVector3d> MCSeeds;
 		for ( int32 VertIdx : Source->VertexIndicesItr() )
 		{
 			MCSeeds.Add(Source->GetVertex(VertIdx));
@@ -271,7 +272,7 @@ public:
 			};
 		}
 
-		TArray<FVector3<double>> MCSeeds;
+		TArray<TVector<double>> MCSeeds;
 		for ( const FVector3d& SeedPoint : SeedPoints )
 		{
 			MCSeeds.Add(SeedPoint);
diff --git a/Engine/Source/Runtime/GeometryCore/Public/Intersection/ContainmentQueries3.h b/Engine/Source/Runtime/GeometryCore/Public/Intersection/ContainmentQueries3.h
--- a/Engine/Source/Runtime/GeometryCore/Public/Intersection/ContainmentQueries3.h
+++ b/Engine/Source/Runtime/GeometryCore/Public/Intersection/ContainmentQueries3.h
@@ -35,7 +35,7 @@ namespace UE
 		template<typename RealType, typename EnumerablePointsType, typename E = decltype(DeclVal<EnumerablePointsType>().begin())>
 		bool IsInside(const TSphere3<RealType>& OuterSphere, EnumerablePointsType EnumerablePts)
 		{
-			for (FVector3<RealType> Point : EnumerablePts)
+			for (TVector<RealType> Point : EnumerablePts)
 			{
 				if (OuterSphere.Contains(Point) == false)
 				{
@@ -67,7 +67,7 @@ namespace UE
 		template<typename RealType, typename EnumerablePointsType, typename E = decltype(DeclVal<EnumerablePointsType>().begin())>
 		bool IsInside(const TCapsule3<RealType>& OuterCapsule, EnumerablePointsType EnumerablePts)
 		{
-			for (FVector3<RealType> Point : EnumerablePts)
+			for (TVector<RealType> Point : EnumerablePts)
 			{
 				if (OuterCapsule.Contains(Point) == false)
 				{
@@ -99,7 +99,7 @@ namespace UE
 		template<typename RealType, typename EnumerablePointsType, typename E = decltype(DeclVal<EnumerablePointsType>().begin())>
 		bool IsInside(const TOrientedBox3<RealType>& OuterBox, EnumerablePointsType EnumerablePts)
 		{
-			for (FVector3<RealType> Point : EnumerablePts)
+			for (TVector<RealType> Point : EnumerablePts)
 			{
 				if (OuterBox.Contains(Point) == false)
 				{
@@ -200,7 +200,7 @@ bool UE::Geometry::IsInsideHull(TArrayView<THalfspace3<RealType>> Halfspaces, En
 {
 	for (const THalfspace3<RealType>& Halfspace : Halfspaces)
 	{
-		for (FVector3<RealType> Point : EnumerablePts)
+		for (TVector<RealType> Point : EnumerablePts)
 		{
 			if (Halfspace.Contains(Point))
 			{
diff --git a/Engine/Source/Runtime/GeometryCore/Public/Intersection/IntersectionUtil.h b/Engine/Source/Runtime/GeometryCore/Public/Intersection/IntersectionUtil.h
--- a/Engine/Source/Runtime/GeometryCore/Public/Intersection/IntersectionUtil.h
+++ b/Engine/Source/Runtime/GeometryCore/Public/Intersection/IntersectionUtil.h
@@ -31,18 +31,19 @@ struct FLinearIntersection
  */
 namespace IntersectionUtil
 {
+	using namespace UE::Math;
 	using namespace UE::Geometry;
 
 	template<typename RealType>
-	bool RayTriangleTest(const UE::Math::TVector<RealType>& RayOrigin, const UE::Math::TVector<RealType>& RayDirection, const UE::Math::TVector<RealType>& V0, const UE::Math::TVector<RealType>& V1, const UE::Math::TVector<RealType>& V2)
+	bool RayTriangleTest(const TVector<RealType>& RayOrigin, const TVector<RealType>& RayDirection, const TVector<RealType>& V0, const TVector<RealType>& V1, const TVector<RealType>& V2)
 	{
 		// same code as IntrRay3Triangle3, but can be called w/o constructing additional data structures
 			
 		// Compute the offset origin, edges, and normal.
-		FVector3<RealType> diff = RayOrigin - V0;
-		FVector3<RealType> edge1 = V1 - V0;
-		FVector3<RealType> edge2 = V2 - V0;
-		FVector3<RealType> normal = edge1.Cross(edge2);
+		TVector<RealType> diff = RayOrigin - V0;
+		TVector<RealType> edge1 = V1 - V0;
+		TVector<RealType> edge2 = V2 - V0;
+		TVector<RealType> normal = edge1.Cross(edge2);
 
 		// Solve Q + t*D = b1*E1 + b2*E2 (Q = kDiff, D = ray direction,
 		// E1 = kEdge1, E2 = kEdge2, N = Cross(E1,E2)) by
@@ -101,9 +102,9 @@ namespace IntersectionUtil
 	 */
 	template<typename RealType>
 	bool LineSphereTest(
-		const UE::Math::TVector<RealType>& LineOrigin,
-		const UE::Math::TVector<RealType>& LineDirection,
-		const UE::Math::TVector<RealType>& SphereCenter,
+		const TVector<RealType>& LineOrigin,
+		const TVector<RealType>& LineDirection,
+		const TVector<RealType>& SphereCenter,
 		RealType SphereRadius)
 	{
 		// adapted from GeometricTools GTEngine
@@ -114,7 +115,7 @@ namespace IntersectionUtil
 		// quadratic equation Q(t) = t^2 + 2*a1*t + a0 = 0, where a1 = D^T*(P-C),
 		// and a0 = (P-C)^T*(P-C)-1.
 
-		FVector3<RealType> diff = LineOrigin - SphereCenter;
+		TVector<RealType> diff = LineOrigin - SphereCenter;
 		RealType a0 = diff.SquaredLength() - SphereRadius * SphereRadius;
 		RealType a1 = LineDirection.Dot(diff);
 
@@ -130,9 +131,9 @@ namespace IntersectionUtil
 	 */
 	template<typename RealType>
 	bool LineSphereIntersection(
-		const UE::Math::TVector<RealType>& LineOrigin,
-		const UE::Math::TVector<RealType>& LineDirection,
-		const UE::Math::TVector<RealType>& SphereCenter,
+		const TVector<RealType>& LineOrigin,
+		const TVector<RealType>& LineDirection,
+		const TVector<RealType>& SphereCenter,
 		RealType SphereRadius,
 		FLinearIntersection& ResultOut)
 	{
@@ -143,7 +144,7 @@ namespace IntersectionUtil
 		// Substitute the line equation into the sphere equation to obtain a
 		// quadratic equation Q(t) = t^2 + 2*a1*t + a0 = 0, where a1 = D^T*(P-C),
 		// and a0 = (P-C)^T*(P-C)-1.
-		FVector3<RealType> diff = LineOrigin - SphereCenter;
+		TVector<RealType> diff = LineOrigin - SphereCenter;
 		RealType a0 = diff.SquaredLength() - SphereRadius * SphereRadius;
 		RealType a1 = LineDirection.Dot(diff);
 
@@ -174,9 +175,9 @@ namespace IntersectionUtil
 
 	template<typename RealType>
 	FLinearIntersection LineSphereIntersection(
-		const UE::Math::TVector<RealType>& LineOrigin,
-		const UE::Math::TVector<RealType>& LineDirection,
-		const UE::Math::TVector<RealType>& SphereCenter,
+		const TVector<RealType>& LineOrigin,
+		const TVector<RealType>& LineDirection,
+		const TVector<RealType>& SphereCenter,
 		RealType SphereRadius)
 	{
 		FLinearIntersection result;
@@ -191,9 +192,9 @@ namespace IntersectionUtil
 	 */
 	template<typename RealType>
 	bool RaySphereTest(
-		const UE::Math::TVector<RealType>& RayOrigin,
-		const UE::Math::TVector<RealType>& RayDirection,
-		const UE::Math::TVector<RealType>& SphereCenter,
+		const TVector<RealType>& RayOrigin,
+		const TVector<RealType>& RayDirection,
+		const TVector<RealType>& SphereCenter,
 		RealType SphereRadius)
 	{
 		// adapted from GeometricTools GTEngine
@@ -204,7 +205,7 @@ namespace IntersectionUtil
 		// quadratic equation Q(t) = t^2 + 2*a1*t + a0 = 0, where a1 = D^T*(P-C),
 		// and a0 = (P-C)^T*(P-C)-1.
 
-		FVector3<RealType> diff = RayOrigin - SphereCenter;
+		TVector<RealType> diff = RayOrigin - SphereCenter;
 		RealType a0 = diff.SquaredLength() - SphereRadius * SphereRadius;
 		if (a0 <= 0)
 		{
@@ -229,9 +230,9 @@ namespace IntersectionUtil
 	 */
 	template<typename RealType>
 	bool RaySphereIntersection(
-		const UE::Math::TVector<RealType>& RayOrigin,
-		const UE::Math::TVector<RealType>& RayDirection,
-		const UE::Math::TVector<RealType>& SphereCenter,
+		const TVector<RealType>& RayOrigin,
+		const TVector<RealType>& RayDirection,
+		const TVector<RealType>& SphereCenter,
 		RealType SphereRadius, 
 		FLinearIntersection& Result)
 	{
@@ -260,9 +261,9 @@ namespace IntersectionUtil
 
 	template<typename RealType>
 	FLinearIntersection RaySphereIntersection(
-		const UE::Math::TVector<RealType>& RayOrigin,
-		const UE::Math::TVector<RealType>& RayDirection,
-		const UE::Math::TVector<RealType>& SphereCenter,
+		const TVector<RealType>& RayOrigin,
+		const TVector<RealType>& RayDirection,
+		const TVector<RealType>& SphereCenter,
 		RealType SphereRadius)
 	{
 		FLinearIntersection result;
diff --git a/Engine/Source/Runtime/GeometryCore/Public/Intersection/IntrLine2Triangle2.h b/Engine/Source/Runtime/GeometryCore/Public/Intersection/IntrLine2Triangle2.h
--- a/Engine/Source/Runtime/GeometryCore/Public/Intersection/IntrLine2Triangle2.h
+++ b/Engine/Source/Runtime/GeometryCore/Public/Intersection/IntrLine2Triangle2.h
@@ -95,7 +95,7 @@ public:
 			return false;
 		}
 
-		FVector3<Real> Dist;
+		TVector<Real> Dist;
 		FVector3i Sign;
 		int Positive = 0, Negative = 0, Zero = 0;
 		TriangleLineRelations(Line.Origin, Line.Direction, Triangle, Dist, Sign, Positive, Negative, Zero);
@@ -147,7 +147,7 @@ public:
 
 	static void TriangleLineRelations(
 		const TVector2<Real>& Origin, const TVector2<Real>& Direction, const TTriangle2<Real>& Tri,
-		FVector3<Real>& Dist, FVector3i& Sign, int& Positive, int& Negative, int& Zero, Real Tolerance = TMathUtil<Real>::ZeroTolerance)
+		TVector<Real>& Dist, FVector3i& Sign, int& Positive, int& Negative, int& Zero, Real Tolerance = TMathUtil<Real>::ZeroTolerance)
 	{
 		Positive = 0;
 		Negative = 0;
@@ -177,10 +177,10 @@ public:
 
 
 	static bool GetInterval(const TVector2<Real>& Origin, const TVector2<Real>& Direction, const TTriangle2<Real>& Tri,
-							const FVector3<Real>& Dist, const FVector3i& Sign, TVector2<Real>& param)
+							const TVector<Real>& Dist, const FVector3i& Sign, TVector2<Real>& param)
 	{
 		// Project Triangle onto Line.
-		FVector3<Real> proj;
+		TVector<Real> proj;
 		int i;
 		for (i = 0; i < 3; ++i)
 		{
diff --git a/Engine/Source/Runtime/GeometryCore/Public/Intersection/IntrRay3AxisAlignedBox3.h b/Engine/Source/Runtime/GeometryCore/Public/Intersection/IntrRay3AxisAlignedBox3.h
--- a/Engine/Source/Runtime/GeometryCore/Public/Intersection/IntrRay3AxisAlignedBox3.h
+++ b/Engine/Source/Runtime/GeometryCore/Public/Intersection/IntrRay3AxisAlignedBox3.h
@@ -13,6 +13,8 @@ namespace UE
 namespace Geometry
 {
 
+using namespace UE::Math;
+
 /**
  * Compute intersection between 3D ray and 3D axis-aligned box
  */
@@ -34,14 +36,14 @@ public:
 	 */
 	static bool TestIntersection(const TRay3<RealType>& Ray, const TAxisAlignedBox3<RealType>& Box, RealType ExpandExtents = 0)
 	{
-		FVector3<RealType> WdU = FVector3<RealType>::Zero();
-		FVector3<RealType> AWdU = FVector3<RealType>::Zero();
-		FVector3<RealType> DdU = FVector3<RealType>::Zero();
-		FVector3<RealType> ADdU = FVector3<RealType>::Zero();
+		TVector<RealType> WdU = TVector<RealType>::Zero();
+		TVector<RealType> AWdU = TVector<RealType>::Zero();
+		TVector<RealType> DdU = TVector<RealType>::Zero();
+		TVector<RealType> ADdU = TVector<RealType>::Zero();
 		RealType RHS;
 
-		FVector3<RealType> diff = Ray.Origin - Box.Center();
-		FVector3<RealType> extent = Box.Extents() + ExpandExtents;
+		TVector<RealType> diff = Ray.Origin - Box.Center();
+		TVector<RealType> extent = Box.Extents() + ExpandExtents;
 
 		WdU.X = Ray.Direction.X;
 		AWdU.X = FMath::Abs(WdU.X);
@@ -70,8 +72,8 @@ public:
 			return false;
 		}
 
-		FVector3<RealType> WxD = Ray.Direction.Cross(diff);
-		FVector3<RealType> AWxDdU = FVector3<RealType>::Zero();
+		TVector<RealType> WxD = Ray.Direction.Cross(diff);
+		TVector<RealType> AWxDdU = TVector<RealType>::Zero();
 
 		AWxDdU.X = FMath::Abs(WxD.X);
 		RHS = extent.Y * AWdU.Z + extent.Z * AWdU.Y;
@@ -111,8 +113,8 @@ public:
 		RealType RayParam0 = 0.0;
 		RealType RayParam1 = TNumericLimits<RealType>::Max();
 		int Quantity = 0;
-		FVector3<RealType> Point0 = FVector3<RealType>::Zero();
-		FVector3<RealType> Point1 = FVector3<RealType>::Zero();
+		TVector<RealType> Point0 = TVector<RealType>::Zero();
+		TVector<RealType> Point1 = TVector<RealType>::Zero();
 		EIntersectionType Type = EIntersectionType::Empty;
 		DoClipping(RayParam0, RayParam1, Ray.Origin, Ray.Direction, Box,
 			true, Quantity, Point0, Point1, Type);
@@ -137,12 +139,12 @@ protected:
 	// internal functions
 
 	static bool DoClipping(RealType t0, RealType t1,
-		const FVector3<RealType>& RayOrigin, const FVector3<RealType>& RayDirection,
+		const TVector<RealType>& RayOrigin, const TVector<RealType>& RayDirection,
 		const TAxisAlignedBox3<RealType>& Box, bool solid, 
-		int& quantity, FVector3<RealType>& Point0, FVector3<RealType>& Point1, EIntersectionType& intrType)
+		int& quantity, TVector<RealType>& Point0, TVector<RealType>& Point1, EIntersectionType& intrType)
 	{
-		FVector3<RealType> BOrigin = RayOrigin - Box.Center();
-		FVector3<RealType> extent = Box.Extents();
+		TVector<RealType> BOrigin = RayOrigin - Box.Center();
+		TVector<RealType> extent = Box.Extents();
 
 		RealType saveT0 = t0, saveT1 = t1;
 		bool notAllClipped =
diff --git a/Engine/Source/Runtime/GeometryCore/Public/Intersection/IntrRay3Triangle3.h b/Engine/Source/Runtime/GeometryCore/Public/Intersection/IntrRay3Triangle3.h
--- a/Engine/Source/Runtime/GeometryCore/Public/Intersection/IntrRay3Triangle3.h
+++ b/Engine/Source/Runtime/GeometryCore/Public/Intersection/IntrRay3Triangle3.h
@@ -14,6 +14,8 @@ namespace UE
 namespace Geometry
 {
 
+using namespace UE::Math;
+
 /**
  * Compute intersection between 3D ray and 3D triangle
  */
@@ -44,10 +46,10 @@ public:
 	bool Test()
 	{
 		// Compute the offset origin, edges, and normal.
-		FVector3<Real> diff = Ray.Origin - Triangle.V[0];
-		FVector3<Real> edge1 = Triangle.V[1] - Triangle.V[0];
-		FVector3<Real> edge2 = Triangle.V[2] - Triangle.V[0];
-		FVector3<Real> normal = edge1.Cross(edge2);
+		TVector<Real> diff = Ray.Origin - Triangle.V[0];
+		TVector<Real> edge1 = Triangle.V[1] - Triangle.V[0];
+		TVector<Real> edge2 = Triangle.V[2] - Triangle.V[0];
+		TVector<Real> normal = edge1.Cross(edge2);
 
 		// Solve Q + t*D = b1*E1 + b2*E2 (Q = kDiff, D = ray direction,
 		// E1 = kEdge1, E2 = kEdge2, N = Cross(E1,E2)) by
@@ -109,10 +111,10 @@ public:
 	bool Find()
 	{
 		// Compute the offset origin, edges, and normal.
-		FVector3<Real> diff = Ray.Origin - Triangle.V[0];
-		FVector3<Real> edge1 = Triangle.V[1] - Triangle.V[0];
-		FVector3<Real> edge2 = Triangle.V[2] - Triangle.V[0];
-		FVector3<Real> normal = edge1.Cross(edge2);
+		TVector<Real> diff = Ray.Origin - Triangle.V[0];
+		TVector<Real> edge1 = Triangle.V[1] - Triangle.V[0];
+		TVector<Real> edge2 = Triangle.V[2] - Triangle.V[0];
+		TVector<Real> normal = edge1.Cross(edge2);
 
 		// Solve Q + t*D = b1*E1 + b2*E2 (Q = kDiff, D = ray direction,
 		// E1 = kEdge1, E2 = kEdge2, N = Cross(E1,E2)) by
diff --git a/Engine/Source/Runtime/GeometryCore/Public/Intersection/IntrSegment2Triangle2.h b/Engine/Source/Runtime/GeometryCore/Public/Intersection/IntrSegment2Triangle2.h
--- a/Engine/Source/Runtime/GeometryCore/Public/Intersection/IntrSegment2Triangle2.h
+++ b/Engine/Source/Runtime/GeometryCore/Public/Intersection/IntrSegment2Triangle2.h
@@ -165,7 +165,7 @@ public:
 		// FSegment2d Direction should always be normalized unless Extent is 0 (handled above)
 		checkSlow(IsNormalized(Segment.Direction));
 
-		FVector3<Real> dist;
+		TVector<Real> dist;
 		FVector3i sign;
 		int positive = 0, negative = 0, zero = 0;
 		TIntrLine2Triangle2<Real>::TriangleLineRelations(Segment.Center, Segment.Direction, Triangle,
diff --git a/Engine/Source/Runtime/GeometryCore/Public/Intersection/IntrTriangle3Triangle3.h b/Engine/Source/Runtime/GeometryCore/Public/Intersection/IntrTriangle3Triangle3.h
--- a/Engine/Source/Runtime/GeometryCore/Public/Intersection/IntrTriangle3Triangle3.h
+++ b/Engine/Source/Runtime/GeometryCore/Public/Intersection/IntrTriangle3Triangle3.h
@@ -51,7 +51,7 @@ public:
 	// intersection points (for point, line, polygon)
 	// only first Quantity elements are relevant
 	int Quantity = 0;
-	FVector3<Real> Points[6];
+	TVector<Real> Points[6];
 
 	TIntrTriangle3Triangle3()
 	{}
@@ -64,7 +64,7 @@ public:
 	/**
 	 * Store an externally-computed segment intersection result
 	 */
-	inline void SetResult(const FVector3<Real>& A, const FVector3<Real>& B)
+	inline void SetResult(const TVector<Real>& A, const TVector<Real>& B)
 	{
 		Result = EIntersectionResult::Intersects;
 		Type = EIntersectionType::Segment;
@@ -155,12 +155,12 @@ public:
 		// Get the plane of Triangle0.
 		TPlane3<Real> Plane0(Triangle0.V[0], Triangle0.V[1], Triangle0.V[2]);
 
-		if (Plane0.Normal == FVector3<Real>::Zero())
+		if (Plane0.Normal == TVector<Real>::Zero())
 		{
 			// This function was written assuming triangle 0 has a valid normal;
 			//  if it's degenerate, try swapping the triangles
 			TPlane3<Real> Plane1(Triangle1.V[0], Triangle1.V[1], Triangle1.V[2]);
-			if (Plane1.Normal != FVector3<Real>::Zero())
+			if (Plane1.Normal != TVector<Real>::Zero())
 			{
 				// Find the intersection with the triangles swapped
 				TIntrTriangle3Triangle3<Real> SwappedTrisIntr(Triangle1, Triangle0);
@@ -188,7 +188,7 @@ public:
 		// an epsilon-thick plane test.
 		int pos1, neg1, zero1;
 		FIndex3i sign1;
-		FVector3<Real> dist1;
+		TVector<Real> dist1;
 		TrianglePlaneRelations(Triangle1, Plane0, dist1, sign1, pos1, neg1, zero1, Tolerance);
 
 		if (pos1 == 3 || neg1 == 3)
@@ -240,7 +240,7 @@ public:
 		// line segment of intersection.  Then test for intersection between this
 		// segment and triangle 0.
 		Real t;
-		FVector3<Real> intr0, intr1;
+		TVector<Real> intr0, intr1;
 		if (zero1 == 0)
 		{
 			int iSign = (pos1 == 1 ? +1 : -1);
@@ -281,13 +281,13 @@ public:
 	bool Test()
 	{
 		// Get edge vectors for Triangle0.
-		FVector3<Real> E0[3];
+		TVector<Real> E0[3];
 		E0[0] = Triangle0.V[1] - Triangle0.V[0];
 		E0[1] = Triangle0.V[2] - Triangle0.V[1];
 		E0[2] = Triangle0.V[0] - Triangle0.V[2];
 
 		// Get normal vector of Triangle0.
-		FVector3<Real> N0 = UnitCross(E0[0], E0[1]);
+		TVector<Real> N0 = UnitCross(E0[0], E0[1]);
 
 		// Project Triangle1 onto normal line of Triangle0, test for separation.
 		Real N0dT0V0 = N0.Dot(Triangle0.V[0]);
@@ -299,19 +299,19 @@ public:
 		}
 
 		// Get edge vectors for Triangle1.
-		FVector3<Real> E1[3];
+		TVector<Real> E1[3];
 		E1[0] = Triangle1.V[1] - Triangle1.V[0];
 		E1[1] = Triangle1.V[2] - Triangle1.V[1];
 		E1[2] = Triangle1.V[0] - Triangle1.V[2];
 
 		// Get normal vector of Triangle1.
-		FVector3<Real> N1 = UnitCross(E1[0], E1[1]);
+		TVector<Real> N1 = UnitCross(E1[0], E1[1]);
 
-		FVector3<Real> dir;
+		TVector<Real> dir;
 		Real min0, max0;
 		int i0, i1;
 
-		FVector3<Real> N0xN1 = UnitCross(N0, N1);
+		TVector<Real> N0xN1 = UnitCross(N0, N1);
 		if (N0xN1.Dot(N0xN1) >= Tolerance)
 		{
 			// Triangles are not parallel.
@@ -386,13 +386,13 @@ public:
 	static bool Intersects(const TTriangle3<Real>& Triangle0, const TTriangle3<Real>& Triangle1, Real Tolerance = TMathUtil<Real>::ZeroTolerance)
 	{
 		// Get edge vectors for Triangle0.
-		FVector3<Real> E0[3];
+		TVector<Real> E0[3];
 		E0[0] = Triangle0.V[1] - Triangle0.V[0];
 		E0[1] = Triangle0.V[2] - Triangle0.V[1];
 		E0[2] = Triangle0.V[0] - Triangle0.V[2];
 
 		// Get normal vector of Triangle0.
-		FVector3<Real> N0 = UnitCross(E0[0], E0[1]);
+		TVector<Real> N0 = UnitCross(E0[0], E0[1]);
 
 		// Project Triangle1 onto normal line of Triangle0, test for separation.
 		Real N0dT0V0 = N0.Dot(Triangle0.V[0]);
@@ -403,19 +403,19 @@ public:
 		}
 
 		// Get edge vectors for Triangle1.
-		FVector3<Real> E1[3];
+		TVector<Real> E1[3];
 		E1[0] = Triangle1.V[1] - Triangle1.V[0];
 		E1[1] = Triangle1.V[2] - Triangle1.V[1];
 		E1[2] = Triangle1.V[0] - Triangle1.V[2];
 
 		// Get normal vector of Triangle1.
-		FVector3<Real> N1 = UnitCross(E1[0], E1[1]);
+		TVector<Real> N1 = UnitCross(E1[0], E1[1]);
 
-		FVector3<Real> dir;
+		TVector<Real> dir;
 		Real min0, max0;
 		int i0, i1;
 
-		FVector3<Real> N0xN1 = UnitCross(N0, N1);
+		TVector<Real> N0xN1 = UnitCross(N0, N1);
 		if (N0xN1.Dot(N0xN1) >= Tolerance) {
 			// Triangles are not parallel.
 
@@ -468,7 +468,7 @@ public:
 
 
 
-	static void ProjectOntoAxis(const TTriangle3<Real>& triangle, const FVector3<Real>& axis, Real& fmin, Real& fmax)
+	static void ProjectOntoAxis(const TTriangle3<Real>& triangle, const TVector<Real>& axis, Real& fmin, Real& fmax)
 	{
 		Real dot0 = axis.Dot(triangle.V[0]);
 		Real dot1 = axis.Dot(triangle.V[1]);
@@ -499,7 +499,7 @@ public:
 
 
 	static void TrianglePlaneRelations(const TTriangle3<Real>& triangle, const TPlane3<Real>& plane,
-									   FVector3<Real>& distance, FIndex3i& sign, int& positive, int& negative, int& zero,
+									   TVector<Real>& distance, FIndex3i& sign, int& positive, int& negative, int& zero,
 									   Real Tolerance)
 	{
 		// Compute the signed distances of triangle vertices to the plane.  Use
@@ -543,8 +543,8 @@ public:
 	 * @return Number of points representing the intersection result (0 for none, 1 for point, 2 for segment)
 	 */
 	static int IntersectTriangleWithCoplanarSegment(
-		const TPlane3<Real>& plane, const TTriangle3<Real>& triangle, const FVector3<Real>& end0, const FVector3<Real>& end1,
-		FVector3<Real>& OutA, FVector3<Real>& OutB, Real Tolerance)
+		const TPlane3<Real>& plane, const TTriangle3<Real>& triangle, const TVector<Real>& end0, const TVector<Real>& end1,
+		TVector<Real>& OutA, TVector<Real>& OutB, Real Tolerance)
 	{
 		// Compute the 2D representations of the triangle vertices and the
 		// segment endpoints relative to the plane of the triangle.  Then
@@ -630,7 +630,7 @@ public:
 			intr[0] = calc.Point0;
 		}
 
-		FVector3<Real>* OutPts[2]{ &OutA, &OutB };
+		TVector<Real>* OutPts[2]{ &OutA, &OutB };
 
 		// Unproject the segment of intersection.
 		if (maxNormal == 0)
@@ -641,7 +641,7 @@ public:
 				Real y = intr[i].X;
 				Real z = intr[i].Y;
 				Real x = invNX * (plane.Constant - plane.Normal.Y * y - plane.Normal.Z * z);
-				*OutPts[i] = FVector3<Real>(x, y, z);
+				*OutPts[i] = TVector<Real>(x, y, z);
 			}
 		}
 		else if (maxNormal == 1)
@@ -652,7 +652,7 @@ public:
 				Real x = intr[i].X;
 				Real z = intr[i].Y;
 				Real y = invNY * (plane.Constant - plane.Normal.X * x - plane.Normal.Z * z);
-				*OutPts[i] = FVector3<Real>(x, y, z);
+				*OutPts[i] = TVector<Real>(x, y, z);
 			}
 		}
 		else
@@ -663,7 +663,7 @@ public:
 				Real x = intr[i].X;
 				Real y = intr[i].Y;
 				Real z = invNZ * (plane.Constant - plane.Normal.X * x - plane.Normal.Y * y);
-				*OutPts[i] = FVector3<Real>(x, y, z);
+				*OutPts[i] = TVector<Real>(x, y, z);
 			}
 		}
 
@@ -674,7 +674,7 @@ public:
 protected:
 
 
-	bool ContainsPoint(const TTriangle3<Real>& triangle, const TPlane3<Real>& plane, const FVector3<Real>& point)
+	bool ContainsPoint(const TTriangle3<Real>& triangle, const TPlane3<Real>& plane, const TVector<Real>& point)
 	{
 		// Generate a coordinate system for the plane.  The incoming triangle has
 		// vertices <V0,V1,V2>.  The incoming plane has unit-length normal N.
@@ -683,15 +683,15 @@ protected:
 		// constructed so that {U0,U1,N} is an orthonormal set.  Any point Q
 		// in the plane may be written as Q = V0 + x0*U0 + x1*U1.  The coordinates
 		// are computed as x0 = Dot(U0,Q-V0) and x1 = Dot(U1,Q-V0).
-		FVector3<Real> U0, U1;
+		TVector<Real> U0, U1;
 		VectorUtil::MakePerpVectors(plane.Normal, U0, U1);
 
 		// Compute the planar coordinates for the points P, V1, and V2.  To
 		// simplify matters, the origin is subtracted from the points, in which
 		// case the planar coordinates are for P-V0, V1-V0, and V2-V0.
-		FVector3<Real> PmV0 = point - triangle.V[0];
-		FVector3<Real> V1mV0 = triangle.V[1] - triangle.V[0];
-		FVector3<Real> V2mV0 = triangle.V[2] - triangle.V[0];
+		TVector<Real> PmV0 = point - triangle.V[0];
+		TVector<Real> V1mV0 = triangle.V[1] - triangle.V[0];
+		TVector<Real> V2mV0 = triangle.V[2] - triangle.V[0];
 
 		// The planar representation of P-V0.
 		TVector2<Real> ProjP(U0.Dot(PmV0), U1.Dot(PmV0));
@@ -713,7 +713,7 @@ protected:
 	}
 
 
-	bool IntersectsSegment(const TPlane3<Real>& plane, const TTriangle3<Real>& triangle, const FVector3<Real>& end0, const FVector3<Real>& end1)
+	bool IntersectsSegment(const TPlane3<Real>& plane, const TTriangle3<Real>& triangle, const TVector<Real>& end0, const TVector<Real>& end1)
 	{
 		Quantity = IntersectTriangleWithCoplanarSegment(plane, triangle, end0, end1, Points[0], Points[1], Tolerance);
 		if (Quantity > 0)
@@ -820,7 +820,7 @@ protected:
 				Real y = intr.Points[i].X;
 				Real z = intr.Points[i].Y;
 				Real x = invNX * (plane.Constant - plane.Normal.Y * y - plane.Normal.Z * z);
-				Points[i] = FVector3<Real>(x, y, z);
+				Points[i] = TVector<Real>(x, y, z);
 			}
 		}
 		else if (maxNormal == 1)
@@ -831,7 +831,7 @@ protected:
 				Real x = intr.Points[i].X;
 				Real z = intr.Points[i].Y;
 				Real y = invNY * (plane.Constant - plane.Normal.X * x - plane.Normal.Z * z);
-				Points[i] = FVector3<Real>(x, y, z);
+				Points[i] = TVector<Real>(x, y, z);
 			}
 		}
 		else
@@ -842,7 +842,7 @@ protected:
 				Real x = intr.Points[i].X;
 				Real y = intr.Points[i].Y;
 				Real z = invNZ * (plane.Constant - plane.Normal.X * x - plane.Normal.Y * y);
-				Points[i] = FVector3<Real>(x, y, z);
+				Points[i] = TVector<Real>(x, y, z);
 			}
 		}
 
diff --git a/Engine/Source/Runtime/GeometryCore/Public/LineTypes.h b/Engine/Source/Runtime/GeometryCore/Public/LineTypes.h
--- a/Engine/Source/Runtime/GeometryCore/Public/LineTypes.h
+++ b/Engine/Source/Runtime/GeometryCore/Public/LineTypes.h
@@ -147,24 +147,24 @@ template<typename T>
 struct TLine3
 {
 	/** Origin / Center Point of Line */
-	FVector3<T> Origin;
+	TVector<T> Origin;
 
 	/** Direction of Line, Normalized */
-	FVector3<T> Direction;
+	TVector<T> Direction;
 
 	/**
 	 * Construct default line along X axis
 	 */
 	TLine3()
 	{
-		Origin = FVector3<T>::Zero();
-		Direction = FVector3<T>::UnitX();
+		Origin = TVector<T>::Zero();
+		Direction = TVector<T>::UnitX();
 	}
 
 	/**
 	 * Construct line with given Origin and Direction
 	 */
-	TLine3(const FVector3<T>& OriginIn, const FVector3<T>& DirectionIn)
+	TLine3(const TVector<T>& OriginIn, const TVector<T>& DirectionIn)
 		: Origin(OriginIn), Direction(DirectionIn)
 	{
 	}
@@ -173,7 +173,7 @@ struct TLine3
 	/**
 	 * @return line between two points
 	 */
-	static TLine3<T> FromPoints(const FVector3<T>& Point0, const FVector3<T>& Point1)
+	static TLine3<T> FromPoints(const TVector<T>& Point0, const TVector<T>& Point1)
 	{
 		return TLine3<T>(Point0, Normalized(Point1 - Point0) );
 	}
@@ -182,7 +182,7 @@ struct TLine3
 	/**
 	 * @return point on line at given line parameter value (distance along line from origin)
 	 */
-	inline FVector3<T> PointAt(T LineParameter) const
+	inline TVector<T> PointAt(T LineParameter) const
 	{
 		return Origin + LineParameter * Direction;
 	}
@@ -191,7 +191,7 @@ struct TLine3
 	/**
 	 * @return line parameter (ie distance from Origin) at nearest point on line to QueryPoint
 	 */
-	inline T Project(const FVector3<T>& QueryPoint) const
+	inline T Project(const TVector<T>& QueryPoint) const
 	{
 		return (QueryPoint - Origin).Dot(Direction);
 	}
@@ -199,17 +199,17 @@ struct TLine3
 	/**
 	 * @return smallest squared distance from line to QueryPoint
 	 */
-	inline T DistanceSquared(const FVector3<T>& QueryPoint) const
+	inline T DistanceSquared(const TVector<T>& QueryPoint) const
 	{
 		T t = (QueryPoint - Origin).Dot(Direction);
-		FVector3<T> proj = Origin + t * Direction;
+		TVector<T> proj = Origin + t * Direction;
 		return (proj - QueryPoint).SquaredLength();
 	}
 
 	/**
 	 * @return nearest point on line to QueryPoint
 	 */
-	inline FVector3<T> NearestPoint(const FVector3<T>& QueryPoint) const
+	inline TVector<T> NearestPoint(const TVector<T>& QueryPoint) const
 	{
 		T ParameterT = (QueryPoint - Origin).Dot(Direction);
 		return Origin + ParameterT * Direction;
diff --git a/Engine/Source/Runtime/GeometryCore/Public/MatrixTypes.h b/Engine/Source/Runtime/GeometryCore/Public/MatrixTypes.h
--- a/Engine/Source/Runtime/GeometryCore/Public/MatrixTypes.h
+++ b/Engine/Source/Runtime/GeometryCore/Public/MatrixTypes.h
@@ -15,9 +15,9 @@ using namespace UE::Math;
 template <typename RealType>
 struct TMatrix3
 {
-	FVector3<RealType> Row0;
-	FVector3<RealType> Row1;
-	FVector3<RealType> Row2;
+	TVector<RealType> Row0;
+	TVector<RealType> Row1;
+	TVector<RealType> Row2;
 
 	TMatrix3()
 	{
@@ -25,16 +25,16 @@ struct TMatrix3
 
 	TMatrix3(RealType ConstantValue)
 	{
-		Row0 = FVector3<RealType>(ConstantValue, ConstantValue, ConstantValue);
+		Row0 = TVector<RealType>(ConstantValue, ConstantValue, ConstantValue);
 		Row1 = Row0;
 		Row2 = Row0;
 	}
 
 	TMatrix3(RealType Diag0, RealType Diag1, RealType Diag2)
 	{
-		Row0 = FVector3<RealType>(Diag0, 0, 0);
-		Row1 = FVector3<RealType>(0, Diag1, 0);
-		Row2 = FVector3<RealType>(0, 0, Diag2);
+		Row0 = TVector<RealType>(Diag0, 0, 0);
+		Row1 = TVector<RealType>(0, Diag1, 0);
+		Row2 = TVector<RealType>(0, 0, Diag2);
 	}
 
 	/**
@@ -65,17 +65,17 @@ struct TMatrix3
 		}
 		else
 		{
-			Row0 = FVector3<RealType>(V1.X, V2.X, V3.X);
-			Row1 = FVector3<RealType>(V1.Y, V2.Y, V3.Y);
-			Row2 = FVector3<RealType>(V1.Z, V2.Z, V3.Z);
+			Row0 = TVector<RealType>(V1.X, V2.X, V3.X);
+			Row1 = TVector<RealType>(V1.Y, V2.Y, V3.Y);
+			Row2 = TVector<RealType>(V1.Z, V2.Z, V3.Z);
 		}
 	}
 
 	template<typename RealType2>
 	explicit constexpr TMatrix3(const TMatrix3<RealType2>& Mat) :
-		Row0(FVector3<RealType>{Mat.Row0}),
-		Row1(FVector3<RealType>{Mat.Row1}),
-		Row2(FVector3<RealType>{Mat.Row2})
+		Row0(TVector<RealType>{Mat.Row0}),
+		Row1(TVector<RealType>{Mat.Row1}),
+		Row2(TVector<RealType>{Mat.Row2})
 	{
 	}
 
@@ -113,9 +113,9 @@ struct TMatrix3
 			Row2.X * Scale, Row2.Y * Scale, Row2.Z * Scale);
 	}
 
-	FVector3<RealType> operator*(const UE::Math::TVector<RealType>& V) const
+	TVector<RealType> operator*(const UE::Math::TVector<RealType>& V) const
 	{
-		return FVector3<RealType>(
+		return TVector<RealType>(
 			Row0.X * V.X + Row0.Y * V.Y + Row0.Z * V.Z,
 			Row1.X * V.X + Row1.Y * V.Y + Row1.Z * V.Z,
 			Row2.X * V.X + Row2.Y * V.Y + Row2.Z * V.Z);
diff --git a/Engine/Source/Runtime/GeometryCore/Public/MeshAdapter.h b/Engine/Source/Runtime/GeometryCore/Public/MeshAdapter.h
--- a/Engine/Source/Runtime/GeometryCore/Public/MeshAdapter.h
+++ b/Engine/Source/Runtime/GeometryCore/Public/MeshAdapter.h
@@ -14,6 +14,8 @@ namespace UE
 namespace Geometry
 {
 
+using namespace UE::Math;
+
 /**
  * Most generic / lazy example of a triangle mesh adapter; possibly useful for prototyping / building on top of (but slower than making a more specific-case adapter)
  */
@@ -28,7 +30,7 @@ struct TTriangleMeshAdapter
 	TFunction<int32()> VertexCount;
 	TFunction<uint64()> GetChangeStamp;
 	TFunction<FIndex3i(int32)> GetTriangle;
-	TFunction<FVector3<RealType>(int32)> GetVertex;
+	TFunction<TVector<RealType>(int32)> GetVertex;
 
 	inline void GetTriVertices(int TID, UE::Math::TVector<RealType>& V0, UE::Math::TVector<RealType>& V1, UE::Math::TVector<RealType>& V2) const
 	{
@@ -127,17 +129,17 @@ struct TIndexMeshArrayAdapter
 		return FIndex3i((int)(*SourceTriangles)[Start], (int)(*SourceTriangles)[Start+1], (int)(*SourceTriangles)[Start+2]);
 	}
 
-	FORCEINLINE FVector3<OutRealType> GetVertex(int32 Index) const
+	FORCEINLINE TVector<OutRealType> GetVertex(int32 Index) const
 	{
-		return FVector3<OutRealType>((*SourceVertices)[Index]);
+		return TVector<OutRealType>((*SourceVertices)[Index]);
 	}
 
 	FORCEINLINE void GetTriVertices(int32 TriIndex, UE::Math::TVector<OutRealType>& V0, UE::Math::TVector<OutRealType>& V1, UE::Math::TVector<OutRealType>& V2) const
 	{
 		int32 Start = TriIndex * 3;
-		V0 = FVector3<OutRealType>((*SourceVertices)[(*SourceTriangles)[Start]]);
-		V1 = FVector3<OutRealType>((*SourceVertices)[(*SourceTriangles)[Start+1]]);
-		V2 = FVector3<OutRealType>((*SourceVertices)[(*SourceTriangles)[Start+2]]);
+		V0 = TVector<OutRealType>((*SourceVertices)[(*SourceTriangles)[Start]]);
+		V1 = TVector<OutRealType>((*SourceVertices)[(*SourceTriangles)[Start+1]]);
+		V2 = TVector<OutRealType>((*SourceVertices)[(*SourceTriangles)[Start+2]]);
 	}
 
 };
@@ -209,17 +211,17 @@ struct TIndexVectorMeshArrayAdapter
 		return FIndex3i((int)Tri[0], (int)Tri[1], (int)Tri[2]);
 	}
 
-	FORCEINLINE FVector3<OutRealType> GetVertex(int32 Index) const
+	FORCEINLINE TVector<OutRealType> GetVertex(int32 Index) const
 	{
-		return FVector3<OutRealType>((*SourceVertices)[Index]);
+		return TVector<OutRealType>((*SourceVertices)[Index]);
 	}
 
 	FORCEINLINE void GetTriVertices(int32 TriIndex, UE::Math::TVector<OutRealType>& V0, UE::Math::TVector<OutRealType>& V1, UE::Math::TVector<OutRealType>& V2) const
 	{
 		const IndexVectorType& Tri = (*SourceTriangles)[TriIndex];
-		V0 = FVector3<OutRealType>((*SourceVertices)[Tri[0]]);
-		V1 = FVector3<OutRealType>((*SourceVertices)[Tri[1]]);
-		V2 = FVector3<OutRealType>((*SourceVertices)[Tri[2]]);
+		V0 = TVector<OutRealType>((*SourceVertices)[Tri[0]]);
+		V1 = TVector<OutRealType>((*SourceVertices)[Tri[1]]);
+		V2 = TVector<OutRealType>((*SourceVertices)[Tri[2]]);
 	}
 
 };
diff --git a/Engine/Source/Runtime/GeometryCore/Public/OrientedBoxTypes.h b/Engine/Source/Runtime/GeometryCore/Public/OrientedBoxTypes.h
--- a/Engine/Source/Runtime/GeometryCore/Public/OrientedBoxTypes.h
+++ b/Engine/Source/Runtime/GeometryCore/Public/OrientedBoxTypes.h
@@ -13,6 +13,7 @@ namespace UE
 namespace Geometry
 {
 
+using namespace UE::Math;
 
 /**
  * TOrientedBox3 is a non-axis-aligned 3D box defined by a 3D frame and extents along the axes of that frame
@@ -27,14 +28,14 @@ struct TOrientedBox3
 	/** 3D position (center) and orientation (axes) of the box */
 	TFrame3<RealType> Frame;
 	/** Half-dimensions of box measured along the three axes */
-	FVector3<RealType> Extents;
+	TVector<RealType> Extents;
 
 	TOrientedBox3() : Extents(1,1,1) {}
 
 	/**
 	 * Create axis-aligned box with given Origin and Extents
 	 */
-	TOrientedBox3(const FVector3<RealType>& Origin, const FVector3<RealType> & ExtentsIn)
+	TOrientedBox3(const TVector<RealType>& Origin, const TVector<RealType> & ExtentsIn)
 		: Frame(Origin), Extents(ExtentsIn)
 	{
 	}
@@ -42,7 +43,7 @@ struct TOrientedBox3
 	/**
 	 * Create oriented box with given Frame and Extents
 	 */
-	TOrientedBox3(const TFrame3<RealType>& FrameIn, const FVector3<RealType> & ExtentsIn)
+	TOrientedBox3(const TFrame3<RealType>& FrameIn, const TVector<RealType> & ExtentsIn)
 		: Frame(FrameIn), Extents(ExtentsIn)
 	{
 	}
@@ -57,26 +58,26 @@ struct TOrientedBox3
 
 
 	/** @return box with unit dimensions centered at origin */
-	static TOrientedBox3<RealType> UnitZeroCentered() { return TOrientedBox3<RealType>(FVector3<RealType>::Zero(), (RealType)0.5*FVector3<RealType>::One()); }
+	static TOrientedBox3<RealType> UnitZeroCentered() { return TOrientedBox3<RealType>(TVector<RealType>::Zero(), (RealType)0.5*TVector<RealType>::One()); }
 
 	/** @return box with unit dimensions where minimum corner is at origin */
-	static TOrientedBox3<RealType> UnitPositive() { return TOrientedBox3<RealType>((RealType)0.5*FVector3<RealType>::One(), (RealType)0.5*FVector3<RealType>::One()); }
+	static TOrientedBox3<RealType> UnitPositive() { return TOrientedBox3<RealType>((RealType)0.5*TVector<RealType>::One(), (RealType)0.5*TVector<RealType>::One()); }
 
 
 	/** @return center of the box */
-	FVector3<RealType> Center() const { return Frame.Origin; }
+	TVector<RealType> Center() const { return Frame.Origin; }
 
 	/** @return X axis of the box */
-	FVector3<RealType> AxisX() const { return Frame.X(); }
+	TVector<RealType> AxisX() const { return Frame.X(); }
 
 	/** @return Y axis of the box */
-	FVector3<RealType> AxisY() const { return Frame.Y(); }
+	TVector<RealType> AxisY() const { return Frame.Y(); }
 
 	/** @return Z axis of the box */
-	FVector3<RealType> AxisZ() const { return Frame.Z(); }
+	TVector<RealType> AxisZ() const { return Frame.Z(); }
 
 	/** @return an axis of the box */
-	FVector3<RealType> GetAxis(int AxisIndex) const { return Frame.GetAxis(AxisIndex); }
+	TVector<RealType> GetAxis(int AxisIndex) const { return Frame.GetAxis(AxisIndex); }
 
 	/** @return maximum extent of box */
 	inline RealType MaxExtent() const
@@ -91,7 +92,7 @@ struct TOrientedBox3
 	}
 
 	/** @return vector from minimum-corner to maximum-corner of box */
-	inline FVector3<RealType> Diagonal() const
+	inline TVector<RealType> Diagonal() const
 	{
 		return Frame.PointAt(Extents.X, Extents.Y, Extents.Z) - Frame.PointAt(-Extents.X, -Extents.Y, -Extents.Z);
 	}
@@ -103,9 +104,9 @@ struct TOrientedBox3
 	}
 
 	/** @return true if box contains point */
-	inline bool Contains(const FVector3<RealType>& Point) const
+	inline bool Contains(const TVector<RealType>& Point) const
 	{
-		FVector3<RealType> InFramePoint = Frame.ToFramePoint(Point);
+		TVector<RealType> InFramePoint = Frame.ToFramePoint(Point);
 		return (TMathUtil<RealType>::Abs(InFramePoint.X) <= Extents.X) &&
 			(TMathUtil<RealType>::Abs(InFramePoint.Y) <= Extents.Y) &&
 			(TMathUtil<RealType>::Abs(InFramePoint.Z) <= Extents.Z);
@@ -128,7 +129,7 @@ struct TOrientedBox3
 	 * @param Index corner index in range 0-7
 	 * @return Corner point on the box identified by the given index. See diagram in OrientedBoxTypes.h for index/corner mapping.
 	 */
-	FVector3<RealType> GetCorner(int Index) const
+	TVector<RealType> GetCorner(int Index) const
 	{
 		check(Index >= 0 && Index <= 7);
 		RealType dx = (((Index & 1) != 0) ^ ((Index & 2) != 0)) ? (Extents.X) : (-Extents.X);
@@ -146,14 +147,14 @@ struct TOrientedBox3
 	{
 		TMatrix3<RealType> RotMatrix = Frame.Rotation.ToRotationMatrix();
 		RealType X = Extents.X, Y = Extents.Y, Z = Extents.Z;
-		CornerPointFunc( RotMatrix*FVector3<RealType>(-X,-Y,-Z) + Frame.Origin );
-		CornerPointFunc( RotMatrix*FVector3<RealType>( X,-Y,-Z) + Frame.Origin );
-		CornerPointFunc( RotMatrix*FVector3<RealType>( X, Y,-Z) + Frame.Origin );
-		CornerPointFunc( RotMatrix*FVector3<RealType>(-X, Y,-Z) + Frame.Origin );
-		CornerPointFunc( RotMatrix*FVector3<RealType>(-X,-Y, Z) + Frame.Origin );
-		CornerPointFunc( RotMatrix*FVector3<RealType>( X,-Y, Z) + Frame.Origin );
-		CornerPointFunc( RotMatrix*FVector3<RealType>( X, Y, Z) + Frame.Origin );
-		CornerPointFunc( RotMatrix*FVector3<RealType>(-X, Y, Z) + Frame.Origin );
+		CornerPointFunc( RotMatrix*TVector<RealType>(-X,-Y,-Z) + Frame.Origin );
+		CornerPointFunc( RotMatrix*TVector<RealType>( X,-Y,-Z) + Frame.Origin );
+		CornerPointFunc( RotMatrix*TVector<RealType>( X, Y,-Z) + Frame.Origin );
+		CornerPointFunc( RotMatrix*TVector<RealType>(-X, Y,-Z) + Frame.Origin );
+		CornerPointFunc( RotMatrix*TVector<RealType>(-X,-Y, Z) + Frame.Origin );
+		CornerPointFunc( RotMatrix*TVector<RealType>( X,-Y, Z) + Frame.Origin );
+		CornerPointFunc( RotMatrix*TVector<RealType>( X, Y, Z) + Frame.Origin );
+		CornerPointFunc( RotMatrix*TVector<RealType>(-X, Y, Z) + Frame.Origin );
 	}
 
 
@@ -166,14 +167,14 @@ struct TOrientedBox3
 	{
 		TMatrix3<RealType> RotMatrix = Frame.Rotation.ToRotationMatrix();
 		RealType X = Extents.X, Y = Extents.Y, Z = Extents.Z;
-		return CornerPointPredicate(RotMatrix * FVector3<RealType>(-X, -Y, -Z) + Frame.Origin) &&
-			CornerPointPredicate(RotMatrix * FVector3<RealType>(X, -Y, -Z) + Frame.Origin) &&
-			CornerPointPredicate(RotMatrix * FVector3<RealType>(X, Y, -Z) + Frame.Origin) &&
-			CornerPointPredicate(RotMatrix * FVector3<RealType>(-X, Y, -Z) + Frame.Origin) &&
-			CornerPointPredicate(RotMatrix * FVector3<RealType>(-X, -Y, Z) + Frame.Origin) &&
-			CornerPointPredicate(RotMatrix * FVector3<RealType>(X, -Y, Z) + Frame.Origin) &&
-			CornerPointPredicate(RotMatrix * FVector3<RealType>(X, Y, Z) + Frame.Origin) &&
-			CornerPointPredicate(RotMatrix * FVector3<RealType>(-X, Y, Z) + Frame.Origin);
+		return CornerPointPredicate(RotMatrix * TVector<RealType>(-X, -Y, -Z) + Frame.Origin) &&
+			CornerPointPredicate(RotMatrix * TVector<RealType>(X, -Y, -Z) + Frame.Origin) &&
+			CornerPointPredicate(RotMatrix * TVector<RealType>(X, Y, -Z) + Frame.Origin) &&
+			CornerPointPredicate(RotMatrix * TVector<RealType>(-X, Y, -Z) + Frame.Origin) &&
+			CornerPointPredicate(RotMatrix * TVector<RealType>(-X, -Y, Z) + Frame.Origin) &&
+			CornerPointPredicate(RotMatrix * TVector<RealType>(X, -Y, Z) + Frame.Origin) &&
+			CornerPointPredicate(RotMatrix * TVector<RealType>(X, Y, Z) + Frame.Origin) &&
+			CornerPointPredicate(RotMatrix * TVector<RealType>(-X, Y, Z) + Frame.Origin);
 	}
 
 
@@ -200,7 +201,7 @@ struct TOrientedBox3
 	 * @param Point input point
 	 * @return squared distance from point to box, or 0 if point is inside box
 	 */
-	RealType DistanceSquared(FVector3<RealType> Point)
+	RealType DistanceSquared(TVector<RealType> Point)
 	{
 		 // Ported from WildMagic5 Wm5DistPoint3Box3.cpp
 			 
@@ -210,7 +211,7 @@ struct TOrientedBox3
 		// Compute squared distance and closest point on box.
 		RealType sqrDistance = 0;
 		RealType delta;
-		FVector3<RealType> closest(0, 0, 0);
+		TVector<RealType> closest(0, 0, 0);
 		for (int i = 0; i < 3; ++i) 
 		{
 			closest[i] = Point.Dot(GetAxis(i));
@@ -239,7 +240,7 @@ struct TOrientedBox3
 	  * @param Point input point
 	  * @return closest point on box. Input point is returned if it is inside box.
 	  */
-	 FVector3<RealType> ClosestPoint(FVector3<RealType> Point)
+	 TVector<RealType> ClosestPoint(TVector<RealType> Point)
 	 {
 		 // Work in the box's coordinate system.
 		 Point -= Frame.Origin;
@@ -247,8 +248,8 @@ struct TOrientedBox3
 		 // Compute squared distance and closest point on box.
 		 RealType sqrDistance = 0;
 		 RealType delta;
-		 FVector3<RealType> closest;
-		 FVector3<RealType> Axes[3];
+		 TVector<RealType> closest;
+		 TVector<RealType> Axes[3];
 		 for (int i = 0; i < 3; ++i) 
 		 {
 			 Axes[i] = GetAxis(i);
diff --git a/Engine/Source/Runtime/GeometryCore/Public/PlaneTypes.h b/Engine/Source/Runtime/GeometryCore/Public/PlaneTypes.h
--- a/Engine/Source/Runtime/GeometryCore/Public/PlaneTypes.h
+++ b/Engine/Source/Runtime/GeometryCore/Public/PlaneTypes.h
@@ -12,10 +12,12 @@ namespace UE
 namespace Geometry
 {
 
+using namespace UE::Math;
+
 template<typename RealType>
 struct TPlane3
 {
-	FVector3<RealType> Normal;
+	TVector<RealType> Normal;
 	RealType Constant;
 
 	TPlane3() {}
@@ -84,7 +86,7 @@ struct TPlane3
 	 * Compute intersection of Line with Plane
 	 * @param LineOrigin origin of line
 	 * @param LineDirection direction of line
-	 * @param HitPointOut intersection point, or FVector3::MaxVector() if line is parallel to plane
+	 * @param HitPointOut intersection point, or invalid point if line is parallel to plane
 	 * @return true if Line intersects Plane and IntersectionPointOut is valid
 	 */
 	bool FindLineIntersection(
@@ -95,7 +97,7 @@ struct TPlane3
 		RealType NormalDot = LineDirection.Dot(Normal);
 		if ( TMathUtil<RealType>::Abs(NormalDot) < TMathUtil<RealType>::ZeroTolerance )
 		{
-			IntersectionPointOut = FVector3<RealType>::MaxVector();
+			IntersectionPointOut = TVector<RealType>(TNumericLimits<RealType>::Max(), TNumericLimits<RealType>::Max(), TNumericLimits<RealType>::Max());
 			return false;
 		}
 		RealType t = -(LineOrigin.Dot(Normal) - Constant) / NormalDot;
@@ -125,8 +127,8 @@ struct TPlane3
 			return 2;
 		}
 
-		FVector3<RealType> DirectionVec = Point1 - Point0;
-		FVector3<RealType> Direction = Normalized(DirectionVec);
+		TVector<RealType> DirectionVec = Point1 - Point0;
+		TVector<RealType> Direction = Normalized(DirectionVec);
 		RealType Length = DirectionVec.Dot(Direction);
 
 		// test if segment is parallel to plane, if so, no intersection
diff --git a/Engine/Source/Runtime/GeometryCore/Public/PointSetAdapter.h b/Engine/Source/Runtime/GeometryCore/Public/PointSetAdapter.h
--- a/Engine/Source/Runtime/GeometryCore/Public/PointSetAdapter.h
+++ b/Engine/Source/Runtime/GeometryCore/Public/PointSetAdapter.h
@@ -9,6 +9,8 @@ namespace UE
 namespace Geometry
 {
 
+using namespace UE::Math;
+
 /**
  * TPointSetAdapter provides a very generic interface to an indexable list of points.
  * The list may be sparse, ie some indices may be invalid.
@@ -23,7 +25,7 @@ struct TPointSetAdapter
 	/** Returns true if this index valid */
 	TFunction<bool(int32)> IsPoint;
 	/** Get point at this index */
-	TFunction<FVector3<RealType>(int32)> GetPoint;
+	TFunction<TVector<RealType>(int32)> GetPoint;
 
 	/** Returns true if this point set has per-point normals */
 	TFunction<bool()> HasNormals;
diff --git a/Engine/Source/Runtime/GeometryCore/Public/Polyline3.h b/Engine/Source/Runtime/GeometryCore/Public/Polyline3.h
--- a/Engine/Source/Runtime/GeometryCore/Public/Polyline3.h
+++ b/Engine/Source/Runtime/GeometryCore/Public/Polyline3.h
@@ -13,6 +13,8 @@ namespace UE
 namespace Geometry
 {
 
+using namespace UE::Math;
+
 /**
  * TPolyline3 represents a 3D polyline stored as a list of Vertices.
  *
@@ -23,7 +25,7 @@ class TPolyline3
 {
 protected:
 	/** The list of vertices of the polyline */
-	TArray<FVector3<T>> Vertices;
+	TArray<TVector<T>> Vertices;
 
 	/** A counter that is incremented every time the polyline vertices are modified */
 	int Timestamp;
@@ -45,14 +47,14 @@ public:
 	/**
 	 * Construct polyline with given list of vertices
 	 */
-	TPolyline3(const TArray<FVector3<T>>& VertexList) : Vertices(VertexList), Timestamp(0)
+	TPolyline3(const TArray<TVector<T>>& VertexList) : Vertices(VertexList), Timestamp(0)
 	{
 	}
 
 	/**
 	 * Construct a single-segment polyline
 	 */
-	TPolyline3(const FVector3<T>& Point0, const FVector3<T>& Point1) : Timestamp(0)
+	TPolyline3(const TVector<T>& Point0, const TVector<T>& Point1) : Timestamp(0)
 	{
 		Vertices.Add(Point0);
 		Vertices.Add(Point1);
@@ -73,7 +75,7 @@ public:
 	/**
 	 * Get the vertex at a given index
 	 */
-	const FVector3<T>& operator[](int Index) const
+	const TVector<T>& operator[](int Index) const
 	{
 		return Vertices[Index];
 	}
@@ -82,7 +84,7 @@ public:
 	 * Get the vertex at a given index
 	 * @warning changing the vertex via this operator does not update Timestamp!
 	 */
-	FVector3<T>& operator[](int Index)
+	TVector<T>& operator[](int Index)
 	{
 		return Vertices[Index];
 	}
@@ -91,7 +93,7 @@ public:
 	/**
 	 * @return first vertex of polyline
 	 */
-	const FVector3<T>& Start() const
+	const TVector<T>& Start() const
 	{
 		return Vertices[0];
 	}
@@ -99,7 +101,7 @@ public:
 	/**
 	 * @return last vertex of polyline
 	 */
-	const FVector3<T>& End() const
+	const TVector<T>& End() const
 	{
 		return Vertices[Vertices.Num()-1];
 	}
@@ -108,7 +110,7 @@ public:
 	/**
 	 * @return list of Vertices of polyline
 	 */
-	const TArray<FVector3<T>>& GetVertices() const
+	const TArray<TVector<T>>& GetVertices() const
 	{
 		return Vertices;
 	}
@@ -140,7 +142,7 @@ public:
 	/**
 	 * Add a vertex to the polyline
 	 */
-	void AppendVertex(const FVector3<T>& Position)
+	void AppendVertex(const TVector<T>& Position)
 	{
 		Vertices.Add(Position);
 		Timestamp++;
@@ -149,7 +151,7 @@ public:
 	/**
 	 * Add a list of Vertices to the polyline
 	 */
-	void AppendVertices(const TArray<FVector3<T>>& NewVertices)
+	void AppendVertices(const TArray<TVector<T>>& NewVertices)
 	{
 		Vertices.Append(NewVertices);
 		Timestamp++;
@@ -165,7 +167,7 @@ public:
 		Vertices.Reserve(Vertices.Num() + NumV);
 		for (int32 k = 0; k < NumV; ++k)
 		{
-			Vertices.Append( (FVector3<T>)NewVertices[k] );
+			Vertices.Append( (TVector<T>)NewVertices[k] );
 		}
 		Timestamp++;
 	}
@@ -173,7 +175,7 @@ public:
 	/**
 	 * Set vertex at given index to a new Position
 	 */
-	void Set(int VertexIndex, const FVector3<T>& Position)
+	void Set(int VertexIndex, const TVector<T>& Position)
 	{
 		Vertices[VertexIndex] = Position;
 		Timestamp++;
@@ -191,7 +193,7 @@ public:
 	/**
 	 * Replace the list of Vertices with a new list
 	 */
-	void SetVertices(const TArray<FVector3<T>>& NewVertices)
+	void SetVertices(const TArray<TVector<T>>& NewVertices)
 	{
 		int NumVerts = NewVertices.Num();
 		Vertices.SetNum(NumVerts, false);
@@ -221,7 +223,7 @@ public:
 	 * Get the tangent vector at a vertex of the polyline, which is the normalized
 	 * vector from the previous vertex to the next vertex
 	 */
-	FVector3<T> GetTangent(int VertexIndex) const
+	TVector<T> GetTangent(int VertexIndex) const
 	{
 		if (VertexIndex == 0)
 		{
@@ -251,7 +253,7 @@ public:
 	 * @param SegmentParam parameter in range [-Extent,Extent] along segment
 	 * @return point on the segment at the given parameter value
 	 */
-	FVector3<T> GetSegmentPoint(int SegmentIndex, T SegmentParam) const
+	TVector<T> GetSegmentPoint(int SegmentIndex, T SegmentParam) const
 	{
 		TSegment3<T> seg(Vertices[SegmentIndex], Vertices[SegmentIndex + 1]);
 		return seg.PointAt(SegmentParam);
@@ -263,7 +265,7 @@ public:
 	 * @param SegmentParam parameter in range [0,1] along segment
 	 * @return point on the segment at the given parameter value
 	 */
-	FVector3<T> GetSegmentPointUnitParam(int SegmentIndex, T SegmentParam) const
+	TVector<T> GetSegmentPointUnitParam(int SegmentIndex, T SegmentParam) const
 	{
 		TSegment3<T> seg(Vertices[SegmentIndex], Vertices[SegmentIndex + 1]);
 		return seg.PointBetween(SegmentParam);
@@ -377,7 +379,7 @@ public:
 	 * @param NearestSegParamOut the parameter value of the nearest point on the segment
 	 * @return squared distance to the polyline
 	 */
-	T DistanceSquared(const FVector3<T>& QueryPoint, int& NearestSegIndexOut, T& NearestSegParamOut) const
+	T DistanceSquared(const TVector<T>& QueryPoint, int& NearestSegIndexOut, T& NearestSegParamOut) const
 	{
 		NearestSegIndexOut = -1;
 		NearestSegParamOut = TNumericLimits<T>::Max();
@@ -418,7 +420,7 @@ public:
 	 * @param QueryPoint the query point
 	 * @return squared distance to the polyline
 	 */
-	T DistanceSquared(const FVector3<T>& QueryPoint) const
+	T DistanceSquared(const TVector<T>& QueryPoint) const
 	{
 		int seg; T segt;
 		return DistanceSquared(QueryPoint, seg, segt);
@@ -455,9 +457,9 @@ public:
 		int k = 1;
 		for (int i = 1; i < N; ++i)
 		{
-			const FVector3<T>& Prev = Vertices[i-1];
-			const FVector3<T>& Cur = Vertices[i];
-			const FVector3<T>& Next = Vertices[i+1];
+			const TVector<T>& Prev = Vertices[i-1];
+			const TVector<T>& Cur = Vertices[i];
+			const TVector<T>& Next = Vertices[i+1];
 			NewPolyline.Vertices[k++] = Alpha * Prev + OneMinusAlpha * Cur;
 			NewPolyline.Vertices[k++] = OneMinusAlpha * Cur + Alpha * Next;
 		}
diff --git a/Engine/Source/Runtime/GeometryCore/Public/QuadricError.h b/Engine/Source/Runtime/GeometryCore/Public/QuadricError.h
--- a/Engine/Source/Runtime/GeometryCore/Public/QuadricError.h
+++ b/Engine/Source/Runtime/GeometryCore/Public/QuadricError.h
@@ -12,6 +12,7 @@ namespace UE
 namespace Geometry
 {
 
+using namespace UE::Math;
 
 /**
  * QuadricError represents a quadratic function that evaluates distance to plane.
@@ -37,7 +38,7 @@ struct TQuadricError
 	/**
 	 * Construct TQuadricError a plane with the given normal and a point on plane
 	 */
-	TQuadricError(const FVector3<RealType>& Normal, const FVector3<RealType>& Point) 
+	TQuadricError(const TVector<RealType>& Normal, const TVector<RealType>& Point) 
 	{
 		Axx = Normal.X * Normal.X;
 		Axy = Normal.X * Normal.Y;
@@ -46,7 +47,7 @@ struct TQuadricError
 		Ayz = Normal.Y * Normal.Z;
 		Azz = Normal.Z * Normal.Z;
 		bx = by = bz = c = 0;
-		FVector3<RealType> v = MultiplyA(Point);
+		TVector<RealType> v = MultiplyA(Point);
 		bx = -v.X; by = -v.Y; bz = -v.Z; // b = -Normal.Dot(Point) Normal
 		c = Point.Dot(v); // note this is the same as Normal.Dot(Point) ^2
 	}
@@ -173,12 +174,12 @@ struct TQuadricError
 	/**
 	 * 
 	 */
-	FVector3<RealType> MultiplyA(const UE::Math::TVector<RealType>& pt) const
+	TVector<RealType> MultiplyA(const UE::Math::TVector<RealType>& pt) const
 	{
 		RealType x = Axx * pt.X + Axy * pt.Y + Axz * pt.Z;
 		RealType y = Axy * pt.X + Ayy * pt.Y + Ayz * pt.Z;
 		RealType z = Axz * pt.X + Ayz * pt.Y + Azz * pt.Z;
-		return FVector3<RealType>(x, y, z);
+		return TVector<RealType>(x, y, z);
 	}
 
 
@@ -202,7 +203,7 @@ struct TQuadricError
 			RealType x = a11 * bvecx + a12 * bvecy + a13 * bvecz;
 			RealType y = a12 * bvecx + a22 * bvecy + a23 * bvecz;
 			RealType z = a13 * bvecx + a23 * bvecy + a33 * bvecz;
-			OutResult = FVector3<RealType>(x, y, z);
+			OutResult = TVector<RealType>(x, y, z);
 			return true;
 		}
 		else 
@@ -236,16 +237,16 @@ struct TQuadricError
 		return Result;
 	}
 
-	static FVector3<RealType> MultiplySymmetricMatrix(const RealType SM[6], const RealType vec[3])
+	static TVector<RealType> MultiplySymmetricMatrix(const RealType SM[6], const RealType vec[3])
 	{
 		
 		RealType a =  SM[0] * vec[0] + SM[1] * vec[1] + SM[2] * vec[2];
 		RealType b =  SM[1] * vec[0] + SM[3] * vec[1] + SM[4] * vec[2];
 		RealType c =  SM[2] * vec[0] + SM[4] * vec[1] + SM[5] * vec[2];
 
-		return FVector3<RealType>(a, b, c);
+		return TVector<RealType>(a, b, c);
 	}
-	static FVector3<RealType> MultiplySymmetricMatrix(const RealType SM[6], const UE::Math::TVector<RealType>& vec)
+	static TVector<RealType> MultiplySymmetricMatrix(const RealType SM[6], const UE::Math::TVector<RealType>& vec)
 	{
 		RealType vectmp[3] = { vec.X, vec.Y, vec.Z };
 		return MultiplySymmetricMatrix(SM, vectmp);
@@ -278,7 +279,7 @@ public:
 
 	struct FPlaneData
 	{
-		FPlaneData(const FVector3<RealType>& Normal, const FVector3<RealType>& Point)
+		FPlaneData(const TVector<RealType>& Normal, const TVector<RealType>& Point)
 		{
 			N    = Normal;
 			Dist = -Normal.Dot(Point);
@@ -298,7 +299,7 @@ public:
 
 		FPlaneData() 
 		{
-			N    = FVector3<RealType>::Zero(); 
+			N    = TVector<RealType>::Zero(); 
 			Dist = RealType(0);
 		}
 
@@ -325,7 +326,7 @@ public:
 			return *this;
 		}
 
-		FVector3<RealType> N;
+		TVector<RealType> N;
 		RealType Dist;
 	};
 
@@ -339,7 +340,7 @@ public:
 	{
 	}
 
-	TVolPresQuadricError(const FVector3<RealType>& Normal, const FVector3<RealType>& Point)
+	TVolPresQuadricError(const TVector<RealType>& Normal, const TVector<RealType>& Point)
 		: BaseStruct(Normal, Point)
 		, PlaneData(Normal, Point)
 	{ }
@@ -403,7 +404,7 @@ public:
 		RealType Tol = minThresh; //(1.e-7);  NB: using minThresh here to better compare with the attribute version..
 		//constexpr RealType Tol = (1.e-7);  //NB: using minThresh here to better compare with the attribute version..
 		
-		FVector3<RealType> Ainv_g;
+		TVector<RealType> Ainv_g;
 		if (BaseStruct::SolveAxEqualsb(Ainv_g, gvol.N.X, gvol.N.Y, gvol.N.Z, Tol))
 		{
 
@@ -447,9 +448,9 @@ public:
 	typedef TQuadricError<RealType>  BaseStruct;
 
 	// Triangle Quadric constructor.  Take vertex locations, vertex normals, face normal, and center of face.
-	TAttrBasedQuadricError(const FVector3<RealType>& P0, const FVector3<RealType>& P1, const FVector3<RealType>& P2,
-		const FVector3<RealType>& N0, const FVector3<RealType>& N1, const FVector3<RealType>& N2,
-		const FVector3<RealType>& NFace, const FVector3<RealType>& CenterPoint, RealType AttrWeight)
+	TAttrBasedQuadricError(const TVector<RealType>& P0, const TVector<RealType>& P1, const TVector<RealType>& P2,
+		const TVector<RealType>& N0, const TVector<RealType>& N1, const TVector<RealType>& N2,
+		const TVector<RealType>& NFace, const TVector<RealType>& CenterPoint, RealType AttrWeight)
 		: BaseClass(NFace, CenterPoint),
 		a(0),
 		attrweight(AttrWeight)
@@ -518,7 +519,7 @@ public:
 		{
 			for (int i = 0; i < 3; ++i)
 			{
-				grad[i] = FVector3<RealType>::Zero();
+				grad[i] = TVector<RealType>::Zero();
 				d[i] = AttrWeight * (N0[i] + N1[i] + N2[i]) / 3.; // just average the attribute.
 			}
 		}
@@ -529,10 +530,9 @@ public:
 		a(0.),
 		attrweight(1.)
 	{
-		const RealType zero[3] = { RealType(0), RealType(0), RealType(0) };
-		grad[0] = FVector3<RealType>(zero);
-		grad[1] = FVector3<RealType>(zero);
-		grad[2] = FVector3<RealType>(zero);
+		grad[0] = TVector<RealType>::Zero();
+		grad[1] = TVector<RealType>::Zero();
+		grad[2] = TVector<RealType>::Zero();
 
 		d[0] = RealType(0);
 		d[1] = RealType(0);
@@ -641,7 +641,7 @@ public:
 			OptPoint = BaseStruct::MultiplySymmetricMatrix(InvSM, b);
 			//SolveAxEqualsb(OptPoint, b[0], b[1], b[2]);
 
-			FVector3<RealType> InvSMgvol = BaseStruct::MultiplySymmetricMatrix(InvSM, gvol.N);
+			TVector<RealType> InvSMgvol = BaseStruct::MultiplySymmetricMatrix(InvSM, gvol.N);
 			//SolveAxEqualsb(InvSMgvol, gvol.N.X, gvol.N.Y, gvol.N.Z);
 
 			RealType gvolDotInvSMgvol = gvol.N.Dot(InvSMgvol);
@@ -664,7 +664,7 @@ public:
 		
 	}
 
-	RealType Evaluate(const FVector3<RealType>& point, const FVector3<RealType>& InAttr) const
+	RealType Evaluate(const TVector<RealType>& point, const TVector<RealType>& InAttr) const
 	{
 		// 6x6 symmetric matrix   (    A      -grad[0], -grad[1], -grad[2] )
 		//                        ( -grad[0]^T                             ) 
@@ -677,11 +677,11 @@ public:
 		//                            ( attr  )
 		//
 
-		const FVector3<RealType> attr = attrweight * InAttr;
+		const TVector<RealType> attr = attrweight * InAttr;
 
 		RealType ptAp = point.Dot(BaseStruct::MultiplyA(point));
 		RealType attrDotattr = attr.Dot(attr);
-		FVector3<RealType> Gradattr(grad[0].X * attr[0] + grad[1].X * attr[1] + grad[2].X * attr[2],
+		TVector<RealType> Gradattr(grad[0].X * attr[0] + grad[1].X * attr[1] + grad[2].X * attr[2],
 			                        grad[0].Y * attr[0] + grad[1].Y * attr[1] + grad[2].Y * attr[2],
 			                        grad[0].Z * attr[0] + grad[1].Z * attr[1] + grad[2].Z * attr[2]);
 		RealType ptGradattr = point.Dot(Gradattr);
@@ -700,7 +700,7 @@ public:
 		return QuadricError;
 	}
 
-	void ComputeAttributes(const FVector3<RealType>& point, FVector3<RealType>& attr) const
+	void ComputeAttributes(const TVector<RealType>& point, TVector<RealType>& attr) const
 	{
 		if (FMath::Abs(a) > 1.e-5)
 		{
@@ -723,9 +723,9 @@ public:
 		attr *= 1. / attrweight;
 	}
 
-	RealType Evaluate(const FVector3<RealType>& point) const
+	RealType Evaluate(const TVector<RealType>& point) const
 	{
-		FVector3<RealType> attr;
+		TVector<RealType> attr;
 		ComputeAttributes(point, attr);
 
 		return Evaluate(point, attr);
@@ -741,7 +741,7 @@ public:
 	RealType           attrweight = 1.;
 
 	// Additional planes for the attributes.
-	FVector3<RealType> grad[3];
+	TVector<RealType> grad[3];
 	RealType           d[3];
 
 
@@ -755,7 +755,7 @@ typedef TAttrBasedQuadricError<double> FAttrBasedQuadricErrord;
 * On return the quadric has been weighted by the length of the (p1-p0).
 */
 template<typename RealType>
-TQuadricError<RealType>  CreateSeamQuadric(const FVector3<RealType>& p0, const FVector3<RealType>& p1, const FVector3<RealType>& AdjFaceNormal)
+TQuadricError<RealType>  CreateSeamQuadric(const TVector<RealType>& p0, const TVector<RealType>& p1, const TVector<RealType>& AdjFaceNormal)
 {
 
 	RealType LenghtSqrd = AdjFaceNormal.SquaredLength();
@@ -767,7 +767,7 @@ TQuadricError<RealType>  CreateSeamQuadric(const FVector3<RealType>& p0, const F
 	}
 
 
-	FVector3<RealType> Edge = p1 - p0;
+	TVector<RealType> Edge = p1 - p0;
 
 	// Weight scaled on edge length 
 
@@ -787,7 +787,7 @@ TQuadricError<RealType>  CreateSeamQuadric(const FVector3<RealType>& p0, const F
 
 	// Normal that is perpendicular to the edge, and face Normal. - i.e. in the face plane and perpendicular to the edge.
 	// The constraint should try to keep points on the plane associated with this constraint plane.
-	FVector3<RealType> ConstraintNormal = Edge.Cross(AdjFaceNormal);
+	TVector<RealType> ConstraintNormal = Edge.Cross(AdjFaceNormal);
 
 
 	TQuadricError<RealType> SeamQuadric(ConstraintNormal, p0);
diff --git a/Engine/Source/Runtime/GeometryCore/Public/Quaternion.h b/Engine/Source/Runtime/GeometryCore/Public/Quaternion.h
--- a/Engine/Source/Runtime/GeometryCore/Public/Quaternion.h
+++ b/Engine/Source/Runtime/GeometryCore/Public/Quaternion.h
@@ -14,6 +14,8 @@ namespace UE
 namespace Geometry
 {
 
+using namespace UE::Math;
+
 
 template<typename RealType>
 struct TQuaternion
@@ -30,14 +32,14 @@ struct TQuaternion
 	TQuaternion(const TQuaternion& Copy);
 	template<typename RealType2>
 	explicit TQuaternion(const TQuaternion<RealType2>& Copy);
-	TQuaternion(const FVector3<RealType>& Axis, RealType Angle, bool bAngleIsDegrees);
-	TQuaternion(const FVector3<RealType>& From, const FVector3<RealType>& To);
+	TQuaternion(const TVector<RealType>& Axis, RealType Angle, bool bAngleIsDegrees);
+	TQuaternion(const TVector<RealType>& From, const TVector<RealType>& To);
 	TQuaternion(const TQuaternion<RealType>& From, const TQuaternion<RealType>& To, RealType InterpT);
 	TQuaternion(const TMatrix3<RealType>& RotationMatrix);
 
-	void SetAxisAngleD(const FVector3<RealType>& Axis, RealType AngleDeg);
-	void SetAxisAngleR(const FVector3<RealType>& Axis, RealType AngleRad);
-	void SetFromTo(const FVector3<RealType>& From, const FVector3<RealType>& To);
+	void SetAxisAngleD(const TVector<RealType>& Axis, RealType AngleDeg);
+	void SetAxisAngleR(const TVector<RealType>& Axis, RealType AngleRad);
+	void SetFromTo(const TVector<RealType>& From, const TVector<RealType>& To);
 	void SetToSlerp(TQuaternion<RealType> From, TQuaternion<RealType> To, RealType InterpT);
 	void SetFromRotationMatrix(const TMatrix3<RealType>& RotationMatrix);
 
@@ -52,10 +54,10 @@ struct TQuaternion
 	RealType Length() const { return (RealType)sqrt(X*X + Y*Y + Z*Z + W*W); }
 	RealType SquaredLength() const { return X*X + Y*Y + Z*Z + W*W; }
 
-	FVector3<RealType> AxisX() const;
-	FVector3<RealType> AxisY() const;
-	FVector3<RealType> AxisZ() const;
-	void GetAxes(FVector3<RealType>& X, FVector3<RealType>& Y, FVector3<RealType>& Z) const;
+	TVector<RealType> AxisX() const;
+	TVector<RealType> AxisY() const;
+	TVector<RealType> AxisZ() const;
+	void GetAxes(TVector<RealType>& X, TVector<RealType>& Y, TVector<RealType>& Z) const;
 
 	RealType Normalize(const RealType epsilon = 0);
 	TQuaternion<RealType> Normalized(const RealType epsilon = 0) const;
@@ -63,7 +65,7 @@ struct TQuaternion
 	RealType Dot(const TQuaternion<RealType>& Other) const;
 
 	TQuaternion<RealType> Inverse() const;
-	FVector3<RealType> InverseMultiply(const FVector3<RealType>& Other) const;
+	TVector<RealType> InverseMultiply(const TVector<RealType>& Other) const;
 
 	TMatrix3<RealType> ToRotationMatrix() const;
 
@@ -148,14 +150,14 @@ TQuaternion<RealType>::TQuaternion(const TQuaternion& Copy)
 }
 
 template<typename RealType>
-TQuaternion<RealType>::TQuaternion(const FVector3<RealType>& Axis, RealType Angle, bool bAngleIsDegrees)
+TQuaternion<RealType>::TQuaternion(const TVector<RealType>& Axis, RealType Angle, bool bAngleIsDegrees)
 {
 	X = Y = Z = 0; W = 1;
 	SetAxisAngleR(Axis, Angle * (bAngleIsDegrees ? TMathUtil<RealType>::DegToRad : (RealType)1));
 }
 
 template<typename RealType>
-TQuaternion<RealType>::TQuaternion(const FVector3<RealType>& From, const FVector3<RealType>& To)
+TQuaternion<RealType>::TQuaternion(const TVector<RealType>& From, const TVector<RealType>& To)
 {
 	X = Y = Z = 0; W = 1;
 	SetFromTo(From, To);
@@ -232,7 +234,7 @@ TQuaternion<RealType> operator -(const TQuaternion<RealType>& A, const TQuaterni
 }
 
 template<typename RealType>
-FVector3<RealType> operator*(const TQuaternion<RealType>& Q, const UE::Math::TVector<RealType>& V)
+TVector<RealType> operator*(const TQuaternion<RealType>& Q, const UE::Math::TVector<RealType>& V)
 {
 	//return q.ToRotationMatrix() * v;
 	// inline-expansion of above:
@@ -240,7 +242,7 @@ FVector3<RealType> operator*(const TQuaternion<RealType>& Q, const UE::Math::TVe
 	RealType twoWX = twoX * Q.W; RealType twoWY = twoY * Q.W; RealType twoWZ = twoZ * Q.W;
 	RealType twoXX = twoX * Q.X; RealType twoXY = twoY * Q.X; RealType twoXZ = twoZ * Q.X;
 	RealType twoYY = twoY * Q.Y; RealType twoYZ = twoZ * Q.Y; RealType twoZZ = twoZ * Q.Z;
-	return FVector3<RealType>(
+	return TVector<RealType>(
 		V.X * ((RealType)1 - (twoYY + twoZZ)) + V.Y * (twoXY - twoWZ) + V.Z * (twoXZ + twoWY),
 		V.X * (twoXY + twoWZ) + V.Y * ((RealType)1 - (twoXX + twoZZ)) + V.Z * (twoYZ - twoWX),
 		V.X * (twoXZ - twoWY) + V.Y * (twoYZ + twoWX) + V.Z * ((RealType)1 - (twoXX + twoYY))); ;
@@ -248,7 +250,7 @@ FVector3<RealType> operator*(const TQuaternion<RealType>& Q, const UE::Math::TVe
 
 
 template<typename RealType>
-FVector3<RealType> TQuaternion<RealType>::InverseMultiply(const FVector3<RealType>& V) const
+TVector<RealType> TQuaternion<RealType>::InverseMultiply(const TVector<RealType>& V) const
 {
 	RealType norm = SquaredLength();
 	if (norm > 0) 
@@ -259,53 +261,53 @@ FVector3<RealType> TQuaternion<RealType>::InverseMultiply(const FVector3<RealTyp
 		RealType twoWX = twoX * qW; RealType twoWY = twoY * qW; RealType twoWZ = twoZ * qW;
 		RealType twoXX = twoX * qX; RealType twoXY = twoY * qX; RealType twoXZ = twoZ * qX;
 		RealType twoYY = twoY * qY; RealType twoYZ = twoZ * qY; RealType twoZZ = twoZ * qZ;
-		return FVector3<RealType>(
+		return TVector<RealType>(
 			V.X * ((RealType)1 - (twoYY + twoZZ)) + V.Y * (twoXY - twoWZ) + V.Z * (twoXZ + twoWY),
 			V.X * (twoXY + twoWZ) + V.Y * ((RealType)1 - (twoXX + twoZZ)) + V.Z * (twoYZ - twoWX),
 			V.X * (twoXZ - twoWY) + V.Y * (twoYZ + twoWX) + V.Z * ((RealType)1 - (twoXX + twoYY)));
 	}
-	return FVector3<RealType>::Zero();
+	return TVector<RealType>::Zero();
 }
 
 
 template<typename RealType>
-FVector3<RealType> TQuaternion<RealType>::AxisX() const
+TVector<RealType> TQuaternion<RealType>::AxisX() const
 {
 	RealType twoY = (RealType)2 * Y; RealType twoZ = (RealType)2 * Z;
 	RealType twoWY = twoY * W; RealType twoWZ = twoZ * W;
 	RealType twoXY = twoY * X; RealType twoXZ = twoZ * X;
 	RealType twoYY = twoY * Y; RealType twoZZ = twoZ * Z;
-	return FVector3<RealType>((RealType)1 - (twoYY + twoZZ), twoXY + twoWZ, twoXZ - twoWY);
+	return TVector<RealType>((RealType)1 - (twoYY + twoZZ), twoXY + twoWZ, twoXZ - twoWY);
 }
 
 template<typename RealType>
-FVector3<RealType> TQuaternion<RealType>::AxisY() const
+TVector<RealType> TQuaternion<RealType>::AxisY() const
 {
 	RealType twoX = (RealType)2 * X; RealType twoY = (RealType)2 * Y; RealType twoZ = (RealType)2 * Z;
 	RealType twoWX = twoX * W; RealType twoWZ = twoZ * W; RealType twoXX = twoX * X;
 	RealType twoXY = twoY * X; RealType twoYZ = twoZ * Y; RealType twoZZ = twoZ * Z;
-	return FVector3<RealType>(twoXY - twoWZ, (RealType)1 - (twoXX + twoZZ), twoYZ + twoWX);
+	return TVector<RealType>(twoXY - twoWZ, (RealType)1 - (twoXX + twoZZ), twoYZ + twoWX);
 }
 
 template<typename RealType>
-FVector3<RealType> TQuaternion<RealType>::AxisZ() const
+TVector<RealType> TQuaternion<RealType>::AxisZ() const
 {
 	RealType twoX = (RealType)2 * X; RealType twoY = (RealType)2 * Y; RealType twoZ = (RealType)2 * Z;
 	RealType twoWX = twoX * W; RealType twoWY = twoY * W; RealType twoXX = twoX * X;
 	RealType twoXZ = twoZ * X; RealType twoYY = twoY * Y; RealType twoYZ = twoZ * Y;
-	return FVector3<RealType>(twoXZ + twoWY, twoYZ - twoWX, (RealType)1 - (twoXX + twoYY));
+	return TVector<RealType>(twoXZ + twoWY, twoYZ - twoWX, (RealType)1 - (twoXX + twoYY));
 }
 
 template<typename RealType>
-void TQuaternion<RealType>::GetAxes(FVector3<RealType>& XOut, FVector3<RealType>& YOut, FVector3<RealType>& ZOut) const
+void TQuaternion<RealType>::GetAxes(TVector<RealType>& XOut, TVector<RealType>& YOut, TVector<RealType>& ZOut) const
 {
 	RealType twoX = (RealType)2 * X; RealType twoY = (RealType)2 * Y; RealType twoZ = (RealType)2 * Z;
 	RealType twoWX = twoX * W; RealType twoWY = twoY * W; RealType twoWZ = twoZ * W;
 	RealType twoXX = twoX * X; RealType twoXY = twoY * X; RealType twoXZ = twoZ * X;
 	RealType twoYY = twoY * Y; RealType twoYZ = twoZ * Y; RealType twoZZ = twoZ * Z;
-	XOut = FVector3<RealType>((RealType)1 - (twoYY + twoZZ), twoXY + twoWZ, twoXZ - twoWY);
-	YOut = FVector3<RealType>(twoXY - twoWZ, (RealType)1 - (twoXX + twoZZ), twoYZ + twoWX);
-	ZOut = FVector3<RealType>(twoXZ + twoWY, twoYZ - twoWX, (RealType)1 - (twoXX + twoYY));
+	XOut = TVector<RealType>((RealType)1 - (twoYY + twoZZ), twoXY + twoWZ, twoXZ - twoWY);
+	YOut = TVector<RealType>(twoXY - twoWZ, (RealType)1 - (twoXX + twoZZ), twoYZ + twoWX);
+	ZOut = TVector<RealType>(twoXZ + twoWY, twoYZ - twoWX, (RealType)1 - (twoXX + twoYY));
 }
 
 
@@ -324,13 +326,13 @@ TQuaternion<RealType> TQuaternion<RealType>::Inverse() const
 
 
 template<typename RealType>
-void TQuaternion<RealType>::SetAxisAngleD(const FVector3<RealType>& Axis, RealType AngleDeg) 
+void TQuaternion<RealType>::SetAxisAngleD(const TVector<RealType>& Axis, RealType AngleDeg) 
 {
 	SetAxisAngleR(Axis, TMathUtil<RealType>::DegToRad * AngleDeg);
 }
 
 template<typename RealType>
-void TQuaternion<RealType>::SetAxisAngleR(const FVector3<RealType>& Axis, RealType AngleRad)
+void TQuaternion<RealType>::SetAxisAngleR(const TVector<RealType>& Axis, RealType AngleRad)
 {
 	RealType halfAngle = (RealType)0.5 * AngleRad;
 	RealType sn = (RealType)sin(halfAngle);
@@ -344,7 +346,7 @@ void TQuaternion<RealType>::SetAxisAngleR(const FVector3<RealType>& Axis, RealTy
 
 // this function can take non-normalized vectors vFrom and vTo (normalizes internally)
 template<typename RealType>
-void TQuaternion<RealType>::SetFromTo(const FVector3<RealType>& From, const FVector3<RealType>& To)
+void TQuaternion<RealType>::SetFromTo(const TVector<RealType>& From, const TVector<RealType>& To)
 {
 	// [TODO] this page seems to have optimized version:
 	//    http://lolengine.net/blog/2013/09/18/beautiful-maths-quaternion-from-vectors
@@ -352,12 +354,12 @@ void TQuaternion<RealType>::SetFromTo(const FVector3<RealType>& From, const FVec
 	// [RMS] not ideal to explicitly normalize here, but if we don't,
 	//   output TQuaternion is not normalized and this causes problems,
 	//   eg like drift if we do repeated SetFromTo()
-	FVector3<RealType> from = UE::Geometry::Normalized(From), to = UE::Geometry::Normalized(To);
-	FVector3<RealType> bisector = UE::Geometry::Normalized(from + to);
+	TVector<RealType> from = UE::Geometry::Normalized(From), to = UE::Geometry::Normalized(To);
+	TVector<RealType> bisector = UE::Geometry::Normalized(from + to);
 	W = from.Dot(bisector);
 	if (W != 0) 
 	{
-		FVector3<RealType> cross = from.Cross(bisector);
+		TVector<RealType> cross = from.Cross(bisector);
 		X = cross.X;
 		Y = cross.Y;
 		Z = cross.Z;
@@ -465,7 +467,7 @@ void TQuaternion<RealType>::SetFromRotationMatrix(const TMatrix3<RealType>& Rota
 
 		root = (RealType)sqrt(RotationMatrix(i, i) - RotationMatrix(j, j) - RotationMatrix(k, k) + (RealType)1);
 
-		FVector3<RealType> quat(X, Y, Z);
+		TVector<RealType> quat(X, Y, Z);
 		quat[i] = ((RealType)0.5) * root;
 		root = ((RealType)0.5) / root;
 		W = (RotationMatrix(k, j) - RotationMatrix(j, k)) * root;
@@ -488,9 +490,9 @@ TMatrix3<RealType> TQuaternion<RealType>::ToRotationMatrix() const
 	RealType twoXX = twoX * X; RealType twoXY = twoY * X; RealType twoXZ = twoZ * X;
 	RealType twoYY = twoY * Y; RealType twoYZ = twoZ * Y; RealType twoZZ = twoZ * Z;
 	TMatrix3<RealType> m = TMatrix3<RealType>::Zero();
-	m.Row0 = FVector3<RealType>(1 - (twoYY + twoZZ), twoXY - twoWZ, twoXZ + twoWY);
-	m.Row1 = FVector3<RealType>(twoXY + twoWZ, 1 - (twoXX + twoZZ), twoYZ - twoWX);
-	m.Row2 = FVector3<RealType>(twoXZ - twoWY, twoYZ + twoWX, 1 - (twoXX + twoYY));
+	m.Row0 = TVector<RealType>(1 - (twoYY + twoZZ), twoXY - twoWZ, twoXZ + twoWY);
+	m.Row1 = TVector<RealType>(twoXY + twoWZ, 1 - (twoXX + twoZZ), twoYZ - twoWX);
+	m.Row2 = TVector<RealType>(twoXZ - twoWY, twoYZ + twoWX, 1 - (twoXX + twoYY));
 	return m;
 }
 
diff --git a/Engine/Source/Runtime/GeometryCore/Public/RayTypes.h b/Engine/Source/Runtime/GeometryCore/Public/RayTypes.h
--- a/Engine/Source/Runtime/GeometryCore/Public/RayTypes.h
+++ b/Engine/Source/Runtime/GeometryCore/Public/RayTypes.h
@@ -13,6 +13,8 @@ namespace UE
 namespace Geometry
 {
 
+using namespace UE::Math;
+
 /*
  * 3D Ray stored as Origin point and normalized Direction vector
  */
@@ -21,15 +23,15 @@ class TRay3
 {
 public:
 	/** Origin point */
-	FVector3<RealType> Origin;
+	TVector<RealType> Origin;
 	/** Direction vector, always normalized */
-	FVector3<RealType> Direction;
+	TVector<RealType> Direction;
 
 	/** Construct ray at origin pointed down Z axis */
 	TRay3()
 	{
-		Origin = FVector3<RealType>::Zero();
-		Direction = FVector3<RealType>::UnitZ();
+		Origin = TVector<RealType>::Zero();
+		Direction = TVector<RealType>::UnitZ();
 	}
 
 	/**
@@ -38,7 +40,7 @@ public:
 	 * @param Direction direction vector. Must be normalized if bIsNormalized=true
 	 * @param bIsNormalized if true, Direction will not be re-normalized
 	 */
-	TRay3(const FVector3<RealType>& Origin, const FVector3<RealType>& Direction, bool bIsNormalized = false)
+	TRay3(const TVector<RealType>& Origin, const TVector<RealType>& Direction, bool bIsNormalized = false)
 	{
 		this->Origin = Origin;
 		this->Direction = Direction;
@@ -51,7 +53,7 @@ public:
 	/**
 	 * @return point on ray at given (signed) Distance from the ray Origin
 	 */
-	FVector3<RealType> PointAt(RealType Distance) const
+	TVector<RealType> PointAt(RealType Distance) const
 	{
 		return Origin + Distance * Direction;
 	}
@@ -60,7 +62,7 @@ public:
 	/**
 	 * @return ray parameter (ie positive distance from Origin) at nearest point on ray to QueryPoint
 	 */
-	inline RealType Project(const FVector3<RealType>& QueryPoint) const
+	inline RealType Project(const TVector<RealType>& QueryPoint) const
 	{
 		RealType LineParam = (QueryPoint - Origin).Dot(Direction);
 		return (LineParam < 0) ? 0 : LineParam;
@@ -69,24 +71,24 @@ public:
 	/**
 	 * @return smallest squared distance from ray to QueryPoint
 	 */
-	inline RealType DistanceSquared(const FVector3<RealType>& QueryPoint) const
+	inline RealType DistanceSquared(const TVector<RealType>& QueryPoint) const
 	{
 		RealType LineParam = (QueryPoint - Origin).Dot(Direction);
 		if (LineParam < 0)
 		{
-			return Origin.DistanceSquared(QueryPoint);
+			return UE::Geometry::DistanceSquared(Origin, QueryPoint);
 		}
 		else
 		{
-			FVector3<RealType> NearestPt = Origin + LineParam * Direction;
-			return NearestPt.DistanceSquared(QueryPoint);
+			TVector<RealType> NearestPt = Origin + LineParam * Direction;
+			return UE::Geometry::DistanceSquared(NearestPt, QueryPoint);
 		}
 	}
 
 	/**
 	 * @return smallest squared distance from ray to QueryPoint
 	 */
-	inline RealType Distance(const FVector3<RealType>& QueryPoint) const
+	inline RealType Distance(const TVector<RealType>& QueryPoint) const
 	{
 		return TMathUtil<RealType>::Sqrt(DistanceSquared(QueryPoint));
 	}
@@ -94,7 +96,7 @@ public:
 	/**
 	 * @return nearest point on line to QueryPoint
 	 */
-	inline FVector3<RealType> NearestPoint(const FVector3<RealType>& QueryPoint) const
+	inline TVector<RealType> NearestPoint(const TVector<RealType>& QueryPoint) const
 	{
 		RealType LineParam = (QueryPoint - Origin).Dot(Direction);
 		if (LineParam < 0)
@@ -117,8 +119,8 @@ public:
 	}
 	inline TRay3(const FRay & RayIn)
 	{
-		Origin = FVector3<RealType>(RayIn.Origin);
-		Direction = FVector3<RealType>(RayIn.Direction);
+		Origin = TVector<RealType>(RayIn.Origin);
+		Direction = TVector<RealType>(RayIn.Direction);
 	}
 
 
diff --git a/Engine/Source/Runtime/GeometryCore/Public/Sampling/Gaussians.h b/Engine/Source/Runtime/GeometryCore/Public/Sampling/Gaussians.h
--- a/Engine/Source/Runtime/GeometryCore/Public/Sampling/Gaussians.h
+++ b/Engine/Source/Runtime/GeometryCore/Public/Sampling/Gaussians.h
@@ -10,6 +10,7 @@ namespace UE
 namespace Geometry
 {
 
+using namespace UE::Math;
 
 /**
  * 1D Gaussian
@@ -168,13 +169,13 @@ public:
 
 	TGaussian3(RealType SigmaIn = 1.0) : Sigma(SigmaIn) {}
 
-	RealType Evaluate(const FVector3<RealType>& XYZ) const
+	RealType Evaluate(const TVector<RealType>& XYZ) const
 	{
 		RealType InvTwoSigmaSqr = 1.0 / (2.0 * Sigma * Sigma);
 		return TMathUtil<RealType>::InvPi * InvTwoSigmaSqr * TMathUtil<RealType>::Exp( -XYZ.SquaredLength() * InvTwoSigmaSqr );
 	}
 
-	RealType EvaluateSqrUnscaled(const FVector3<RealType>& XYZ) const
+	RealType EvaluateSqrUnscaled(const TVector<RealType>& XYZ) const
 	{
 		RealType InvTwoSigmaSqr = 1.0 / (2.0 * Sigma * Sigma);
 		return TMathUtil<RealType>::Exp(-XYZ.SquaredLength() * InvTwoSigmaSqr);
diff --git a/Engine/Source/Runtime/GeometryCore/Public/Sampling/NormalHistogram.h b/Engine/Source/Runtime/GeometryCore/Public/Sampling/NormalHistogram.h
--- a/Engine/Source/Runtime/GeometryCore/Public/Sampling/NormalHistogram.h
+++ b/Engine/Source/Runtime/GeometryCore/Public/Sampling/NormalHistogram.h
@@ -11,6 +11,7 @@ namespace UE
 namespace Geometry
 {
 
+using namespace UE::Math;
 
 /**
  * TNormalHistogram calculates/represents a histogram on a 3D sphere. The histogram bins
@@ -41,7 +42,7 @@ public:
 	 * Count a Normal in the histogram, ie find the Bin it should be included in and add the Weight to that Bin (The Normal does not strictly need to be normalized)
 	 * @param Weight optional weight for this normal
 	 */
-	void Count(const FVector3<RealType>& Normal, RealType Weight = 1)
+	void Count(const TVector<RealType>& Normal, RealType Weight = 1)
 	{
 		int32 BinIndex = BinPoints.FindIndex(Normal);
 		WeightedCounts[BinIndex] += Weight;
@@ -50,7 +51,7 @@ public:
 	/**
 	 * @return the normal direction for the Histogram bin that contains the most points / largest mass
 	 */
-	FVector3<RealType> FindMaxNormal() const
+	TVector<RealType> FindMaxNormal() const
 	{
 		int MaxIndex = 0;
 		for (int k = 1; k < Bins; ++k)
diff --git a/Engine/Source/Runtime/GeometryCore/Public/Sampling/SphericalFibonacci.h b/Engine/Source/Runtime/GeometryCore/Public/Sampling/SphericalFibonacci.h
--- a/Engine/Source/Runtime/GeometryCore/Public/Sampling/SphericalFibonacci.h
+++ b/Engine/Source/Runtime/GeometryCore/Public/Sampling/SphericalFibonacci.h
@@ -41,7 +41,7 @@ public:
 	 * @param Index point index in range [0,Num()-1]
 	 * @return sphere point for given Index
 	 */
-	FVector3<RealType> Point(int32 Index) const
+	TVector<RealType> Point(int32 Index) const
 	{
 		static const RealType PHI = ( TMathUtil<RealType>::Sqrt(5.0) + 1.0) / 2.0;
 
@@ -54,7 +54,7 @@ public:
 		RealType theta = TMathUtil<RealType>::ACos(z);
 		RealType sin_theta = TMathUtil<RealType>::Sin(theta);
 
-		return FVector3<RealType>(cos_phi * sin_theta, sin_phi * sin_theta, z);
+		return TVector<RealType>(cos_phi * sin_theta, sin_phi * sin_theta, z);
 	}
 
 
@@ -62,7 +62,7 @@ public:
 	 * @param Index point index in range [0,Num()-1]
 	 * @return sphere point for given Index
 	 */
-	FVector3<RealType> operator[](int32 Index) const
+	TVector<RealType> operator[](int32 Index) const
 	{
 		return Point(Index);
 	}
@@ -71,7 +71,7 @@ public:
 	/**
 	 * @return Index of the sphere point closest to the given point P
 	 */
-	int32 FindIndex(const FVector3<RealType>& P)
+	int32 FindIndex(const TVector<RealType>& P)
 	{
 		static const RealType PHI = (TMathUtil<RealType>::Sqrt(5.0) + 1.0) / 2.0;
 
@@ -105,11 +105,11 @@ public:
 			phi = 2.0 * TMathUtil<RealType>::Pi * MultiplyAddFrac(i, PHI - 1);
 			cosTheta = 1.0 - (2.0 * i + 1.0) * (1.0 / N); // rcp(n);
 			RealType sinTheta = TMathUtil<RealType>::Sqrt(1.0 - cosTheta * cosTheta);
-			FVector3<RealType> q( 
+			TVector<RealType> q( 
 				TMathUtil<RealType>::Cos(phi) * sinTheta, 
 				TMathUtil<RealType>::Sin(phi) * sinTheta, 
 				cosTheta );
-			RealType SquaredDistance = q.DistanceSquared(P);
+			RealType SquaredDistance = DistanceSquared(q, P);
 			if (SquaredDistance < d)
 			{
 				d = SquaredDistance;
@@ -243,11 +243,11 @@ public:
 	 * @param Index point index in range [0,Num()-1]
 	 * @return sphere point for given Index
 	 */
-	FVector3<RealType> Point(int32 Index) const
+	TVector<RealType> Point(int32 Index) const
 	{
 		checkSlow(Index >= 0 && Index < N);
 
-		FVector3<RealType> Point;
+		TVector<RealType> Point;
 		switch (Distribution)
 		{
 		case EDistribution::Uniform:
@@ -270,7 +270,7 @@ public:
 			// Reference: Physically Based Rendering: From Theory to Implementation
 			// http://www.pbr-book.org/3ed-2018/Monte_Carlo_Integration/2D_Sampling_with_Multidimensional_Transformations.html#Cosine-WeightedHemisphereSampling
 			RealType Z = TMathUtil<RealType>::Sqrt(TMathUtil<RealType>::Max(0.0, 1.0 - Pt.X * Pt.X - Pt.Y * Pt.Y));
-			Point = FVector3<RealType>(Pt.X, Pt.Y, Z);
+			Point = TVector<RealType>(Pt.X, Pt.Y, Z);
 			break;
 		}
 		}
@@ -281,7 +281,7 @@ public:
 	 * @param Index point index in range [0,Num()-1]
 	 * @return sphere point for given Index
 	 */
-	FVector3<RealType> operator[](int32 Index) const
+	TVector<RealType> operator[](int32 Index) const
 	{
 		return Point(Index);
 	}
diff --git a/Engine/Source/Runtime/GeometryCore/Public/Sampling/VectorSetAnalysis.h b/Engine/Source/Runtime/GeometryCore/Public/Sampling/VectorSetAnalysis.h
--- a/Engine/Source/Runtime/GeometryCore/Public/Sampling/VectorSetAnalysis.h
+++ b/Engine/Source/Runtime/GeometryCore/Public/Sampling/VectorSetAnalysis.h
@@ -10,6 +10,8 @@ namespace UE
 namespace Geometry
 {
 
+using namespace UE::Math;
+
 /**
  * TVectorSetAnalysis3 computes various analyses of a set of input Vectors (currently mainly clustering.
  */
@@ -21,7 +23,7 @@ public:
 	// Input data
 	//
 
-	TArray<FVector3<RealType>> Vectors;
+	TArray<TVector<RealType>> Vectors;
 	TArray<int32> VectorIDs;
 	bool bNormalized = false;
 
@@ -30,7 +32,7 @@ public:
 	//
 
 	/** Set of vectors that represent centers of clusters */
-	TArray<FVector3<RealType>> ClusterVectors;
+	TArray<TVector<RealType>> ClusterVectors;
 	/** Mapping from Vector index to ClusterVectors index */
 	TArray<int32> VectorToClusterMap;
 
@@ -41,7 +43,7 @@ public:
 	 * @param bIsNormalizedHint indicate whether vectors are normalized
 	 */
 	template<typename EnumerableIDType>
-	void Initialize(EnumerableIDType EnumerableIDs, TFunctionRef<FVector3<RealType>(int32)> GetVectorFunc, int32 NumVectorsHint = 0, bool bIsNormalizedHint = false)
+	void Initialize(EnumerableIDType EnumerableIDs, TFunctionRef<TVector<RealType>(int32)> GetVectorFunc, int32 NumVectorsHint = 0, bool bIsNormalizedHint = false)
 	{
 		Vectors.Reserve(NumVectorsHint);
 		VectorIDs.Reserve(NumVectorsHint);
diff --git a/Engine/Source/Runtime/GeometryCore/Public/SegmentTypes.h b/Engine/Source/Runtime/GeometryCore/Public/SegmentTypes.h
--- a/Engine/Source/Runtime/GeometryCore/Public/SegmentTypes.h
+++ b/Engine/Source/Runtime/GeometryCore/Public/SegmentTypes.h
@@ -336,9 +336,9 @@ struct TSegment3
 {
 public:
 	/** Center point of segment */
-	FVector3<T> Center = FVector3<T>::Zero();
+	TVector<T> Center = TVector<T>::Zero();
 	/** normalized Direction vector of segment */
-	FVector3<T> Direction = FVector3<T>::UnitX();
+	TVector<T> Direction = TVector<T>::UnitX();
 	/** Extent of segment, which is half the total length */
 	T Extent = (T)0;
 
@@ -347,7 +347,7 @@ public:
 	/**
 	 * Construct a Segment from two Points
 	 */
-	TSegment3(const FVector3<T>& Point0, const FVector3<T>& Point1)
+	TSegment3(const TVector<T>& Point0, const TVector<T>& Point1)
 	{
 		// set from endpoints 
 		Center = T(.5) * (Point0 + Point1);
@@ -358,7 +358,7 @@ public:
 	/**
 	 * Construct a segment from a Center Point, normalized Direction, and scalar Extent
 	 */
-	TSegment3(const FVector3<T>& CenterIn, const FVector3<T>& DirectionIn, T ExtentIn)
+	TSegment3(const TVector<T>& CenterIn, const TVector<T>& DirectionIn, T ExtentIn)
 	{
 		Center = CenterIn;
 		Direction = DirectionIn;
@@ -368,13 +368,13 @@ public:
 
 
 	/** Update the Segment with a new start point */
-	inline void SetStartPoint(const FVector3<T>& Point)
+	inline void SetStartPoint(const TVector<T>& Point)
 	{
 		update_from_endpoints(Point, EndPoint());
 	}
 
 	/** Update the Segment with a new end point */
-	inline void SetEndPoint(const FVector3<T>& Point)
+	inline void SetEndPoint(const TVector<T>& Point)
 	{
 		update_from_endpoints(StartPoint(), Point);
 	}
@@ -389,13 +389,13 @@ public:
 
 
 	/** @return start point of segment */
-	inline FVector3<T> StartPoint() const
+	inline TVector<T> StartPoint() const
 	{
 		return Center - Extent * Direction;
 	}
 
 	/** @return end point of segment */
-	inline FVector3<T> EndPoint() const
+	inline TVector<T> EndPoint() const
 	{
 		return Center + Extent * Direction;
 	}
@@ -408,7 +408,7 @@ public:
 	}
 
 	/** @return first (i == 0) or second (i == 1) endpoint of the Segment  */
-	inline FVector3<T> GetPointFromIndex(int i) const
+	inline TVector<T> GetPointFromIndex(int i) const
 	{
 		return (i == 0) ? (Center - Extent * Direction) : (Center + Extent * Direction);
 	}
@@ -416,7 +416,7 @@ public:
 	/**
 	 * @return point on segment at given (signed) Distance from the segment Origin
 	 */
-	inline FVector3<T> PointAt(T DistanceParameter) const
+	inline TVector<T> PointAt(T DistanceParameter) const
 	{
 		return Center + DistanceParameter * Direction;
 	}
@@ -425,7 +425,7 @@ public:
 	 * @param UnitParameter value in range [0,1]
 	 * @return point on segment at that linearly interpolates between start and end based on Unit Parameter
 	 */
-	inline FVector3<T> PointBetween(T UnitParameter) const
+	inline TVector<T> PointBetween(T UnitParameter) const
 	{
 		return Center + ((T)2 * UnitParameter - (T)1) * Extent * Direction;
 	}
@@ -433,7 +433,7 @@ public:
 	/**
 	 * @return minimum squared distance from Point to segment
 	 */
-	inline T DistanceSquared(const FVector3<T>& Point) const
+	inline T DistanceSquared(const TVector<T>& Point) const
 	{
 		T DistParameter;
 		return DistanceSquared(Point, DistParameter);
@@ -446,7 +446,7 @@ public:
 	 * @param DistParameterOut calculated distance parameter in range [-Extent,Extent]
 	 * @return minimum squared distance from Point to Segment
 	 */
-	T DistanceSquared(const FVector3<T>& Point, T& DistParameterOut) const
+	T DistanceSquared(const TVector<T>& Point, T& DistParameterOut) const
 	{
 		DistParameterOut = (Point - Center).Dot(Direction);
 		if (DistParameterOut >= Extent)
@@ -459,7 +459,7 @@ public:
 			DistParameterOut = -Extent;
 			return UE::Geometry::DistanceSquared(Point, StartPoint());
 		}
-		FVector3<T> ProjectedPt = Center + DistParameterOut * Direction;
+		TVector<T> ProjectedPt = Center + DistParameterOut * Direction;
 		return UE::Geometry::DistanceSquared(ProjectedPt, Point);
 	}
 
@@ -467,7 +467,7 @@ public:
 	/**
 	 * @return nearest point on segment to QueryPoint
 	 */
-	inline FVector3<T> NearestPoint(const FVector3<T>& QueryPoint) const
+	inline TVector<T> NearestPoint(const TVector<T>& QueryPoint) const
 	{
 		T t = (QueryPoint - Center).Dot(Direction);
 		if (t >= Extent)
@@ -485,7 +485,7 @@ public:
 	/**
 	 * @return scalar projection of QueryPoint onto line of Segment (not clamped to Extents)
 	 */
-	inline T Project(const FVector3<T>& QueryPoint) const
+	inline T Project(const TVector<T>& QueryPoint) const
 	{
 		return (QueryPoint - Center).Dot(Direction);
 	}
@@ -494,7 +494,7 @@ public:
 	/**
 	 * @return scalar projection of QueryPoint onto line of Segment, mapped to [0,1] range along segment
 	 */
-	inline T ProjectUnitRange(const FVector3<T>& QueryPoint) const
+	inline T ProjectUnitRange(const TVector<T>& QueryPoint) const
 	{
 		T ProjT = (QueryPoint - Center).Dot(Direction);
 		T Alpha = ((ProjT / Extent) + (T)1) * (T)0.5;
@@ -531,7 +531,7 @@ public:
 protected:
 
 	// update segment based on new endpoints
-	inline void update_from_endpoints(const FVector3<T>& p0, const FVector3<T>& p1)
+	inline void update_from_endpoints(const TVector<T>& p0, const TVector<T>& p1)
 	{
 		Center = 0.5 * (p0 + p1);
 		Direction = p1 - p0;
diff --git a/Engine/Source/Runtime/GeometryCore/Public/Spatial/MeshAABBTree3.h b/Engine/Source/Runtime/GeometryCore/Public/Spatial/MeshAABBTree3.h
--- a/Engine/Source/Runtime/GeometryCore/Public/Spatial/MeshAABBTree3.h
+++ b/Engine/Source/Runtime/GeometryCore/Public/Spatial/MeshAABBTree3.h
@@ -50,6 +50,8 @@ namespace UE
 namespace Geometry
 {
 
+using namespace UE::Math;
+
 template <class TriangleMeshType>
 class TFastWindingTree;
 
@@ -730,7 +732,7 @@ public:
 			if (TransformF != nullptr)
 			{
 				TestMeshBounds = FAxisAlignedBox3d(TestMeshBounds,
-					[&](const FVector3<double>& P) { return TransformF(P); });
+					[&](const FVector3d& P) { return TransformF(P); });
 			}
 			if (box_box_intersect(RootIndex, TestMeshBounds) == false)
 			{
@@ -915,7 +917,7 @@ protected:
 		{
 			FAxisAlignedBox3d box = GetBox(iBox);
 			return FAxisAlignedBox3d(box, 
-				[&](const FVector3<double>& P) { return TransformF(P); });
+				[&](const FVector3d& P) { return TransformF(P); });
 		}
 		else
 		{
diff --git a/Engine/Source/Runtime/GeometryCore/Public/Spatial/PointHashGrid3.h b/Engine/Source/Runtime/GeometryCore/Public/Spatial/PointHashGrid3.h
--- a/Engine/Source/Runtime/GeometryCore/Public/Spatial/PointHashGrid3.h
+++ b/Engine/Source/Runtime/GeometryCore/Public/Spatial/PointHashGrid3.h
@@ -13,6 +13,8 @@ namespace UE
 namespace Geometry
 {
 
+using namespace UE::Math;
+
 /**
  * Hash Grid for values associated with 3D points.
  *
@@ -56,7 +58,7 @@ public:
 	 * @param Value the point/value to insert
 	 * @param Position the position associated with this value
 	 */
-	void InsertPoint(const PointDataType& Value, const FVector3<RealType>& Position)
+	void InsertPoint(const PointDataType& Value, const TVector<RealType>& Position)
 	{
 		FVector3i idx = Indexer.ToGrid(Position);
 		{
@@ -70,7 +72,7 @@ public:
 	 * @param Value the point/value to insert
 	 * @param Position the position associated with this value
 	 */
-	void InsertPointUnsafe(const PointDataType& Value, const FVector3<RealType>& Position)
+	void InsertPointUnsafe(const PointDataType& Value, const TVector<RealType>& Position)
 	{
 		FVector3i idx = Indexer.ToGrid(Position);
 		Hash.Add(idx, Value);
@@ -83,7 +85,7 @@ public:
 	 * @param Position the position associated with this value
 	 * @return true if the value existed at this position
 	 */
-	bool RemovePoint(const PointDataType& Value, const FVector3<RealType>& Position)
+	bool RemovePoint(const PointDataType& Value, const TVector<RealType>& Position)
 	{
 		FVector3i idx = Indexer.ToGrid(Position);
 		{
@@ -98,7 +100,7 @@ public:
 	 * @param Position the position associated with this value
 	 * @return true if the value existed at this position
 	 */
-	bool RemovePointUnsafe(const PointDataType& Value, const FVector3<RealType>& Position)
+	bool RemovePointUnsafe(const PointDataType& Value, const TVector<RealType>& Position)
 	{
 		FVector3i idx = Indexer.ToGrid(Position);
 		return Hash.RemoveSingle(idx, Value) > 0;
@@ -111,7 +113,7 @@ public:
 	 * 
 	 * @return true if the cell containing Position is empty
 	 */
-	bool IsCellEmpty(const FVector3<RealType>& Position)
+	bool IsCellEmpty(const TVector<RealType>& Position)
 	{
 		FVector3i Idx = Indexer.ToGrid(Position);
 		{
@@ -127,7 +129,7 @@ public:
 	 * 
 	 * @return true if the cell containing Position is empty
 	 */
-	bool IsCellEmptyUnsafe(const FVector3<RealType>& Position)
+	bool IsCellEmptyUnsafe(const TVector<RealType>& Position)
 	{
 		FVector3i Idx = Indexer.ToGrid(Position);
 		return !Hash.Contains(Idx);
@@ -140,7 +142,7 @@ public:
 	 * @param OldPosition the current position associated with this value
 	 * @param NewPosition the new position for this value
 	 */
-	void UpdatePoint(const PointDataType& Value, const FVector3<RealType>& OldPosition, const FVector3<RealType>& NewPosition)
+	void UpdatePoint(const PointDataType& Value, const TVector<RealType>& OldPosition, const TVector<RealType>& NewPosition)
 	{
 		FVector3i old_idx = Indexer.ToGrid(OldPosition);
 		FVector3i new_idx = Indexer.ToGrid(NewPosition);
@@ -168,7 +170,7 @@ public:
 	 * @param OldPosition the current position associated with this value
 	 * @param NewPosition the new position for this value
 	 */
-	void UpdatePointUnsafe(const PointDataType& Value, const FVector3<RealType>& OldPosition, const FVector3<RealType>& NewPosition)
+	void UpdatePointUnsafe(const PointDataType& Value, const TVector<RealType>& OldPosition, const TVector<RealType>& NewPosition)
 	{
 		FVector3i old_idx = Indexer.ToGrid(OldPosition);
 		FVector3i new_idx = Indexer.ToGrid(NewPosition);
@@ -192,7 +194,7 @@ public:
 	 * @return the found pair (Value,DistanceSqFunc(Value)), or (InvalidValue,MaxDouble) if not found
 	 */
 	TPair<PointDataType, RealType> FindNearestInRadius(
-		const FVector3<RealType>& QueryPoint, RealType Radius, 
+		const TVector<RealType>& QueryPoint, RealType Radius, 
 		TFunctionRef<RealType(const PointDataType&)> DistanceSqFunc,
 		TFunctionRef<bool(const PointDataType&)> IgnoreFunc = [](const PointDataType& data) { return false; }) const
 	{
@@ -201,8 +203,8 @@ public:
 			return TPair<PointDataType, RealType>(GetInvalidValue(), TNumericLimits<RealType>::Max());
 		}
 
-		FVector3i min_idx = Indexer.ToGrid(QueryPoint - Radius * FVector3<RealType>::One());
-		FVector3i max_idx = Indexer.ToGrid(QueryPoint + Radius * FVector3<RealType>::One());
+		FVector3i min_idx = Indexer.ToGrid(QueryPoint - Radius * TVector<RealType>::One());
+		FVector3i max_idx = Indexer.ToGrid(QueryPoint + Radius * TVector<RealType>::One());
 
 		RealType min_distsq = TNumericLimits<RealType>::Max();
 		PointDataType nearest = GetInvalidValue();
@@ -249,7 +251,7 @@ public:
 	 * @return the number of found points
 	 */
 	int FindPointsInBall(
-		const FVector3<RealType>& QueryPoint, RealType Radius,
+		const TVector<RealType>& QueryPoint, RealType Radius,
 		TFunctionRef<RealType(const PointDataType&)> DistanceSqFunc,
 		TArray<PointDataType>& ResultsOut,
 		TFunctionRef<bool(const PointDataType&)> IgnoreFunc = [](const PointDataType& data) { return false; }) const
@@ -259,8 +261,8 @@ public:
 			return 0;
 		}
 
-		FVector3i min_idx = Indexer.ToGrid(QueryPoint - Radius * FVector3<RealType>::One());
-		FVector3i max_idx = Indexer.ToGrid(QueryPoint + Radius * FVector3<RealType>::One());
+		FVector3i min_idx = Indexer.ToGrid(QueryPoint - Radius * TVector<RealType>::One());
+		FVector3i max_idx = Indexer.ToGrid(QueryPoint + Radius * TVector<RealType>::One());
 
 		RealType RadiusSquared = Radius * Radius;
 
diff --git a/Engine/Source/Runtime/GeometryCore/Public/SphereTypes.h b/Engine/Source/Runtime/GeometryCore/Public/SphereTypes.h
--- a/Engine/Source/Runtime/GeometryCore/Public/SphereTypes.h
+++ b/Engine/Source/Runtime/GeometryCore/Public/SphereTypes.h
@@ -10,6 +10,8 @@ namespace UE
 namespace Geometry
 {
 
+using namespace UE::Math;
+
 /*
  * 3D Sphere stored as Center point and Radius
  */
@@ -18,13 +20,13 @@ struct TSphere3
 {
 public:
 	/** Center of the sphere */
-	FVector3<T> Center = FVector3<T>::Zero();
+	TVector<T> Center = TVector<T>::Zero();
 	/** Radius of the sphere */
 	T Radius = (T)0;
 
 	TSphere3() = default;
 
-	TSphere3(const FVector3<T>& CenterIn, T RadiusIn)
+	TSphere3(const TVector<T>& CenterIn, T RadiusIn)
 		: Center(CenterIn), Radius(RadiusIn) {}
 
 	/** @return Diameter of sphere */
@@ -52,16 +54,16 @@ public:
 	}
 
 	/** @return true if Sphere contains given Point */
-	bool Contains(const FVector3<T>& Point) const
+	bool Contains(const TVector<T>& Point) const
 	{
-		T DistSqr = Center.DistanceSquared(Point);
+		T DistSqr = UE::Geometry::DistanceSquared(Center, Point);
 		return DistSqr <= Radius * Radius;
 	}
 
 	/** @return true if Sphere contains given OtherSphere */
 	bool Contains(const TSphere3<T>& OtherSphere) const
 	{
-		T CenterDist = Center.Distance(OtherSphere.Center);
+		T CenterDist = Distance(Center, OtherSphere.Center);
 		return (CenterDist + OtherSphere.Radius) <= Radius;
 	}
 
@@ -69,7 +71,7 @@ public:
 	/**
 	 * @return minimum squared distance from Point to Sphere surface for points outside sphere, 0 for points inside
 	 */
-	inline T DistanceSquared(const FVector3<T>& Point) const
+	inline T DistanceSquared(const TVector<T>& Point) const
 	{
 		const T PosDistance = TMathUtil<T>::Max(SignedDistance(Point), (T)0);
 		return PosDistance * PosDistance;
@@ -78,9 +80,9 @@ public:
 	/**
 	 * @return signed distance from Point to Sphere surface. Points inside sphere return negative distance.
 	 */
-	inline T SignedDistance(const FVector3<T>& Point) const
+	inline T SignedDistance(const TVector<T>& Point) const
 	{
-		return Center.Distance(Point) - Radius;
+		return UE::Geometry::Distance(Center, Point) - Radius;
 	}
 
 
diff --git a/Engine/Source/Runtime/GeometryCore/Public/TransformSequence.h b/Engine/Source/Runtime/GeometryCore/Public/TransformSequence.h
--- a/Engine/Source/Runtime/GeometryCore/Public/TransformSequence.h
+++ b/Engine/Source/Runtime/GeometryCore/Public/TransformSequence.h
@@ -10,6 +10,7 @@ namespace UE
 namespace Geometry
 {
 
+using namespace UE::Math;
 
 /**
  * TTransformSequence3 represents a sequence of 3D transforms. 
@@ -67,9 +68,9 @@ public:
 	/**
 	 * @return Cumulative scale across Transforms.
 	 */
-	FVector3<RealType> GetAccumulatedScale() const
+	TVector<RealType> GetAccumulatedScale() const
 	{
-		FVector3<RealType> FinalScale = FVector3<RealType>::One();
+		TVector<RealType> FinalScale = TVector<RealType>::One();
 		for (const TTransform3<RealType>& Transform : Transforms)
 		{
 			FinalScale = FinalScale * Transform.GetScale();
@@ -84,14 +85,14 @@ public:
 	{
 		for (TTransform3<RealType>& Transform : Transforms)
 		{
-			Transform.SetScale(FVector3<RealType>::One());
+			Transform.SetScale(TVector<RealType>::One());
 		}
 	}
 
 	/**
 	 * @return point P with transform sequence applied
 	 */
-	FVector3<RealType> TransformPosition(UE::Math::TVector<RealType> P) const
+	TVector<RealType> TransformPosition(UE::Math::TVector<RealType> P) const
 	{
 		for (const TTransform3<RealType>& Transform : Transforms)
 		{
@@ -103,7 +104,7 @@ public:
 	/**
 	 * @return point P with inverse transform sequence applied
 	 */
-	FVector3<RealType> InverseTransformPosition(UE::Math::TVector<RealType> P) const
+	TVector<RealType> InverseTransformPosition(UE::Math::TVector<RealType> P) const
 	{
 		int32 N = Transforms.Num();
 		for (int32 k = N - 1; k >= 0; k--)
@@ -116,7 +117,7 @@ public:
 	/**
 	 * @return Vector V with transform sequence applied
 	 */
-	FVector3<RealType> TransformVector(UE::Math::TVector<RealType> V) const
+	TVector<RealType> TransformVector(UE::Math::TVector<RealType> V) const
 	{
 		for (const TTransform3<RealType>& Transform : Transforms)
 		{
@@ -129,7 +130,7 @@ public:
 	/**
 	 * @return Normal with transform sequence applied
 	 */
-	FVector3<RealType> TransformNormal(UE::Math::TVector<RealType> Normal) const
+	TVector<RealType> TransformNormal(UE::Math::TVector<RealType> Normal) const
 	{
 		for (const TTransform3<RealType>& Transform : Transforms)
 		{
diff --git a/Engine/Source/Runtime/GeometryCore/Public/TransformTypes.h b/Engine/Source/Runtime/GeometryCore/Public/TransformTypes.h
--- a/Engine/Source/Runtime/GeometryCore/Public/TransformTypes.h
+++ b/Engine/Source/Runtime/GeometryCore/Public/TransformTypes.h
@@ -11,6 +11,8 @@ namespace UE
 namespace Geometry
 {
 
+using namespace UE::Math;
+
 /**
  * TTransform3 is a double/float templated version of standard UE FTransform.
  */
@@ -19,16 +21,16 @@ class TTransform3
 {
 protected:
 	TQuaternion<RealType> Rotation;
-	FVector3<RealType> Translation;
-	FVector3<RealType> Scale3D;
+	TVector<RealType> Translation;
+	TVector<RealType> Scale3D;
 
 public:
 
 	TTransform3()
 	{
 		Rotation = TQuaternion<RealType>::Identity();
-		Translation = FVector3<RealType>::Zero();
-		Scale3D = FVector3<RealType>::One();
+		Translation = TVector<RealType>::Zero();
+		Scale3D = TVector<RealType>::One();
 	}
 
 	TTransform3(const TQuaternion<RealType>& RotationIn, const UE::Math::TVector<RealType>& TranslationIn, const UE::Math::TVector<RealType>& ScaleIn)
@@ -42,21 +44,21 @@ public:
 	{
 		Rotation = RotationIn;
 		Translation = TranslationIn;
-		Scale3D = FVector3<RealType>::One();
+		Scale3D = TVector<RealType>::One();
 	}
 
 	explicit TTransform3(const UE::Math::TVector<RealType>& TranslationIn)
 	{
 		Rotation = TQuaternion<RealType>::Identity();
 		Translation = TranslationIn;
-		Scale3D = FVector3<RealType>::One();
+		Scale3D = TVector<RealType>::One();
 	}
 
 	explicit TTransform3(const FTransform& Transform)
 	{
 		Rotation = TQuaternion<RealType>(Transform.GetRotation());
-		Translation = FVector3<RealType>(Transform.GetTranslation());
-		Scale3D = FVector3<RealType>(Transform.GetScale3D());
+		Translation = TVector<RealType>(Transform.GetTranslation());
+		Scale3D = TVector<RealType>(Transform.GetScale3D());
 	}
 
 	/**
@@ -64,7 +66,7 @@ public:
 	 */
 	static TTransform3<RealType> Identity()
 	{
-		return TTransform3<RealType>(TQuaternion<RealType>::Identity(), FVector3<RealType>::Zero(), FVector3<RealType>::One());
+		return TTransform3<RealType>(TQuaternion<RealType>::Identity(), TVector<RealType>::Zero(), TVector<RealType>::One());
 	}
 
 	/**
@@ -95,7 +97,7 @@ public:
 	/**
 	 * @return Translation portion of transform
 	 */
-	const FVector3<RealType>& GetTranslation() const
+	const TVector<RealType>& GetTranslation() const
 	{
 		return Translation;
 	}
@@ -111,7 +113,7 @@ public:
 	/**
 	 * @return Scale portion of transform
 	 */
-	const FVector3<RealType>& GetScale() const
+	const TVector<RealType>& GetScale() const
 	{
 		return Scale3D;
 	}
@@ -148,8 +150,8 @@ public:
 	TTransform3<RealType> Inverse() const
 	{
 		TQuaternion<RealType> InvRotation = Rotation.Inverse();
-		FVector3<RealType> InvScale3D = GetSafeScaleReciprocal(Scale3D);
-		FVector3<RealType> InvTranslation = InvRotation * (InvScale3D * -Translation);
+		TVector<RealType> InvScale3D = GetSafeScaleReciprocal(Scale3D);
+		TVector<RealType> InvTranslation = InvRotation * (InvScale3D * -Translation);
 		return TTransform3<RealType>(InvRotation, InvTranslation, InvScale3D);
 	}
 
@@ -157,7 +159,7 @@ public:
 	/**
 	 * @return input point with QST transformation applied, ie QST(P) = Rotate(Scale*P) + Translate
 	 */
-	FVector3<RealType> TransformPosition(const FVector3<RealType>& P) const
+	TVector<RealType> TransformPosition(const TVector<RealType>& P) const
 	{
 		//Transform using QST is following
 		//QST(P) = Q.Rotate(S*P) + T where Q = quaternion, S = scale, T = translation
@@ -167,7 +169,7 @@ public:
 	/**
 	 * @return input vector with QS transformation applied, ie QS(V) = Rotate(Scale*V)
 	 */
-	FVector3<RealType> TransformVector(const UE::Math::TVector<RealType>& V) const
+	TVector<RealType> TransformVector(const UE::Math::TVector<RealType>& V) const
 	{
 		return Rotation * (Scale3D*V);
 	}
@@ -175,7 +177,7 @@ public:
 	/**
 	 * @return input vector with Q transformation applied, ie Q(V) = Rotate(V)
 	 */
-	FVector3<RealType> TransformVectorNoScale(const UE::Math::TVector<RealType>& V) const
+	TVector<RealType> TransformVectorNoScale(const UE::Math::TVector<RealType>& V) const
 	{
 		return Rotation * V;
 	}
@@ -185,12 +187,12 @@ public:
 	 * However 1/Scale requires special handling in case any component is near-zero.
 	 * @return input surface normal with transform applied.
 	 */
-	FVector3<RealType> TransformNormal(const UE::Math::TVector<RealType>& Normal) const
+	TVector<RealType> TransformNormal(const UE::Math::TVector<RealType>& Normal) const
 	{
 		// transform normal by a safe inverse scale + normalize, and a standard rotation
-		const FVector3<RealType>& S = Scale3D;
+		const TVector<RealType>& S = Scale3D;
 		RealType DetSign = FMathd::SignNonZero(S.X * S.Y * S.Z); // we only need to multiply by the sign of the determinant, rather than divide by it, since we normalize later anyway
-		FVector3<RealType> SafeInvS(S.Y*S.Z*DetSign, S.X*S.Z*DetSign, S.X*S.Y*DetSign);
+		TVector<RealType> SafeInvS(S.Y*S.Z*DetSign, S.X*S.Z*DetSign, S.X*S.Y*DetSign);
 		return TransformVectorNoScale( Normalized(SafeInvS*Normal) );
 	}
 
@@ -198,7 +200,7 @@ public:
 	/**
 	 * @return input vector with inverse-QST transformation applied, ie QSTinv(P) = InverseScale(InverseRotate(P - Translate))
 	 */
-	FVector3<RealType> InverseTransformPosition(const UE::Math::TVector<RealType> &P) const
+	TVector<RealType> InverseTransformPosition(const UE::Math::TVector<RealType> &P) const
 	{
 		return GetSafeScaleReciprocal(Scale3D) * Rotation.InverseMultiply(P - Translation);
 	}
@@ -206,7 +208,7 @@ public:
 	/**
 	 * @return input vector with inverse-QS transformation applied, ie QSinv(V) = InverseScale(InverseRotate(V))
 	 */
-	FVector3<RealType> InverseTransformVector(const UE::Math::TVector<RealType> &V) const
+	TVector<RealType> InverseTransformVector(const UE::Math::TVector<RealType> &V) const
 	{
 		return GetSafeScaleReciprocal(Scale3D) * Rotation.InverseMultiply(V);
 	}
@@ -215,7 +217,7 @@ public:
 	/**
 	 * @return input vector with inverse-Q transformation applied, ie Qinv(V) = InverseRotate(V)
 	 */
-	FVector3<RealType> InverseTransformVectorNoScale(const UE::Math::TVector<RealType> &V) const
+	TVector<RealType> InverseTransformVectorNoScale(const UE::Math::TVector<RealType> &V) const
 	{
 		return Rotation.InverseMultiply(V);
 	}
@@ -225,7 +227,7 @@ public:
 	 * Surface Normals are special, their inverse transform is InverseRotate( Normalize(Scale * Normal) ) )
 	 * @return input surface normal with inverse transform applied.
 	 */
-	FVector3<RealType> InverseTransformNormal(const UE::Math::TVector<RealType>& Normal) const
+	TVector<RealType> InverseTransformNormal(const UE::Math::TVector<RealType>& Normal) const
 	{
 		return InverseTransformVectorNoScale( Normalized(Scale3D*Normal) );
 	}
@@ -251,9 +253,9 @@ public:
 
 	
 
-	static FVector3<RealType> GetSafeScaleReciprocal(const FVector3<RealType>& InScale, RealType Tolerance = TMathUtil<RealType>::ZeroTolerance)
+	static TVector<RealType> GetSafeScaleReciprocal(const TVector<RealType>& InScale, RealType Tolerance = TMathUtil<RealType>::ZeroTolerance)
 	{
-		FVector3<RealType> SafeReciprocalScale;
+		TVector<RealType> SafeReciprocalScale;
 		if (TMathUtil<RealType>::Abs(InScale.X) <= Tolerance)
 		{
 			SafeReciprocalScale.X = (RealType)0;
@@ -293,44 +295,44 @@ public:
 	// case where this conversion is an issue
 
 	template<typename RealType2, TEMPLATE_REQUIRES(std::is_same<RealType, RealType2>::value == false)>
-	FVector3<RealType2> TransformPosition(const UE::Math::TVector<RealType2>& P) const
+	TVector<RealType2> TransformPosition(const UE::Math::TVector<RealType2>& P) const
 	{
-		return FVector3<RealType2>(TransformPosition(FVector3<RealType>(P)));
+		return TVector<RealType2>(TransformPosition(TVector<RealType>(P)));
 	}
 	template<typename RealType2, TEMPLATE_REQUIRES(std::is_same<RealType, RealType2>::value == false)>
-	FVector3<RealType2> TransformVector(const UE::Math::TVector<RealType2>& V) const
+	TVector<RealType2> TransformVector(const UE::Math::TVector<RealType2>& V) const
 	{
-		return FVector3<RealType2>(TransformVector(FVector3<RealType>(V)));
+		return TVector<RealType2>(TransformVector(TVector<RealType>(V)));
 	}
 	template<typename RealType2, TEMPLATE_REQUIRES(std::is_same<RealType, RealType2>::value == false)>
-	FVector3<RealType2> TransformVectorNoScale(const UE::Math::TVector<RealType2>& V) const
+	TVector<RealType2> TransformVectorNoScale(const UE::Math::TVector<RealType2>& V) const
 	{
-		return FVector3<RealType2>(TransformVectorNoScale(FVector3<RealType>(V)));
+		return TVector<RealType2>(TransformVectorNoScale(TVector<RealType>(V)));
 	}
 	template<typename RealType2, TEMPLATE_REQUIRES(std::is_same<RealType, RealType2>::value == false)>
-	FVector3<RealType2> TransformNormal(const UE::Math::TVector<RealType2>& V) const
+	TVector<RealType2> TransformNormal(const UE::Math::TVector<RealType2>& V) const
 	{
-		return FVector3<RealType2>(TransformNormal(FVector3<RealType>(V)));
+		return TVector<RealType2>(TransformNormal(TVector<RealType>(V)));
 	}
 	template<typename RealType2, TEMPLATE_REQUIRES(std::is_same<RealType, RealType2>::value == false)>
-	FVector3<RealType2> InverseTransformPosition(const UE::Math::TVector<RealType2>& P) const
+	TVector<RealType2> InverseTransformPosition(const UE::Math::TVector<RealType2>& P) const
 	{
-		return FVector3<RealType2>(InverseTransformPosition(FVector3<RealType>(P)));
+		return TVector<RealType2>(InverseTransformPosition(TVector<RealType>(P)));
 	}
 	template<typename RealType2, TEMPLATE_REQUIRES(std::is_same<RealType, RealType2>::value == false)>
-	FVector3<RealType2> InverseTransformVector(const UE::Math::TVector<RealType2>& V) const
+	TVector<RealType2> InverseTransformVector(const UE::Math::TVector<RealType2>& V) const
 	{
-		return FVector3<RealType2>(InverseTransformVector(FVector3<RealType>(V)));
+		return TVector<RealType2>(InverseTransformVector(TVector<RealType>(V)));
 	}
 	template<typename RealType2, TEMPLATE_REQUIRES(std::is_same<RealType, RealType2>::value == false)>
-	FVector3<RealType2> InverseTransformVectorNoScale(const UE::Math::TVector<RealType2>& V) const
+	TVector<RealType2> InverseTransformVectorNoScale(const UE::Math::TVector<RealType2>& V) const
 	{
-		return FVector3<RealType2>(InverseTransformVectorNoScale(FVector3<RealType>(V)));
+		return TVector<RealType2>(InverseTransformVectorNoScale(TVector<RealType>(V)));
 	}
 	template<typename RealType2, TEMPLATE_REQUIRES(std::is_same<RealType, RealType2>::value == false)>
-	FVector3<RealType2> InverseTransformNormal(const UE::Math::TVector<RealType2>& V) const
+	TVector<RealType2> InverseTransformNormal(const UE::Math::TVector<RealType2>& V) const
 	{
-		return FVector3<RealType2>(InverseTransformNormal(FVector3<RealType>(V)));
+		return TVector<RealType2>(InverseTransformNormal(TVector<RealType>(V)));
 	}
 
 };
diff --git a/Engine/Source/Runtime/GeometryCore/Public/TriangleTypes.h b/Engine/Source/Runtime/GeometryCore/Public/TriangleTypes.h
--- a/Engine/Source/Runtime/GeometryCore/Public/TriangleTypes.h
+++ b/Engine/Source/Runtime/GeometryCore/Public/TriangleTypes.h
@@ -57,12 +57,12 @@ struct TTriangle2
 		return Bary0 * V[0] + Bary1 * V[1] + Bary2 * V[2];
 	}
 
-	TVector2<RealType> BarycentricPoint(const FVector3<RealType>& BaryCoords) const
+	TVector2<RealType> BarycentricPoint(const TVector<RealType>& BaryCoords) const
 	{
 		return BaryCoords[0] * V[0] + BaryCoords[1] * V[1] + BaryCoords[2] * V[2];
 	}
 
-	FVector3<RealType> GetBarycentricCoords(const TVector2<RealType>& Point) const
+	TVector<RealType> GetBarycentricCoords(const TVector2<RealType>& Point) const
 	{
 		return VectorUtil::BarycentricCoords(Point, V[0], V[1], V[2]);
 	}
@@ -183,50 +183,50 @@ typedef TTriangle2<double> FTriangle2d;
 template<typename RealType>
 struct TTriangle3
 {
-	FVector3<RealType> V[3];
+	TVector<RealType> V[3];
 
 	TTriangle3() {}
 
-	TTriangle3(const FVector3<RealType>& V0, const FVector3<RealType>& V1, const FVector3<RealType>& V2)
+	TTriangle3(const TVector<RealType>& V0, const TVector<RealType>& V1, const TVector<RealType>& V2)
 	{
 		V[0] = V0;
 		V[1] = V1;
 		V[2] = V2;
 	}
 
-	TTriangle3(const FVector3<RealType> VIn[3])
+	TTriangle3(const TVector<RealType> VIn[3])
 	{
 		V[0] = VIn[0];
 		V[1] = VIn[1];
 		V[2] = VIn[2];
 	}
 
-	FVector3<RealType> BarycentricPoint(RealType Bary0, RealType Bary1, RealType Bary2) const
+	TVector<RealType> BarycentricPoint(RealType Bary0, RealType Bary1, RealType Bary2) const
 	{
 		return Bary0*V[0] + Bary1*V[1] + Bary2*V[2];
 	}
 
-	FVector3<RealType> BarycentricPoint(const FVector3<RealType> & BaryCoords) const
+	TVector<RealType> BarycentricPoint(const TVector<RealType> & BaryCoords) const
 	{
 		return BaryCoords[0]*V[0] + BaryCoords[1]*V[1] + BaryCoords[2]*V[2];
 	}
 
-	FVector3<RealType> GetBarycentricCoords(const FVector3<RealType>& Point) const
+	TVector<RealType> GetBarycentricCoords(const TVector<RealType>& Point) const
 	{
 		return VectorUtil::BarycentricCoords(Point, V[0], V[1], V[2]);
 	}
 
 	/** @return vector that is perpendicular to the plane of this triangle */
-	FVector3<RealType> Normal() const
+	TVector<RealType> Normal() const
 	{
 		return VectorUtil::Normal(V[0], V[1], V[2]);
 	}
 
 	/** @return centroid of this triangle */
-	FVector3<RealType> Centroid() const
+	TVector<RealType> Centroid() const
 	{
 		constexpr RealType f = 1.0 / 3.0;
-		return FVector3<RealType>(
+		return TVector<RealType>(
 			(V[0].X + V[1].X + V[2].X) * f,
 			(V[0].Y + V[1].Y + V[2].Y) * f,
 			(V[0].Z + V[1].Z + V[2].Z) * f
@@ -236,7 +236,7 @@ struct TTriangle3
 	/** grow the triangle around the centroid */
 	void Expand(RealType Delta)
 	{
-		FVector3<RealType> Centroid(Centroid());
+		TVector<RealType> Centroid(Centroid());
 		V[0] += Delta * ((V[0] - Centroid).Normalized());
 		V[1] += Delta * ((V[1] - Centroid).Normalized());
 		V[2] += Delta * ((V[2] - Centroid).Normalized());
diff --git a/Engine/Source/Runtime/GeometryCore/Public/Util/DynamicVector.h b/Engine/Source/Runtime/GeometryCore/Public/Util/DynamicVector.h
--- a/Engine/Source/Runtime/GeometryCore/Public/Util/DynamicVector.h
+++ b/Engine/Source/Runtime/GeometryCore/Public/Util/DynamicVector.h
@@ -16,6 +16,8 @@ namespace UE
 namespace Geometry
 {
 
+using namespace UE::Math;
+
 /*
  * Blocked array with fixed, power-of-two sized blocks.
  *
@@ -652,14 +654,14 @@ public:
 	{
 		return Data[TopIndex * N + SubIndex];
 	}
-	inline void SetVector2(unsigned int TopIndex, const UE::Math::TVector2<Type>& V)
+	inline void SetVector2(unsigned int TopIndex, const TVector2<Type>& V)
 	{
 		check(N >= 2);
 		unsigned int i = TopIndex * N;
 		Data[i] = V.X;
 		Data[i + 1] = V.Y;
 	}
-	inline void SetVector3(unsigned int TopIndex, const FVector3<Type>& V)
+	inline void SetVector3(unsigned int TopIndex, const TVector<Type>& V)
 	{
 		check(N >= 3);
 		unsigned int i = TopIndex * N;
@@ -667,17 +669,17 @@ public:
 		Data[i + 1] = V.Y;
 		Data[i + 2] = V.Z;
 	}
-	inline UE::Math::TVector2<Type> AsVector2(unsigned int TopIndex) const
+	inline TVector2<Type> AsVector2(unsigned int TopIndex) const
 	{
 		check(N >= 2);
-		return UE::Math::TVector2<Type>(
+		return TVector2<Type>(
 			Data[TopIndex * N + 0],
 			Data[TopIndex * N + 1]);
 	}
-	inline FVector3<Type> AsVector3(unsigned int TopIndex) const
+	inline TVector<Type> AsVector3(unsigned int TopIndex) const
 	{
 		check(N >= 3);
-		return FVector3<Type>(
+		return TVector<Type>(
 			Data[TopIndex * N + 0],
 			Data[TopIndex * N + 1],
 			Data[TopIndex * N + 2]);
diff --git a/Engine/Source/Runtime/GeometryCore/Public/Util/ElementLinearization.h b/Engine/Source/Runtime/GeometryCore/Public/Util/ElementLinearization.h
--- a/Engine/Source/Runtime/GeometryCore/Public/Util/ElementLinearization.h
+++ b/Engine/Source/Runtime/GeometryCore/Public/Util/ElementLinearization.h
@@ -11,6 +11,8 @@ namespace UE
 namespace Geometry
 {
 
+using namespace UE::Math;
+
 /**
 * FElementLinearization maps a potentially-sparse index list into a linear array.
 * Used linearize things like VtxIds of a mesh as a single array and allow bidirectional mapping between array offset and mesh VtxId.
@@ -176,19 +178,19 @@ public:
 		ZVector[i] = Value;
 	}
 
-	void SetXYZ(int32 i, const FVector3<RealType>& Value)
+	void SetXYZ(int32 i, const TVector<RealType>& Value)
 	{
 		XVector[i] = Value.X;
 		YVector[i] = Value.Y;
 		ZVector[i] = Value.Z;
 	}
 
-	FVector3<RealType> Get(int32 i)
+	TVector<RealType> Get(int32 i)
 	{
-		return FVector3<RealType>(XVector[i], YVector[i], ZVector[i]);
+		return TVector<RealType>(XVector[i], YVector[i], ZVector[i]);
 	}
 
-	void Set(int32 i, const FVector3<RealType>& Value)
+	void Set(int32 i, const TVector<RealType>& Value)
 	{
 		XVector[i] = Value.X;
 		YVector[i] = Value.Y;
diff --git a/Engine/Source/Runtime/GeometryCore/Public/Util/GridIndexing3.h b/Engine/Source/Runtime/GeometryCore/Public/Util/GridIndexing3.h
--- a/Engine/Source/Runtime/GeometryCore/Public/Util/GridIndexing3.h
+++ b/Engine/Source/Runtime/GeometryCore/Public/Util/GridIndexing3.h
@@ -13,6 +13,8 @@ namespace UE
 namespace Geometry
 {
 
+using namespace UE::Math;
+
 /**
  * Convert between integer grid coordinates and scaled real-valued coordinates (ie assumes integer grid origin == real origin)
  */
@@ -32,7 +34,7 @@ struct TScaleGridIndexer3
 	}
 
 	/** Convert real-valued point to integer grid coordinates */
-	inline FVector3i ToGrid(const FVector3<RealType>& P) const
+	inline FVector3i ToGrid(const TVector<RealType>& P) const
 	{
 		return FVector3i(
 			(int)TMathUtil<RealType>::Floor(P.X / CellSize),
@@ -41,9 +43,9 @@ struct TScaleGridIndexer3
 	}
 
 	/** Convert integer grid coordinates to real-valued point */
-	inline FVector3<RealType> FromGrid(const FVector3i& GridPoint) const
+	inline TVector<RealType> FromGrid(const FVector3i& GridPoint) const
 	{
-		return FVector3<RealType>(GridPoint.X*CellSize, GridPoint.Y*CellSize, GridPoint.Z*CellSize);
+		return TVector<RealType>(GridPoint.X*CellSize, GridPoint.Y*CellSize, GridPoint.Z*CellSize);
 	}
 };
 typedef TScaleGridIndexer3<float> FScaleGridIndexer3f;
@@ -60,22 +62,22 @@ struct TShiftGridIndexer3
 	/** Real-valued size of an integer grid cell */
 	RealType CellSize;
 	/** Real-valued origin of grid, position of integer grid origin */
-	FVector3<RealType> Origin;
+	TVector<RealType> Origin;
 
 	TShiftGridIndexer3() 
-		: CellSize((RealType)1), Origin(FVector3<RealType>::Zero())
+		: CellSize((RealType)1), Origin(TVector<RealType>::Zero())
 	{
 		ensure(CellSize >= TMathUtil<RealType>::ZeroTolerance);
 	}
 
-	TShiftGridIndexer3(const FVector3<RealType>& origin, RealType cellSize)
+	TShiftGridIndexer3(const TVector<RealType>& origin, RealType cellSize)
 		: CellSize(cellSize), Origin(origin)
 	{
 		ensure(CellSize >= TMathUtil<RealType>::ZeroTolerance);
 	}
 
 	/** Convert real-valued point to integer grid coordinates */
-	inline FVector3i ToGrid(const FVector3<RealType>& point) const
+	inline FVector3i ToGrid(const TVector<RealType>& point) const
 	{
 		return FVector3i(
 			(int)TMathUtil<RealType>::Floor((point.X - Origin.X) / CellSize),
@@ -84,18 +86,18 @@ struct TShiftGridIndexer3
 	}
 
 	/** Convert integer grid coordinates to real-valued point */
-	inline FVector3<RealType> FromGrid(const FVector3i& gridpoint) const
+	inline TVector<RealType> FromGrid(const FVector3i& gridpoint) const
 	{
-		return FVector3<RealType>(
+		return TVector<RealType>(
 			((RealType)gridpoint.X * CellSize) + Origin.X,
 			((RealType)gridpoint.Y * CellSize) + Origin.Y,
 			((RealType)gridpoint.Z * CellSize) + Origin.Z);
 	}
 
 	/** Convert real-valued grid coordinates to real-valued point */
-	inline FVector3<RealType> FromGrid(const FVector3<RealType>& RealGridPoint)  const
+	inline TVector<RealType> FromGrid(const TVector<RealType>& RealGridPoint)  const
 	{
-		return FVector3<RealType>(
+		return TVector<RealType>(
 			((RealType)RealGridPoint.X * CellSize) + Origin.X,
 			((RealType)RealGridPoint.Y * CellSize) + Origin.Y,
 			((RealType)RealGridPoint.Z * CellSize) + Origin.Z);
@@ -111,7 +113,7 @@ struct TBoundsGridIndexer3 : public TShiftGridIndexer3<RealType>
 	using TShiftGridIndexer3<RealType>::CellSize;
 	using TShiftGridIndexer3<RealType>::Origin;
 
-	FVector3<RealType> BoundsMax;
+	TVector<RealType> BoundsMax;
 
 	TBoundsGridIndexer3(const TAxisAlignedBox3<RealType>& Bounds, RealType CellSize)
 		: TShiftGridIndexer3<RealType>(Bounds.Min, CellSize),
@@ -121,11 +123,11 @@ struct TBoundsGridIndexer3 : public TShiftGridIndexer3<RealType>
 	FVector3i GridResolution() const
 	{
 		const RealType InvCellSize = 1.0 / CellSize;
-		const FVector3<RealType> Extents = BoundsMax - Origin;
+		const TVector<RealType> Extents = BoundsMax - Origin;
 		return CeilInt(Extents * InvCellSize);
 	}
 
-	static FVector3i CeilInt(const FVector3<RealType>& V)
+	static FVector3i CeilInt(const TVector<RealType>& V)
 	{
 		return FVector3i{ (int)TMathUtil<RealType>::Ceil(V[0]),
 			(int)TMathUtil<RealType>::Ceil(V[1]),
diff --git a/Engine/Source/Runtime/GeometryCore/Public/Util/IndexUtil.h b/Engine/Source/Runtime/GeometryCore/Public/Util/IndexUtil.h
--- a/Engine/Source/Runtime/GeometryCore/Public/Util/IndexUtil.h
+++ b/Engine/Source/Runtime/GeometryCore/Public/Util/IndexUtil.h
@@ -12,6 +12,7 @@
 
 namespace IndexUtil
 {
+	using namespace UE::Math;
 	using namespace UE::Geometry;
 	using UE::Geometry::FIndex2i;
 
@@ -252,7 +253,7 @@ namespace IndexUtil
 	 * Replace Val with MapFunc[Val] using index operator
 	 */
 	template<typename T, typename Func>
-	void ApplyMap(FVector3<T> & Val, Func MapFunc)
+	void ApplyMap(TVector<T> & Val, Func MapFunc)
 	{
 		Val[0] = MapFunc[Val[0]];
 		Val[1] = MapFunc[Val[1]];
@@ -274,9 +275,9 @@ namespace IndexUtil
 	 * @return MapFunc[Val] using index operator
 	 */
 	template<typename T, typename Func>
-	FVector3<T> ApplyMap(const FVector3<T> & Val, Func MapFunc)
+	TVector<T> ApplyMap(const TVector<T> & Val, Func MapFunc)
 	{
-		return FVector3<T>(MapFunc[Val[0]], MapFunc[Val[1]], MapFunc[Val[2]]);
+		return TVector<T>(MapFunc[Val[0]], MapFunc[Val[1]], MapFunc[Val[2]]);
 	}
 	
 	/**
diff --git a/Engine/Source/Runtime/GeometryCore/Public/VectorTypes.h b/Engine/Source/Runtime/GeometryCore/Public/VectorTypes.h
--- a/Engine/Source/Runtime/GeometryCore/Public/VectorTypes.h
+++ b/Engine/Source/Runtime/GeometryCore/Public/VectorTypes.h
@@ -208,240 +208,6 @@ UE::Math::TVector2<T> Lerp(const UE::Math::TVector2<T>& A, const UE::Math::TVect
 }
 
 
-/**
- * Templated 3D Vector. Ported from g3Sharp library, with the intention of 
- * maintaining compatibility with existing g3Sharp code. Has an API
- * similar to WildMagic, GTEngine, Eigen, etc. 
- * 
- * Convenience typedefs for FVector3f/FVector3d are defined, and
- * should be preferentially used over the base template type
- * 
- * @todo Possibly can be replaced/merged with Chaos TVector<T,N>
- */
-template <typename T>
-struct FVector3 : public UE::Math::TVector<T>
-{
-	using UE::Math::TVector<T>::TVector;
-	using UE::Math::TVector<T>::X;
-	using UE::Math::TVector<T>::Y;
-	using UE::Math::TVector<T>::Z;
-
-	FVector3() : UE::Math::TVector<T>((T)0)
-	{
-	}
-
-	FVector3(const UE::Math::TVector<T>& Vec)
-	{
-		X = Vec.X;
-		Y = Vec.Y;
-		Z = Vec.Z;
-	}
-
-	template<typename OtherType, TEMPLATE_REQUIRES(!std::is_same<T, OtherType>::value)>
-	explicit FVector3(const UE::Math::TVector<OtherType>& Vec)
-	{
-		X = (T)Vec.X;
-		Y = (T)Vec.Y;
-		Z = (T)Vec.Z;
-	}
-
-	FVector3& operator=(const UE::Math::TVector<T>& Vec)
-	{
-		X = Vec.X;
-		Y = Vec.Y;
-		Z = Vec.Z;
-		return *this;
-	}
-
-	explicit operator FVector3f() const
-	{
-		return FVector3f((float)X, (float)Y, (float)Z);
-	}
-
-	explicit operator FVector3d() const
-	{
-		return FVector3d((double)X, (double)Y, (double)Z);
-	}
-
-
-	explicit operator FLinearColor() const
-	{
-		return FLinearColor((float)X, (float)Y, (float)Z);
-	}
-	FVector3(const FLinearColor& Color)
-	{
-		X = (T)Color.R;
-		Y = (T)Color.G;
-		Z = (T)Color.B;
-	}
-
-
-	explicit FVector3(const T* Data)
-	{
-		X = Data[0];
-		Y = Data[1];
-		Z = Data[2];
-	}
-
-
-	explicit operator const T*() const
-	{
-		return &X;
-	};
-	explicit operator T*()
-	{
-		return &X;
-	}
-
-
-	static FVector3<T> Zero()
-	{
-		return FVector3<T>((T)0, (T)0, (T)0);
-	}
-	static FVector3<T> One()
-	{
-		return FVector3<T>((T)1, (T)1, (T)1);
-	}
-	static FVector3<T> UnitX()
-	{
-		return FVector3<T>((T)1, (T)0, (T)0);
-	}
-	static FVector3<T> UnitY()
-	{
-		return FVector3<T>((T)0, (T)1, (T)0);
-	}
-	static FVector3<T> UnitZ()
-	{
-		return FVector3<T>((T)0, (T)0, (T)1);
-	}
-	static FVector3<T> MaxVector()
-	{
-		return FVector3<T>(TNumericLimits<T>::Max(), TNumericLimits<T>::Max(), TNumericLimits<T>::Max());
-	}
-
-	T Length() const
-	{
-		return TMathUtil<T>::Sqrt(X * X + Y * Y + Z * Z);
-	}
-	T SquaredLength() const
-	{
-		return X * X + Y * Y + Z * Z;
-	}
-
-	T Distance(const UE::Math::TVector<T>& V2) const
-	{
-		T dx = V2.X - X;
-		T dy = V2.Y - Y;
-		T dz = V2.Z - Z;
-		return TMathUtil<T>::Sqrt(dx * dx + dy * dy + dz * dz);
-	}
-	T DistanceSquared(const UE::Math::TVector<T>& V2) const
-	{
-		T dx = V2.X - X;
-		T dy = V2.Y - Y;
-		T dz = V2.Z - Z;
-		return dx * dx + dy * dy + dz * dz;
-	}
-
-	FVector3<T> operator-() const
-	{
-		return FVector3<T>(-X, -Y, -Z);
-	}
-
-	FVector3<T> operator+(const UE::Math::TVector<T>& V2) const
-	{
-		return FVector3<T>(X + V2.X, Y + V2.Y, Z + V2.Z);
-	}
-
-	FVector3<T> operator-(const UE::Math::TVector<T>& V2) const
-	{
-		return FVector3<T>(X - V2.X, Y - V2.Y, Z - V2.Z);
-	}
-
-	FVector3<T> operator+(const T& Scalar) const
-	{
-		return FVector3<T>(X + Scalar, Y + Scalar, Z + Scalar);
-	}
-
-	FVector3<T> operator-(const T& Scalar) const
-	{
-		return FVector3<T>(X - Scalar, Y - Scalar, Z - Scalar);
-	}
-
-	FVector3<T> operator*(const T& Scalar) const
-	{
-		return FVector3<T>(X * Scalar, Y * Scalar, Z * Scalar);
-	}
-
-	FVector3<T> operator*(const UE::Math::TVector<T>& V2) const // component-wise
-	{
-		return FVector3<T>(X * V2.X, Y * V2.Y, Z * V2.Z);
-	}
-
-	FVector3<T> operator/(const T& Scalar) const
-	{
-		return FVector3<T>(X / Scalar, Y / Scalar, Z / Scalar);
-	}
-
-	FVector3<T> operator/(const UE::Math::TVector<T>& V2) const // component-wise
-	{
-		return FVector3<T>(X / V2.X, Y / V2.Y, Z / V2.Z);
-	}
-
-	template<typename RealType2, TEMPLATE_REQUIRES(std::is_floating_point<RealType2>::value)>
-	FVector3<T> operator*(const RealType2& Scalar) const
-	{
-		return FVector3<T>(X * (T)Scalar, Y * (T)Scalar, Z * (T)Scalar);
-	}
-
-	FVector3<T>& operator+=(const UE::Math::TVector<T>& V2)
-	{
-		X += V2.X;
-		Y += V2.Y;
-		Z += V2.Z;
-		return *this;
-	}
-
-	FVector3<T>& operator-=(const UE::Math::TVector<T>& V2)
-	{
-		X -= V2.X;
-		Y -= V2.Y;
-		Z -= V2.Z;
-		return *this;
-	}
-
-	FVector3<T>& operator*=(const T& Scalar)
-	{
-		X *= Scalar;
-		Y *= Scalar;
-		Z *= Scalar;
-		return *this;
-	}
-
-	FVector3<T>& operator/=(const T& Scalar)
-	{
-		X /= Scalar;
-		Y /= Scalar;
-		Z /= Scalar;
-		return *this;
-	}
-
-	T Dot(const UE::Math::TVector<T>& V2) const
-	{
-		return X * V2.X + Y * V2.Y + Z * V2.Z;
-	}
-
-	FVector3<T> Cross(const UE::Math::TVector<T>& V2) const
-	{
-		return FVector3(
-			Y * V2.Z - Z * V2.Y,
-			Z * V2.X - X * V2.Z,
-			X * V2.Y - Y * V2.X);
-	}
-};
-
-
-
 
 
 /** @return unit vector along axis X=0, Y=1, Z=2 */
@@ -588,17 +354,17 @@ constexpr UE::Math::TVector2<T> GetYZ(const UE::Math::TVector<T>& V)
 
 
 template<typename T>
-constexpr FVector3<T> Min(const UE::Math::TVector<T>& V0, const UE::Math::TVector<T>& V1)
+constexpr UE::Math::TVector<T> Min(const UE::Math::TVector<T>& V0, const UE::Math::TVector<T>& V1)
 {
-	return FVector3<T>(TMathUtil<T>::Min(V0.X, V1.X),
+	return UE::Math::TVector<T>(TMathUtil<T>::Min(V0.X, V1.X),
 		TMathUtil<T>::Min(V0.Y, V1.Y),
 		TMathUtil<T>::Min(V0.Z, V1.Z));
 }
 
 template<typename T>
-constexpr FVector3<T> Max(const UE::Math::TVector<T>& V0, const UE::Math::TVector<T>& V1)
+constexpr UE::Math::TVector<T> Max(const UE::Math::TVector<T>& V0, const UE::Math::TVector<T>& V1)
 {
-	return FVector3<T>(TMathUtil<T>::Max(V0.X, V1.X),
+	return UE::Math::TVector<T>(TMathUtil<T>::Max(V0.X, V1.X),
 		TMathUtil<T>::Max(V0.Y, V1.Y),
 		TMathUtil<T>::Max(V0.Z, V1.Z));
 }
@@ -680,290 +446,18 @@ UE::Math::TVector<T> Blend3(const UE::Math::TVector<T>& A, const UE::Math::TVect
 		WeightA * A.Z + WeightB * B.Z + WeightC * C.Z);
 }
 
-template <typename RealType>
-inline FVector3<RealType> operator*(RealType Scalar, const FVector3<RealType>& V)
-{
-	return FVector3<RealType>(Scalar * V.X, Scalar * V.Y, Scalar * V.Z);
-}
-
-// allow float*Vector3<double> and double*Vector3<float>
-template <typename RealType, typename RealType2, TEMPLATE_REQUIRES(std::is_floating_point<RealType2>::value)>
-inline FVector3<RealType> operator*(RealType2 Scalar, const FVector3<RealType>& V)
-{
-	return FVector3<RealType>((RealType)Scalar * V.X, (RealType)Scalar * V.Y, (RealType)Scalar * V.Z);
-}
 
 template <typename RealType>
-std::ostream& operator<<(std::ostream& os, const FVector3<RealType>& Vec)
+std::ostream& operator<<(std::ostream& os, const UE::Math::TVector<RealType>& Vec)
 {
 	os << Vec.X << " " << Vec.Y << " " << Vec.Z;
 	return os;
 }
 
-template <typename T>
-FORCEINLINE uint32 GetTypeHash(const FVector3<T>& Vector)
-{
-	// (this is how FIntVector and all the other FVectors do their hash functions)
-	// Note: this assumes there's no padding that could contain uncompared data.
-	return FCrc::MemCrc_DEPRECATED(&Vector, sizeof(FVector3<T>));
-}
-
-
-
-
-
-
-
-
-
-
-
-
-template <typename T>
-struct GVector4
-{
-	T X{}, Y{}, Z{}, W{};
-
-	constexpr GVector4()
-		: X(0), Y(0), Z(0), W(0)
-	{
-	}
-
-	constexpr GVector4(T ValX, T ValY, T ValZ, T ValW)
-		: X(ValX), Y(ValY), Z(ValZ), W(ValW)
-	{
-	}
-
-	constexpr GVector4(const GVector4& Vec) = default;
-
-	template<typename RealType2>
-	explicit constexpr GVector4(const GVector4<RealType2>& Vec)
-		: X((T)Vec.X), Y((T)Vec.Y), Z((T)Vec.Z), W((T)Vec.W)
-	{
-	}
-
-	static GVector4<T> Zero()
-	{
-		return GVector4<T>((T)0, (T)0, (T)0, (T)0);
-	}
-	static GVector4<T> One()
-	{
-		return GVector4<T>((T)1, (T)1, (T)1, (T)1);
-	}
-
-
-	GVector4<T>& operator=(const GVector4<T>& V2)
-	{
-		X = V2.X;
-		Y = V2.Y;
-		Z = V2.Z;
-		W = V2.W;
-		return *this;
-	}
-
-	T& operator[](int Idx)
-	{
-		return (&X)[Idx];
-	}
-	const T& operator[](int Idx) const
-	{
-		return (&X)[Idx];
-	}
-
-	constexpr GVector4<T> operator-() const
-	{
-		return GVector4<T>(-X, -Y, -Z, -W);
-	}
-
-	constexpr GVector4<T> operator+(const GVector4<T>& V2) const
-	{
-		return GVector4<T>(X + V2.X, Y + V2.Y, Z + V2.Z, W + V2.W);
-	}
-
-	constexpr GVector4<T> operator-(const GVector4<T>& V2) const
-	{
-		return GVector4<T>(X - V2.X, Y - V2.Y, Z - V2.Z, W - V2.W);
-	}
-
-	constexpr GVector4<T> operator+(const T& Scalar) const
-	{
-		return GVector4<T>(X + Scalar, Y + Scalar, Z + Scalar, W + Scalar);
-	}
-
-	constexpr GVector4<T> operator-(const T& Scalar) const
-	{
-		return GVector4<T>(X - Scalar, Y - Scalar, Z - Scalar, W - Scalar);
-	}
-
-	constexpr GVector4<T> operator*(const T& Scalar) const
-	{
-		return GVector4<T>(X * Scalar, Y * Scalar, Z * Scalar, W * Scalar);
-	}
-
-	template<typename RealType2>
-	constexpr GVector4<T> operator*(const RealType2& Scalar) const
-	{
-		return GVector4<T>(X * (T)Scalar, Y * (T)Scalar, Z * (T)Scalar, W * (T)Scalar);
-	}
-
-	constexpr GVector4<T> operator/(const T& Scalar) const
-	{
-		return GVector4<T>(X / Scalar, Y / Scalar, Z / Scalar, W / Scalar);
-	}
-
-	constexpr GVector4<T>& operator+=(const GVector4<T>& V2)
-	{
-		X += V2.X;
-		Y += V2.Y;
-		Z += V2.Z;
-		W += V2.W;
-		return *this;
-	}
-
-	constexpr GVector4<T>& operator-=(const GVector4<T>& V2)
-	{
-		X -= V2.X;
-		Y -= V2.Y;
-		Z -= V2.Z;
-		W -= V2.W;
-		return *this;
-	}
-
-	constexpr GVector4<T>& operator*=(const T& Scalar)
-	{
-		X *= Scalar;
-		Y *= Scalar;
-		Z *= Scalar;
-		W *= Scalar;
-		return *this;
-	}
-
-	constexpr GVector4<T>& operator/=(const T& Scalar)
-	{
-		X /= Scalar;
-		Y /= Scalar;
-		Z /= Scalar;
-		W /= Scalar;
-		return *this;
-	}
-
-	T Dot(const GVector4<T>& V2) const
-	{
-		return X * V2.X + Y * V2.Y + Z * V2.Z + W * V2.W;
-	}
-
-	constexpr bool operator==(const GVector4<T>& Other) const
-	{
-		return X == Other.X && Y == Other.Y && Z == Other.Z && W == Other.W;
-	}
-
-	constexpr bool operator!=(const GVector4<T>& Other) const
-	{
-		return X != Other.X || Y != Other.Y || Z != Other.Z || W != Other.W;
-	}
-
-	/**
-	 * Serialization operator for GVector4.
-	 *
-	 * @param Ar Archive to serialize with.
-	 * @param Vec Vector to serialize.
-	 * @returns Passing down serializing archive.
-	 */
-	friend FArchive& operator<<(FArchive& Ar, GVector4& Vec)
-	{
-		Vec.Serialize(Ar);
-		return Ar;
-	}
-
-	/** Serialize GVector4 to an archive. */
-	void Serialize(FArchive& Ar)
-	{
-		Ar << X;
-		Ar << Y;
-		Ar << Z;
-		Ar << W;
-	}
-};
-
-template <typename T>
-constexpr bool IsNormalized(const GVector4<T>& Vector, const T Tolerance = TMathUtil<T>::ZeroTolerance)
-{
-	return TMathUtil<T>::Abs((Vector.X*Vector.X + Vector.Y*Vector.Y + Vector.Z*Vector.Z + Vector.W*Vector.W) - 1) < Tolerance;
-}
-
-template<typename T>
-T Normalize(GVector4<T>& Vector, const T Epsilon = 0)
-{
-	T length = Vector.Length();
-	if (length > Epsilon)
-	{
-		T invLength = ((T)1) / length;
-		Vector.X *= invLength;
-		Vector.Y *= invLength;
-		Vector.Z *= invLength;
-		Vector.W *= invLength;
-		return length;
-	}
-	Vector.X = Vector.Y = Vector.Z = Vector.W = (T)0;
-	return (T)0;
-}
-
-template<typename T>
-GVector4<T> Normalized(const GVector4<T>& Vector, const T Epsilon = 0)
-{
-	T length = Vector.Length();
-	if (length > Epsilon)
-	{
-		T invLength = ((T)1) / length;
-		return GVector4<T>(Vector.X*invLength, Vector.Y*invLength, Vector.Z*invLength, Vector.W*invLength);
-	}
-	return GVector4<T>::Zero();
-}
-
-template<typename T>
-constexpr FVector3<T> GetXYZ(const GVector4<T>& V)
-{
-	return FVector3<T>(V.X, V.Y, V.Z);
-}
-
-template<typename T>
-GVector4<T> Lerp(const GVector4<T>& A, const GVector4<T>& B, T Alpha)
-{
-	T OneMinusAlpha = (T)1 - Alpha;
-	return GVector4<T>(OneMinusAlpha * A.X + Alpha * B.X,
-		OneMinusAlpha * A.Y + Alpha * B.Y,
-		OneMinusAlpha * A.Z + Alpha * B.Z,
-		OneMinusAlpha * A.W + Alpha * B.W);
-}
 
-template<typename T>
-GVector4<T> Blend3(const GVector4<T>& A, const GVector4<T>& B, const GVector4<T>& C, const T& WeightA, const T& WeightB, const T& WeightC)
-{
-	return GVector4<T>(
-		WeightA * A.X + WeightB * B.X + WeightC * C.X,
-		WeightA * A.Y + WeightB * B.Y + WeightC * C.Y,
-		WeightA * A.Z + WeightB * B.Z + WeightC * C.Z,
-		WeightA * A.W + WeightB * B.W + WeightC * C.W);
-}
 
-template <typename RealType, TEMPLATE_REQUIRES(std::is_arithmetic<RealType>::value)>
-inline GVector4<RealType> operator*(RealType Scalar, const GVector4<RealType>& V)
-{
-	return GVector4<RealType>(Scalar * V.X, Scalar * V.Y, Scalar * V.Z, Scalar * V.W);
-}
 
-// allow float*Vector4<double> and double*Vector4<float>
-template <typename RealType, typename RealType2, TEMPLATE_REQUIRES(std::is_arithmetic<RealType2>::value)>
-inline GVector4<RealType> operator*(RealType2 Scalar, const GVector4<RealType>& V)
-{
-	return GVector4<RealType>((RealType)Scalar * V.X, (RealType)Scalar * V.Y, (RealType)Scalar * V.Z, (RealType)Scalar * V.W);
-}
 
-template<typename T>
-FLinearColor ToLinearColor(const GVector4<T>& V)
-{
-	return FLinearColor((float)V.X, (float)V.Y, (float)V.Z, (float)V.W);
-}
 
 template<typename T>
 FLinearColor ToLinearColor(const UE::Math::TVector4<T>& V)
@@ -977,16 +471,6 @@ UE::Math::TVector4<T> ToVector4(const FLinearColor& Color)
 	return UE::Math::TVector4<T>( (T)Color.R, (T)Color.G, (T)Color.B, (T)Color.A );
 }
 
-template<typename T>
-FColor ToFColor(const GVector4<T>& V)
-{
-	return FColor(
-		FMathf::Clamp((int)((float)V.X*255.0f), 0, 255),
-		FMathf::Clamp((int)((float)V.Y*255.0f), 0, 255),
-		FMathf::Clamp((int)((float)V.Z*255.0f), 0, 255),
-		FMathf::Clamp((int)((float)V.W*255.0f), 0, 255));
-}
-
 template<typename T>
 constexpr FColor ToFColor(const UE::Math::TVector4<T>& Vector)
 {
@@ -996,20 +480,6 @@ constexpr FColor ToFColor(const UE::Math::TVector4<T>& Vector)
 		FMathf::Clamp((int)((float)Vector.Z * 255.0f), 0, 255));
 }
 
-template <typename RealType>
-std::ostream& operator<<(std::ostream& os, const GVector4<RealType>& Vec)
-{
-	os << Vec.X << " " << Vec.Y << " " << Vec.Z << " " << Vec.W;
-	return os;
-}
-
-template <typename T>
-FORCEINLINE uint32 GetTypeHash(const GVector4<T>& Vector)
-{
-	// (this is how FIntVector and all the other FVectors do their hash functions)
-	// Note: this assumes there's no padding that could contain uncompared data.
-	return FCrc::MemCrc_DEPRECATED(&Vector, sizeof(GVector4<T>));
-}
 
 } // end namespace UE::Geometry
 
@@ -1034,19 +504,3 @@ std::ostream& operator<<(std::ostream& os, const TVector4<RealType>& Vec)
 
 } // end namespace UE
 
-
-
-//typedef UE::Geometry::FVector2<float> FVector2f;
-//typedef UE::Geometry::FVector2<double> FVector2d;
-
-//typedef UE::Geometry::FVector3<float> FVector3f;
-//typedef UE::Geometry::FVector3<double> FVector3d;
-
-//typedef UE::Geometry::GVector4<float> FVector4f;
-//typedef UE::Geometry::GVector4<double> FVector4d;
-// 
-//template<> struct TCanBulkSerialize<FVector2f> { enum { Value = true }; };
-//template<> struct TCanBulkSerialize<FVector2d> { enum { Value = true }; };
-
-template<> struct TCanBulkSerialize<UE::Geometry::GVector4<float>> { enum { Value = true }; };
-template<> struct TCanBulkSerialize<UE::Geometry::GVector4<double>> { enum { Value = true }; };
diff --git a/Engine/Source/Runtime/GeometryCore/Public/VectorUtil.h b/Engine/Source/Runtime/GeometryCore/Public/VectorUtil.h
--- a/Engine/Source/Runtime/GeometryCore/Public/VectorUtil.h
+++ b/Engine/Source/Runtime/GeometryCore/Public/VectorUtil.h
@@ -26,11 +26,15 @@ enum class EIntersectionType
 	Unknown
 };
 
-namespace VectorUtil
+namespace UE
+{
+namespace Geometry
 {
-	using namespace UE::Geometry;
 	using namespace UE::Math;
 
+namespace VectorUtil
+{
+
 	/**
 	 * @return true if all components of V are finite
 	 */
@@ -44,7 +48,7 @@ namespace VectorUtil
 	 * @return true if all components of V are finite
 	 */
 	template <typename RealType>
-	inline bool IsFinite(const UE::Math::TVector<RealType>& V)
+	inline bool IsFinite(const TVector<RealType>& V)
 	{
 		return TMathUtil<RealType>::IsFinite(V.X) && TMathUtil<RealType>::IsFinite(V.Y) && TMathUtil<RealType>::IsFinite(V.Z);
 	}
@@ -63,15 +67,15 @@ namespace VectorUtil
 	 * @return normalized vector that is perpendicular to triangle V0,V1,V2  (triangle normal)
 	 */
 	template <typename RealType>
-	inline FVector3<RealType> Normal(const UE::Math::TVector<RealType>& V0, const UE::Math::TVector<RealType>& V1, const UE::Math::TVector<RealType>& V2)
+	inline TVector<RealType> Normal(const TVector<RealType>& V0, const TVector<RealType>& V1, const TVector<RealType>& V2)
 	{
-		FVector3<RealType> edge1(V1 - V0);
-		FVector3<RealType> edge2(V2 - V0);
+		TVector<RealType> edge1(V1 - V0);
+		TVector<RealType> edge2(V2 - V0);
 		Normalize(edge1);
 		Normalize(edge2);
 		// Unreal has Left-Hand Coordinate System so we need to reverse this cross-product to get proper triangle normal
-		FVector3<RealType> vCross(edge2.Cross(edge1));
-		//FVector3<RealType> vCross(edge1.Cross(edge2));
+		TVector<RealType> vCross(edge2.Cross(edge1));
+		//TVector<RealType> vCross(edge1.Cross(edge2));
 		return Normalized(vCross);
 	}
 
@@ -79,7 +83,7 @@ namespace VectorUtil
 	 * @return un-normalized direction that is parallel to normal of triangle V0,V1,V2
 	 */
 	template <typename RealType>
-	inline FVector3<RealType> NormalDirection(const UE::Math::TVector<RealType>& V0, const UE::Math::TVector<RealType>& V1, const UE::Math::TVector<RealType>& V2)
+	inline TVector<RealType> NormalDirection(const TVector<RealType>& V0, const TVector<RealType>& V1, const TVector<RealType>& V2)
 	{
 		// Unreal has Left-Hand Coordinate System so we need to reverse this cross-product to get proper triangle normal
 		return (V2 - V0).Cross(V1 - V0);
@@ -90,11 +94,11 @@ namespace VectorUtil
 	 * @return area of 3D triangle V0,V1,V2
 	 */
 	template <typename RealType>
-	inline RealType Area(const UE::Math::TVector<RealType>& V0, const UE::Math::TVector<RealType>& V1, const UE::Math::TVector<RealType>& V2)
+	inline RealType Area(const TVector<RealType>& V0, const TVector<RealType>& V1, const TVector<RealType>& V2)
 	{
-		FVector3<RealType> Edge1(V1 - V0);
-		FVector3<RealType> Edge2(V2 - V0);
-		FVector3<RealType> Cross = Edge2.Cross(Edge1);
+		TVector<RealType> Edge1(V1 - V0);
+		TVector<RealType> Edge2(V2 - V0);
+		TVector<RealType> Cross = Edge2.Cross(Edge1);
 		return (RealType)0.5 * Cross.Length();
 	}
 
@@ -123,7 +127,7 @@ namespace VectorUtil
 	 * @return true if triangle V1,V2,V3 is obtuse
 	 */
 	template <typename RealType>
-	inline bool IsObtuse(const UE::Math::TVector<RealType>& V1, const UE::Math::TVector<RealType>& V2, const UE::Math::TVector<RealType>& V3)
+	inline bool IsObtuse(const TVector<RealType>& V1, const TVector<RealType>& V2, const TVector<RealType>& V3)
 	{
 		RealType a2 = DistanceSquared(V1, V2);
 		RealType b2 = DistanceSquared(V1, V3);
@@ -136,10 +140,10 @@ namespace VectorUtil
 	 * @return triangle normal
 	 */
 	template <typename RealType>
-	inline FVector3<RealType> NormalArea(const UE::Math::TVector<RealType>& V0, const UE::Math::TVector<RealType>& V1, const UE::Math::TVector<RealType>& V2, RealType& AreaOut)
+	inline TVector<RealType> NormalArea(const TVector<RealType>& V0, const TVector<RealType>& V1, const TVector<RealType>& V2, RealType& AreaOut)
 	{
-		FVector3<RealType> edge1(V1 - V0);
-		FVector3<RealType> edge2(V2 - V0);
+		TVector<RealType> edge1(V1 - V0);
+		TVector<RealType> edge2(V2 - V0);
 		// Unreal has Left-Hand Coordinate System so we need to reverse this cross-product to get proper triangle normal
 		FVector3d vCross = edge2.Cross(edge1);
 		//FVector3d vCross = edge1.Cross(edge2);
@@ -163,14 +167,14 @@ namespace VectorUtil
 
 	/** @return true if all coordinates of V0 and V1 are within Epsilon of eachother */
 	template <typename RealType>
-	inline bool EpsilonEqual(const UE::Math::TVector<RealType>& V0, const UE::Math::TVector<RealType>& V1, RealType Epsilon)
+	inline bool EpsilonEqual(const TVector<RealType>& V0, const TVector<RealType>& V1, RealType Epsilon)
 	{
 		return EpsilonEqual(V0.X, V1.X, Epsilon) && EpsilonEqual(V0.Y, V1.Y, Epsilon) && EpsilonEqual(V0.Z, V1.Z, Epsilon);
 	}
 
 	/** @return true if all coordinates of V0 and V1 are within Epsilon of each other */
 	template <typename RealType>
-	inline bool EpsilonEqual(const UE::Math::TVector4<RealType>& V0, const UE::Math::TVector4<RealType>& V1, RealType Epsilon)
+	inline bool EpsilonEqual(const TVector4<RealType>& V0, const TVector4<RealType>& V1, RealType Epsilon)
 	{
 		return EpsilonEqual(V0.X, V1.X, Epsilon) && EpsilonEqual(V0.Y, V1.Y, Epsilon) && EpsilonEqual(V0.Z, V1.Z, Epsilon) && EpsilonEqual(V0.W, V1.W, Epsilon);
 	}
@@ -195,7 +199,7 @@ namespace VectorUtil
 	 * Calculates two vectors perpendicular to input Normal, as efficiently as possible.
 	 */
 	template <typename RealType>
-	inline void MakePerpVectors(const UE::Math::TVector<RealType>& Normal, UE::Math::TVector<RealType>& OutPerp1, UE::Math::TVector<RealType>& OutPerp2)
+	inline void MakePerpVectors(const TVector<RealType>& Normal, TVector<RealType>& OutPerp1, TVector<RealType>& OutPerp2)
 	{
 		// Duff et al method, from https://graphics.pixar.com/library/OrthonormalB/paper.pdf
 		if (Normal.Z < (RealType)0)
@@ -226,7 +230,7 @@ namespace VectorUtil
 	 * Calculates one vector perpendicular to input Normal, as efficiently as possible.
 	 */
 	template <typename RealType>
-	inline void MakePerpVector(const FVector3<RealType>& Normal, FVector3<RealType>& OutPerp1)
+	inline void MakePerpVector(const TVector<RealType>& Normal, TVector<RealType>& OutPerp1)
 	{
 		// Duff et al method, from https://graphics.pixar.com/library/OrthonormalB/paper.pdf
 		if (Normal.Z < (RealType)0)
@@ -253,13 +257,13 @@ namespace VectorUtil
 	 * @return angle in degrees
 	 */
 	template <typename RealType>
-	inline RealType PlaneAngleSignedD(const UE::Math::TVector<RealType>& VFrom, const UE::Math::TVector<RealType>& VTo, const UE::Math::TVector<RealType>& PlaneN)
+	inline RealType PlaneAngleSignedD(const TVector<RealType>& VFrom, const TVector<RealType>& VTo, const TVector<RealType>& PlaneN)
 	{
-		FVector3<RealType> vFrom = VFrom - VFrom.Dot(PlaneN) * PlaneN;
-		FVector3<RealType> vTo = VTo - VTo.Dot(PlaneN) * PlaneN;
+		TVector<RealType> vFrom = VFrom - VFrom.Dot(PlaneN) * PlaneN;
+		TVector<RealType> vTo = VTo - VTo.Dot(PlaneN) * PlaneN;
 		Normalize(vFrom);
 		Normalize(vTo);
-		FVector3<RealType> C = vFrom.Cross(vTo);
+		TVector<RealType> C = vFrom.Cross(vTo);
 		if (C.SquaredLength() < TMathUtil<RealType>::ZeroTolerance)
 		{ // vectors are parallel
 			return vFrom.Dot(vTo) < 0 ? (RealType)180 : (RealType)0;
@@ -273,7 +277,7 @@ namespace VectorUtil
 	 * @return positive value of tan(theta/2) where theta is angle between normalized vectors A and B
 	 */
 	template <typename RealType>
-	RealType VectorTanHalfAngle(const UE::Math::TVector<RealType>& A, const UE::Math::TVector<RealType>& B)
+	RealType VectorTanHalfAngle(const TVector<RealType>& A, const TVector<RealType>& B)
 	{
 		RealType cosAngle = A.Dot(B);
 		RealType sqr = ((RealType)1 - cosAngle) / ((RealType)1 + cosAngle);
@@ -287,7 +291,7 @@ namespace VectorUtil
 	 * @return cotangent of angle between V1 and V2, or zero if result would be unstable (eg infinity)
 	 */ 
 	template <typename RealType>
-	RealType VectorCot(const UE::Math::TVector<RealType>& V1, const UE::Math::TVector<RealType>& V2)
+	RealType VectorCot(const TVector<RealType>& V1, const TVector<RealType>& V2)
 	{
 		// formula from http://www.geometry.caltech.edu/pubs/DMSB_III.pdf
 		RealType fDot = V1.Dot(V2);
@@ -311,11 +315,11 @@ namespace VectorUtil
 	 * TODO: make robust to degenerate triangles?
 	 */
 	template <typename RealType>
-	FVector3<RealType> BarycentricCoords(const UE::Math::TVector<RealType>& Point, const UE::Math::TVector<RealType>& V0, const UE::Math::TVector<RealType>& V1, const UE::Math::TVector<RealType>& V2)
+	TVector<RealType> BarycentricCoords(const TVector<RealType>& Point, const TVector<RealType>& V0, const TVector<RealType>& V1, const TVector<RealType>& V2)
 	{
-		FVector3<RealType> kV02 = V0 - V2;
-		FVector3<RealType> kV12 = V1 - V2;
-		FVector3<RealType> kPV2 = Point - V2;
+		TVector<RealType> kV02 = V0 - V2;
+		TVector<RealType> kV12 = V1 - V2;
+		TVector<RealType> kPV2 = Point - V2;
 		RealType fM00 = kV02.Dot(kV02);
 		RealType fM01 = kV02.Dot(kV12);
 		RealType fM11 = kV12.Dot(kV12);
@@ -326,7 +330,7 @@ namespace VectorUtil
 		RealType fBary1 = (fM11 * fR0 - fM01 * fR1) * fInvDet;
 		RealType fBary2 = (fM00 * fR1 - fM01 * fR0) * fInvDet;
 		RealType fBary3 = 1.0 - fBary1 - fBary2;
-		return FVector3<RealType>(fBary1, fBary2, fBary3);
+		return TVector<RealType>(fBary1, fBary2, fBary3);
 	}
 
 	/**
@@ -336,7 +340,7 @@ namespace VectorUtil
 	* TODO: make robust to degenerate triangles?
 	*/
 	template <typename RealType>
-	FVector3<RealType> BarycentricCoords(const TVector2<RealType>& Point, const TVector2<RealType>& V0, const TVector2<RealType>& V1, const TVector2<RealType>& V2)
+	TVector<RealType> BarycentricCoords(const TVector2<RealType>& Point, const TVector2<RealType>& V0, const TVector2<RealType>& V1, const TVector2<RealType>& V2)
 	{
 		TVector2<RealType> kV02 = V0 - V2;
 		TVector2<RealType> kV12 = V1 - V2;
@@ -351,14 +355,14 @@ namespace VectorUtil
 		RealType fBary1 = (fM11 * fR0 - fM01 * fR1) * fInvDet;
 		RealType fBary2 = (fM00 * fR1 - fM01 * fR0) * fInvDet;
 		RealType fBary3 = 1.0 - fBary1 - fBary2;
-		return FVector3<RealType>(fBary1, fBary2, fBary3);
+		return TVector<RealType>(fBary1, fBary2, fBary3);
 	}
 
 	/**
 	 * @return solid angle at point P for triangle A,B,C
 	 */
 	template <typename RealType>
-	inline RealType TriSolidAngle(UE::Math::TVector<RealType> A, UE::Math::TVector<RealType> B, UE::Math::TVector<RealType> C, const UE::Math::TVector<RealType>& P)
+	inline RealType TriSolidAngle(TVector<RealType> A, TVector<RealType> B, TVector<RealType> C, const TVector<RealType>& P)
 	{
 		// Formula from https://igl.ethz.ch/projects/winding-number/
 		A -= P;
@@ -378,14 +382,14 @@ namespace VectorUtil
 	 * @return gradient (3D vector) lying in plane of triangle.
 	 */
 	template <typename RealType>
-	inline FVector3<RealType> TriGradient(UE::Math::TVector<RealType> Vi, UE::Math::TVector<RealType> Vj, UE::Math::TVector<RealType> Vk, RealType fi, RealType fj, RealType fk)
+	inline TVector<RealType> TriGradient(TVector<RealType> Vi, TVector<RealType> Vj, TVector<RealType> Vk, RealType fi, RealType fj, RealType fk)
 	{
 		// recenter (better for precision)
-		FVector3<RealType> Centroid = (Vi + Vj + Vk) / (RealType)3;
+		TVector<RealType> Centroid = (Vi + Vj + Vk) / (RealType)3;
 		Vi -= Centroid; Vj -= Centroid; Vk -= Centroid;
 		// calculate tangent-normal frame
-		FVector3<RealType> Normal = VectorUtil::Normal<RealType>(Vi, Vj, Vk);
-		FVector3<RealType> Perp0, Perp1;
+		TVector<RealType> Normal = VectorUtil::Normal<RealType>(Vi, Vj, Vk);
+		TVector<RealType> Perp0, Perp1;
 		VectorUtil::MakePerpVectors<RealType>(Normal, Perp0, Perp1);
 		// project points to triangle plane coordinates
 		TVector2<RealType> vi(Vi.Dot(Perp0), Vi.Dot(Perp1));
@@ -404,7 +408,7 @@ namespace VectorUtil
 	 * @return angle between vectors (A-CornerPt) and (B-CornerPt)
 	 */
 	template<typename RealType>
-	inline RealType OpeningAngleD(UE::Math::TVector<RealType> A, UE::Math::TVector<RealType> B, const UE::Math::TVector<RealType>& P)
+	inline RealType OpeningAngleD(TVector<RealType> A, TVector<RealType> B, const TVector<RealType>& P)
 	{
 		A -= P; 
 		Normalize(A);
@@ -419,7 +423,7 @@ namespace VectorUtil
 	 * @return sign of Bitangent relative to Normal and Tangent
 	 */
 	template<typename RealType>
-	inline RealType BitangentSign(const UE::Math::TVector<RealType>& NormalIn, const UE::Math::TVector<RealType>& TangentIn, const UE::Math::TVector<RealType>& BitangentIn)
+	inline RealType BitangentSign(const TVector<RealType>& NormalIn, const TVector<RealType>& TangentIn, const TVector<RealType>& BitangentIn)
 	{
 		// following math from RenderUtils.h::GetBasisDeterminantSign()
 		RealType Cross00 = BitangentIn.Y*NormalIn.Z - BitangentIn.Z*NormalIn.Y;
@@ -433,9 +437,9 @@ namespace VectorUtil
 	 * @return Bitangent vector based on given Normal, Tangent, and Sign value (+1/-1)
 	 */
 	template<typename RealType>
-	inline FVector3<RealType> Bitangent(const UE::Math::TVector<RealType>& NormalIn, const UE::Math::TVector<RealType>& TangentIn, RealType BitangentSign)
+	inline TVector<RealType> Bitangent(const TVector<RealType>& NormalIn, const TVector<RealType>& TangentIn, RealType BitangentSign)
 	{
-		return BitangentSign * FVector3<RealType>(
+		return BitangentSign * TVector<RealType>(
 			NormalIn.Y*TangentIn.Z - NormalIn.Z*TangentIn.Y,
 			NormalIn.Z*TangentIn.X - NormalIn.X*TangentIn.Z,
 			NormalIn.X*TangentIn.Y - NormalIn.Y*TangentIn.X);
@@ -445,7 +449,7 @@ namespace VectorUtil
 	 * @return Tangent-Space vector based on given Normal and Bitangent
 	 */
 	template<typename RealType>
-	inline FVector3<RealType> TangentFromBitangent(const UE::Math::TVector<RealType>& NormalIn, const UE::Math::TVector<RealType>& BitangentIn)
+	inline TVector<RealType> TangentFromBitangent(const TVector<RealType>& NormalIn, const TVector<RealType>& BitangentIn)
 	{
 		return BitangentIn.Cross(NormalIn);
 	}
@@ -454,7 +458,7 @@ namespace VectorUtil
 	 * @return Bitangent vector based on given Normal and Tangent
 	 */
 	template<typename RealType>
-	inline FVector3<RealType> BitangentFromTangent(const UE::Math::TVector<RealType>& NormalIn, const UE::Math::TVector<RealType>& TangentIn)
+	inline TVector<RealType> BitangentFromTangent(const TVector<RealType>& NormalIn, const TVector<RealType>& TangentIn)
 	{
 		return NormalIn.Cross(TangentIn);
 	}
@@ -469,3 +473,7 @@ namespace VectorUtil
 
 
 }; // namespace VectorUtil
+
+
+}  // end namespace Geometry
+}  // end namespace UE
diff --git a/Engine/Source/Runtime/GeometryFramework/Private/Components/DynamicMeshSceneProxy.h b/Engine/Source/Runtime/GeometryFramework/Private/Components/DynamicMeshSceneProxy.h
--- a/Engine/Source/Runtime/GeometryFramework/Private/Components/DynamicMeshSceneProxy.h
+++ b/Engine/Source/Runtime/GeometryFramework/Private/Components/DynamicMeshSceneProxy.h
@@ -111,7 +111,7 @@ public:
 		// fallback to orthogonal basis
 		return [](int VertexID, int TriangleID, int TriVtxIdx, const FVector3f& Normal, FVector3f& TangentX, FVector3f& TangentY) -> void
 		{
-			VectorUtil::MakePerpVectors(Normal, TangentX, TangentY);
+			UE::Geometry::VectorUtil::MakePerpVectors(Normal, TangentX, TangentY);
 		};
 	}
 
diff --git a/Engine/Source/Runtime/MeshConversion/Private/MeshDescriptionToDynamicMesh.cpp b/Engine/Source/Runtime/MeshConversion/Private/MeshDescriptionToDynamicMesh.cpp
--- a/Engine/Source/Runtime/MeshConversion/Private/MeshDescriptionToDynamicMesh.cpp
+++ b/Engine/Source/Runtime/MeshConversion/Private/MeshDescriptionToDynamicMesh.cpp
@@ -859,10 +859,10 @@ static void CopyTangents_Internal(const FMeshDescription* SourceMesh, const FDyn
 			FVector Normal = InstanceNormals[InstanceTri[j]];
 			FVector Tangent = InstanceTangents[InstanceTri[j]];
 			float BitangentSign = InstanceSigns[InstanceTri[j]];
-			FVector3<RealType> Bitangent = VectorUtil::Bitangent((FVector3<RealType>)Normal, (FVector3<RealType>)Tangent, (RealType)BitangentSign);
+			TVector<RealType> Bitangent = VectorUtil::Bitangent((TVector<RealType>)Normal, (TVector<RealType>)Tangent, (RealType)BitangentSign);
 			Tangent.Normalize();
 			UE::Geometry::Normalize(Bitangent);
-			TangentsOut->SetPerTriangleTangent(TriID, j, (FVector3<RealType>)Tangent, (FVector3<RealType>)Bitangent);
+			TangentsOut->SetPerTriangleTangent(TriID, j, (TVector<RealType>)Tangent, (TVector<RealType>)Bitangent);
 		}
 	}
 }
