add 2D triangle-AABB intersection tests + add some unit test coverage and make 2D triangle-triangle intersection compile

diff --git a/Engine/Source/Runtime/GeometryCore/Public/Intersection/IntrTriangle2AxisAlignedBox2.h b/Engine/Source/Runtime/GeometryCore/Public/Intersection/IntrTriangle2AxisAlignedBox2.h
new file mode 100644
--- /dev/null
+++ b/Engine/Source/Runtime/GeometryCore/Public/Intersection/IntrTriangle2AxisAlignedBox2.h
@@ -0,0 +1,135 @@
+// Copyright Epic Games, Inc. All Rights Reserved.
+
+#pragma once
+
+#include "VectorTypes.h"
+#include "TriangleTypes.h"
+#include "BoxTypes.h"
+#include "VectorUtil.h"
+
+namespace UE
+{
+namespace Geometry
+{
+
+using namespace UE::Math;
+
+/**
+ * Compute intersection a 2D triangle and a 2D axis-aligned box
+ */
+template <typename Real>
+class TIntrTriangle2AxisAlignedBox2
+{
+protected:
+	// Input
+	TTriangle2<Real> Triangle;
+	TAxisAlignedBox2<Real> Box;
+	bool bTriangleIsOriented = true;
+
+public:
+	// Output
+	// TODO: Unlike other TIntrPrim2Prim2 classes, this only computes the Boolean yes/no Test()
+	// But in the future we should add the option to Compute() the intersection polygon as well (e.g. refer to IntrTriangle2Triangle.h)
+
+	TIntrTriangle2AxisAlignedBox2()
+	{}
+	TIntrTriangle2AxisAlignedBox2(TTriangle2<Real> Triangle, TAxisAlignedBox2<Real> Box, bool bTriangleIsOriented = true)
+		: Triangle(Triangle), Box(Box), bTriangleIsOriented(bTriangleIsOriented)
+	{
+	}
+
+	TTriangle2<Real> GetTriangle() const
+	{
+		return Triangle;
+	}
+	bool IsTriangleOriented() const
+	{
+		return bTriangleIsOriented;
+	}
+	TTriangle2<Real> GetBox() const
+	{
+		return Box;
+	}
+	void SetTriangle(const TTriangle2<Real>& TriangleIn)
+	{
+		// TODO: Reset any result variable(s), once support for computing full intersection result is added
+		Triangle = TriangleIn;
+	}
+	void SetTriangleOriented(const bool bIsOrientedIn)
+	{
+		bTriangleIsOriented = bIsOrientedIn;
+	}
+	void SetBox(const TAxisAlignedBox2<Real>& BoxIn)
+	{
+		// TODO: Reset any result variable(s), once support for computing full intersection result is added
+		Box = BoxIn;
+	}
+	
+
+	bool Test()
+	{
+		TAxisAlignedBox2<Real> TriBox(Triangle.V[0], Triangle.V[0]);
+		TriBox.Contain(Triangle.V[1]);
+		TriBox.Contain(Triangle.V[2]);
+
+		if (!TriBox.Intersects(Box))
+		{
+			// iff we're separated by a bounding box axis, the bounds don't intersect
+			return false;
+		}
+
+		Real OrientationSign = 1;
+		Real SignedArea = Triangle.SignedArea();
+		if (SignedArea < 0)
+		{
+			if (bTriangleIsOriented && SignedArea < -TMathUtil<Real>::ZeroTolerance)
+			{
+				// if we're using oriented triangles and this one is inverted, it can't intersect
+				return false;
+			}
+			OrientationSign = -1;
+		}
+
+		// test if a box is entirely in the half-plane starting at EdgeVertex, going in the EdgePerp direction
+		auto BoxOutsideEdge = [this](TVector2<Real> EdgeVertex, TVector2<Real> EdgePerp)
+		{
+			for (int32 CornerIdx = 0; CornerIdx < 4; CornerIdx++)
+			{
+				bool bOutside = EdgePerp.Dot(Box.GetCorner(CornerIdx) - EdgeVertex) > 0;
+				if (!bOutside)
+				{
+					return false;
+				}
+			}
+			return true;
+		};
+
+		// Test triangle edges for separation
+		for (int32 LeadVert = 0, PrevVert = 2; LeadVert < 3; PrevVert = LeadVert++)
+		{
+			TVector2<Real> Edge = Triangle.V[LeadVert] - Triangle.V[PrevVert];
+			Real SquaredEdgeLen = Edge.SquaredLength();
+			if (SquaredEdgeLen < TMathUtil<Real>::ZeroTolerance)
+			{
+				continue;
+			}
+
+			if (BoxOutsideEdge(Triangle.V[LeadVert], PerpCW(Edge) * OrientationSign))
+			{
+				// separated by a triangle edge
+				return false;
+			}
+		}
+
+		// no separating axis found
+		return true;
+	}
+
+
+};
+
+typedef TIntrTriangle2AxisAlignedBox2<float> FIntrTriangle2AxisAlignedBox2f;
+typedef TIntrTriangle2AxisAlignedBox2<double> FIntrTriangle2AxisAlignedBox2d;
+
+} // end namespace UE::Geometry
+} // end namespace UE
diff --git a/Engine/Source/Runtime/GeometryCore/Public/Intersection/IntrTriangle2Triangle2.h b/Engine/Source/Runtime/GeometryCore/Public/Intersection/IntrTriangle2Triangle2.h
--- a/Engine/Source/Runtime/GeometryCore/Public/Intersection/IntrTriangle2Triangle2.h
+++ b/Engine/Source/Runtime/GeometryCore/Public/Intersection/IntrTriangle2Triangle2.h
@@ -74,9 +74,9 @@ public:
 		for (i0 = 0, i1 = 2; i0 < 3; i1 = i0++)
 		{
 			// Test axis V0[i1] + t*perp(V0[i0]-V0[i1]), perp(x,y) = (y,-x).
-			dir.X = Triangle0[i0].Y - Triangle0[i1].Y;
-			dir.Y = Triangle0[i1].X - Triangle0[i0].X;
-			if (WhichSide(Triangle1, Triangle0[i1], dir) > 0)
+			dir.X = Triangle0.V[i0].Y - Triangle0.V[i1].Y;
+			dir.Y = Triangle0.V[i1].X - Triangle0.V[i0].X;
+			if (WhichSide(Triangle1, Triangle0.V[i1], dir) > 0)
 			{
 				// Triangle1 is entirely on positive side of Triangle0 edge.
 				return false;
@@ -87,9 +87,9 @@ public:
 		for (i0 = 0, i1 = 2; i0 < 3; i1 = i0++)
 		{
 			// Test axis V1[i1] + t*perp(V1[i0]-V1[i1]), perp(x,y) = (y,-x).
-			dir.X = Triangle1[i0].Y - Triangle1[i1].Y;
-			dir.Y = Triangle1[i1].X - Triangle1[i0].X;
-			if (WhichSide(Triangle0, Triangle1[i1], dir) > 0) {
+			dir.X = Triangle1.V[i0].Y - Triangle1.V[i1].Y;
+			dir.Y = Triangle1.V[i1].X - Triangle1.V[i0].X;
+			if (WhichSide(Triangle0, Triangle1.V[i1], dir) > 0) {
 				// Triangle0 is entirely on positive side of Triangle1 edge.
 				return false;
 			}
@@ -168,7 +168,7 @@ public:
 		int positive = 0, negative = 0, zero = 0;
 		for (int i = 0; i < 3; ++i)
 		{
-			double t = D.Dot(V[i] - P);
+			double t = D.Dot(V.V[i] - P);
 			if (t > (double)0)
 			{
 				++positive;
