GeometryCore: handle case where a mesh overlay has serialized unreferenced elements, ie that have no parent vertex. Existing code will try to look up the remapping w/ a parent ID of -1, causing a crash.

#rnx

diff --git a/Engine/Source/Runtime/GeometryCore/Private/DynamicMesh/DynamicMeshOverlay.cpp b/Engine/Source/Runtime/GeometryCore/Private/DynamicMesh/DynamicMeshOverlay.cpp
--- a/Engine/Source/Runtime/GeometryCore/Private/DynamicMesh/DynamicMeshOverlay.cpp
+++ b/Engine/Source/Runtime/GeometryCore/Private/DynamicMesh/DynamicMeshOverlay.cpp
@@ -1678,7 +1678,8 @@ void TDynamicMeshOverlay<RealType, ElementSize>::Serialize(FArchive& Ar, const F
 				{
 					ElementsCompact[EidCompact * ElementSize + i] = Elements[Eid * ElementSize + i];
 				}
-				ParentVerticesCompact[EidCompact] = CompactMaps->GetVertexMapping(ParentVertices[Eid]);
+				ParentVerticesCompact[EidCompact] =
+					(ParentVertices[Eid] >= 0) ? CompactMaps->GetVertexMapping(ParentVertices[Eid]) : FDynamicMesh3::InvalidID;
 				ElementMapping[Eid] = EidCompact;
 				++EidCompact;
 			}
