manually copy CIS fix from 4.27-chaos, because file moved (add size_t cast to SparseDynamicPointOctree3 function)

#rnx

diff --git a/Engine/Source/Runtime/GeometryCore/Public/Spatial/SparseDynamicPointOctree3.h b/Engine/Source/Runtime/GeometryCore/Public/Spatial/SparseDynamicPointOctree3.h
--- a/Engine/Source/Runtime/GeometryCore/Public/Spatial/SparseDynamicPointOctree3.h
+++ b/Engine/Source/Runtime/GeometryCore/Public/Spatial/SparseDynamicPointOctree3.h
@@ -332,7 +332,7 @@ protected:
 
 	uint32 GetCellForPoint(int32 PointID) const
 	{
-		if (PointID >= 0 && PointID < PointIDToCellMap.Num())
+		if (PointID >= 0 && size_t(PointID) < PointIDToCellMap.Num())
 		{
 			return PointIDToCellMap[PointID];
 		}
