Add WriteOBJ function for FDynamicMeshAABBTree3 and ReadOBJ function useful to inline small obj files

Also do some refactoring to pass settings structs which improve readability when setting options, reduce the need to change call sites due to signature changes, and allow overloading of Read/Load functions which return some type

#rnx

diff --git a/Engine/Source/Runtime/GeometryCore/Public/Spatial/MeshAABBTree3.h b/Engine/Source/Runtime/GeometryCore/Public/Spatial/MeshAABBTree3.h
--- a/Engine/Source/Runtime/GeometryCore/Public/Spatial/MeshAABBTree3.h
+++ b/Engine/Source/Runtime/GeometryCore/Public/Spatial/MeshAABBTree3.h
@@ -61,11 +61,25 @@ class TMeshAABBTree3 : public IMeshSpatial
 {
 	friend class TFastWindingTree<TriangleMeshType>;
 
+public:
+	using MeshType = TriangleMeshType;
+	using GetSplitAxisFunc = TUniqueFunction<int(int Depth, const FAxisAlignedBox3d& Box)>;
+
 protected:
 	const TriangleMeshType* Mesh;
 	uint64 MeshChangeStamp = 0;
 	int TopDownLeafMaxTriCount = 4;
 
+	static GetSplitAxisFunc MakeDefaultSplitAxisFunc()
+	{
+		return [](int Depth, const FAxisAlignedBox3d&)
+		{
+			return Depth % 3;
+		};
+	}
+	
+	GetSplitAxisFunc GetSplitAxis = MakeDefaultSplitAxisFunc();
+
 public:
 	static constexpr double DOUBLE_MAX = TNumericLimits<double>::Max();
 
@@ -113,9 +127,10 @@ public:
 		return true;
 	}
 
-	void SetBuildOptions(int32 MaxBoxTriCount)
+	void SetBuildOptions(int32 MaxBoxTriCount, GetSplitAxisFunc&& GetSplitAxisIn = MakeDefaultSplitAxisFunc())
 	{
 		TopDownLeafMaxTriCount = MaxBoxTriCount;
+		GetSplitAxis = MoveTemp(GetSplitAxisIn);
 	}
 
 	void Build()
@@ -630,19 +645,17 @@ public:
 	class FTreeTraversal
 	{
 	public:
-		// return false to terminate this branch
-		// arguments are Box and Depth in tree
+		// The traversal is terminated at the current node if this function returns false
 		TFunction<bool(const FAxisAlignedBox3d&, int)> NextBoxF = [](const FAxisAlignedBox3d& Box, int Depth) { return true; };
 
-		TFunction<void(int)> BeginBoxTrianglesF = [](int32 BoxID) {};
-
-		TFunction<void(int)> NextTriangleF = [](int TID) {};
-
-		TFunction<void(int)> EndBoxTrianglesF = [](int32 BoxID) {};
+		// These functions are called in sequence if NextBoxF returned true and the current node is a leaf
+		TFunction<void(int, int)> BeginBoxTrianglesF = [](int BoxID, int Depth) {};
+		TFunction<void(int)> NextTriangleF = [](int TriangleID) {}; // Call may be skipped by TriangleFilterF
+		TFunction<void(int)> EndBoxTrianglesF = [](int BoxID) {};
 	};
 
 	/**
-	 * Hierarchically descend through the tree Nodes, calling the TreeTrversal functions at each level
+	 * Hierarchically descend through the tree Nodes, calling the FTreeTraversal functions at each level
 	 */
 	virtual void DoTraversal(FTreeTraversal& Traversal, const FQueryOptions& Options = FQueryOptions()) const
 	{
@@ -651,7 +664,10 @@ public:
 			return;
 		}
 
-		TreeTraversalImpl(RootIndex, 0, Traversal, Options);
+		if (Traversal.NextBoxF(GetBox(RootIndex), 0))
+		{
+			TreeTraversalImpl(RootIndex, 0, Traversal, Options);
+		}
 	}
 
 protected:
@@ -664,7 +680,7 @@ protected:
 
 		if (idx < TrianglesEnd)
 		{
-			Traversal.BeginBoxTrianglesF(IBox);
+			Traversal.BeginBoxTrianglesF(IBox, Depth);
 
 			// triangle-list case, array is [N t1 t2 ... tN]
 			int n = IndexList[idx];
@@ -700,6 +716,7 @@ protected:
 				{
 					TreeTraversalImpl(i0, Depth + 1, Traversal, Options);
 				}
+				
 				int i1 = IndexList[idx + 1] - 1;
 				if (Traversal.NextBoxF(GetBox(i1), Depth + 1))
 				{
@@ -1157,7 +1174,7 @@ public:
 			FAxisAlignedBox3d::Empty() : FAxisAlignedBox3d(FVector3d::Zero(), 0.0);
 		int IBox = -1;
 
-		if (ICount < MinTriCount)
+		if (ICount <= MinTriCount)
 		{
 			// append new Triangles Box
 			IBox = Tris.IBoxCur++;
@@ -1177,7 +1194,7 @@ public:
 		}
 
 		//compute interval along an axis and find midpoint
-		int axis = Depth % 3;
+		int axis = GetSplitAxis(Depth, Box);
 		FInterval1d interval = FInterval1d::Empty();
 		for (int i = 0; i < ICount; ++i)
 		{
diff --git a/Engine/Source/Runtime/GeometryFramework/Private/Components/MeshRenderDecomposition.cpp b/Engine/Source/Runtime/GeometryFramework/Private/Components/MeshRenderDecomposition.cpp
--- a/Engine/Source/Runtime/GeometryFramework/Private/Components/MeshRenderDecomposition.cpp
+++ b/Engine/Source/Runtime/GeometryFramework/Private/Components/MeshRenderDecomposition.cpp
@@ -95,7 +95,7 @@ static void CollectSubDecomposition(
 	TArray<int32> SpillSet;
 
 	FDynamicMeshAABBTree3::FTreeTraversal Collector;
-	Collector.BeginBoxTrianglesF = [&](int32 BoxID)
+	Collector.BeginBoxTrianglesF = [&](int, int)
 	{
 		ActiveSet.Reset();
 		ActiveSet.Reserve(MaxTrisPerGroup);
@@ -104,7 +104,7 @@ static void CollectSubDecomposition(
 	{
 		ActiveSet.Add(TriangleID);
 	};
-	Collector.EndBoxTrianglesF = [&](int32 BoxID)
+	Collector.EndBoxTrianglesF = [&](int)
 	{
 		if (ActiveSet.Num() > 0)
 		{
