fix geometryprocessing tests by changing MeshAABBTree3 to have the original 'effective' default tris per leaf (of 3), rather than the original value of 4, now that the off-by-one bug in how it handles that value has been fixed.

#rnx

diff --git a/Engine/Source/Runtime/GeometryCore/Public/Spatial/MeshAABBTree3.h b/Engine/Source/Runtime/GeometryCore/Public/Spatial/MeshAABBTree3.h
--- a/Engine/Source/Runtime/GeometryCore/Public/Spatial/MeshAABBTree3.h
+++ b/Engine/Source/Runtime/GeometryCore/Public/Spatial/MeshAABBTree3.h
@@ -68,7 +68,7 @@ public:
 protected:
 	const TriangleMeshType* Mesh;
 	uint64 MeshChangeStamp = 0;
-	int TopDownLeafMaxTriCount = 4;
+	int TopDownLeafMaxTriCount = 3;
 
 	static GetSplitAxisFunc MakeDefaultSplitAxisFunc()
 	{
