ModelingTools: Fix crash in floating/spiral stair generator.
Set the minimum number of steps to 2.

#rnx

diff --git a/Engine/Plugins/Runtime/MeshModelingToolset/Source/MeshModelingTools/Public/AddPrimitiveTool.h b/Engine/Plugins/Runtime/MeshModelingToolset/Source/MeshModelingTools/Public/AddPrimitiveTool.h
--- a/Engine/Plugins/Runtime/MeshModelingToolset/Source/MeshModelingTools/Public/AddPrimitiveTool.h
+++ b/Engine/Plugins/Runtime/MeshModelingToolset/Source/MeshModelingTools/Public/AddPrimitiveTool.h
@@ -390,7 +390,7 @@ public:
 	EProceduralStairsType StairsType = EProceduralStairsType::Linear;
 
 	/** Number of steps */
-	UPROPERTY(EditAnywhere, Category = ShapeSettings, meta = (DisplayName = "Number of Steps", UIMin = "1", UIMax = "100", ClampMin = "1", ClampMax = "1000000", ProceduralShapeSetting))
+	UPROPERTY(EditAnywhere, Category = ShapeSettings, meta = (DisplayName = "Number of Steps", UIMin = "2", UIMax = "100", ClampMin = "2", ClampMax = "1000000", ProceduralShapeSetting))
 	int NumSteps = 8;
 
 	/** Width of each step */
diff --git a/Engine/Source/Runtime/GeometryCore/Private/Generators/StairGenerator.cpp b/Engine/Source/Runtime/GeometryCore/Private/Generators/StairGenerator.cpp
--- a/Engine/Source/Runtime/GeometryCore/Private/Generators/StairGenerator.cpp
+++ b/Engine/Source/Runtime/GeometryCore/Private/Generators/StairGenerator.cpp
@@ -1,6 +1,7 @@
 // Copyright Epic Games, Inc. All Rights Reserved.
 
 #include "Generators/StairGenerator.h"
+#include "MathUtil.h"
 
 using namespace UE::Geometry;
 
@@ -10,7 +11,7 @@ using namespace UE::Geometry;
 
 FMeshShapeGenerator& FStairGenerator::Generate()
 {
-	Reset();
+	ResetData();
 
 	switch (StairStyle)
 	{
@@ -28,16 +29,19 @@ FMeshShapeGenerator& FStairGenerator::Generate()
  * Stair topology composition. (side view cross section)
  *
  * (EStairStyle::Solid)
- *                  .___.
+ *                  C___C
  *                  |   |          | z (height)
- *              .___|___|          |
+ *              C___C___I          |
  *              |   |   |          |_____ x (depth)
- *          .___|___|___|           \
+ *          C___C___|___I           \
  *          |   |   |   |            \ y (width)
- *  Row --> |___|___|___|
+ *  Row --> C___I___I___C
  *
  *          ^-- Column
  *
+ *          C = corner vertex
+ *          I = interior connect quad vertex
+ *
  */
 FMeshShapeGenerator& FStairGenerator::GenerateSolidStairs()
 {
@@ -46,22 +50,27 @@ FMeshShapeGenerator& FStairGenerator::GenerateSolidStairs()
 		return (x * (x + 1) / 2);
 	};
 
+	// Number of quads that connect the left & right sides.
 	const int NumConnectQuads = 4 * NumSteps;
+	// Number of verts whose adjacent edges on the left & right sides are not parallel.
 	const int NumCornerVerts = 4 * (NumSteps + 1);
 	NumVertsPerSide = TriangleNumber(NumSteps + 1) + NumSteps;
 	NumVerts = 2 * NumVertsPerSide;
 	NumQuadsPerSide = TriangleNumber(NumSteps);
 	NumQuads = 2 * NumQuadsPerSide + NumConnectQuads;
 
+	// Number of mesh attributes for the side quads assuming sharing.
 	const int NumSideQuadAttrs = NumVerts;
+	// Number of mesh attributes for the corner vertices.
 	const int NumConnectQuadCornerAttrs = 2 * NumCornerVerts;
+	// Number of mesh attributes for the interior vertices of the bottom and back connect quads.
 	const int NumConnectQuadInteriorAttrs = 4 * (NumSteps - 1);
 	NumAttrs = NumConnectQuadCornerAttrs + NumConnectQuadInteriorAttrs + NumSideQuadAttrs;
 	SetBufferSizes(NumVerts, 2 * NumQuads, NumAttrs, NumAttrs);
 
 	// Ordered Side lists for iteration.
-	const ESide RightLeftSides[2] = { ESide::Right, ESide::Left };
-	const ESide AllSides[6] = { ESide::Right, ESide::Left, ESide::Front, ESide::Top, ESide::Back, ESide::Bottom };
+	constexpr ESide RightLeftSides[2] = { ESide::Right, ESide::Left };
+	constexpr ESide AllSides[6] = { ESide::Right, ESide::Left, ESide::Front, ESide::Top, ESide::Back, ESide::Bottom };
 
 	// Generate vertices by vertical column per Right/Left side.
 	LeftSideColumnId = (NumSteps + 1);
@@ -363,35 +372,43 @@ FMeshShapeGenerator& FStairGenerator::GenerateSolidStairs()
  * Stair topology composition. (side view cross section)
  *
  * (EStairStyle::Floating)
- *                  .___.
+ *                  C___C
  *                  |   |          | z (height)
- *              .___|___|          |
+ *              C___C___I          |
  *              |   |   |          |_____ x (depth)
- *          .___|___|___|           \
+ *          C___C___C___C           \
  *          |   |   |                \ y (width)
- *  Row --> |___|___|
+ *  Row --> C___I___C
  *
  *          ^-- Column
  *
+ *          C = corner vertex
+ *          I = interior connect quad vertex.
  */
 FMeshShapeGenerator& FStairGenerator::GenerateFloatingStairs()
 {
+	// Number of quads that connect the left & right sides.
 	const int NumConnectQuads = 4 * NumSteps;
+	// Number of verts whose adjacent edges on the left & right sides are not parallel.
 	const int NumCornerVerts = 8 * NumSteps - 4;
+	
 	NumVertsPerSide = 4 * NumSteps;
 	NumVerts = 2 * NumVertsPerSide;
 	NumQuadsPerSide = 2 * NumSteps - 1;
 	NumQuads = 2 * NumQuadsPerSide + NumConnectQuads;
 
+	// Number of mesh attributes for the side quads assuming sharing.
 	const int NumSideQuadAttrs = NumVerts;
+	// Number of mesh attributes for the corner vertices.
 	const int NumConnectQuadCornerAttrs = 2 * NumCornerVerts;
-	const int NumConnectQuadInteriorAttrs = 2 * (NumSteps - 1);
+	// Number of mesh attributes for the interior vertices of the bottom and back connect quads.
+	constexpr int NumConnectQuadInteriorAttrs = 4;
 	NumAttrs = NumConnectQuadCornerAttrs + NumConnectQuadInteriorAttrs + NumSideQuadAttrs;
 	SetBufferSizes(NumVerts, 2 * NumQuads, NumAttrs, NumAttrs);
 
 	// Ordered Side lists for iteration.
-	const ESide RightLeftSides[2] = { ESide::Right, ESide::Left };
-	const ESide AllSides[6] = { ESide::Right, ESide::Left, ESide::Front, ESide::Top, ESide::Back, ESide::Bottom };
+	constexpr ESide RightLeftSides[2] = { ESide::Right, ESide::Left };
+	constexpr ESide AllSides[6] = { ESide::Right, ESide::Left, ESide::Front, ESide::Top, ESide::Back, ESide::Bottom };
 
 	// Generate vertices by vertical column per Right/Left side.
 	LeftSideColumnId = (NumSteps + 1);
@@ -816,11 +833,11 @@ FStairGenerator::ESide FStairGenerator::FaceToSide(int FaceId)
  *
  * This is invoked at the head of the Generate() method.
  */
-void FStairGenerator::Reset()
+void FStairGenerator::ResetData()
 {
-	if (NumSteps <= 0)
+	if (NumSteps < 2)
 	{
-		NumSteps = 1;
+		NumSteps = 2;
 	}
 
 	VertexIds.Reset();
@@ -1025,9 +1042,9 @@ FVector2f FFloatingStairGenerator::GenerateUV(ESide Side, int Step, int VertexId
  * FCurvedStairGenerator
  */
 
-void FCurvedStairGenerator::Reset()
+void FCurvedStairGenerator::ResetData()
 {
-	Super::Reset();
+	Super::ResetData();
 
 	bIsClockwise = CurveAngle > 0.0f;
 	CurveRadians = CurveAngle * TMathUtilConstants<float>::DegToRad;
diff --git a/Engine/Source/Runtime/GeometryCore/Public/Generators/StairGenerator.h b/Engine/Source/Runtime/GeometryCore/Public/Generators/StairGenerator.h
--- a/Engine/Source/Runtime/GeometryCore/Public/Generators/StairGenerator.h
+++ b/Engine/Source/Runtime/GeometryCore/Public/Generators/StairGenerator.h
@@ -3,7 +3,6 @@
 
 #include "Generators/MeshShapeGenerator.h"
 #include "IndexTypes.h"
-#include "MathUtil.h"
 
 
 namespace UE
@@ -44,7 +43,9 @@ public:
 	};
 
 public:
-	virtual ~FStairGenerator() {}
+	virtual ~FStairGenerator() override
+	{
+	}
 
 	/** Generate the mesh */
 	virtual FMeshShapeGenerator& Generate() override;
@@ -102,7 +103,7 @@ protected:
 
 protected:
 	/** Reset state data on the generator */
-	virtual void Reset();
+	virtual void ResetData();
 
 	/**
 	 * Returns a vertex position.
@@ -173,7 +174,9 @@ public:
 	{
 		StairStyle = EStairStyle::Solid;
 	}
-	virtual ~FLinearStairGenerator() {}
+	virtual ~FLinearStairGenerator() override
+	{
+	}
 
 protected:
 	typedef FStairGenerator Super;
@@ -195,7 +198,9 @@ public:
 	{
 		StairStyle = EStairStyle::Floating;
 	}
-	virtual ~FFloatingStairGenerator() {}
+	virtual ~FFloatingStairGenerator() override
+	{
+	}
 
 protected:
 	typedef FLinearStairGenerator Super;
@@ -222,7 +227,9 @@ public:
 	{
 		StairStyle = EStairStyle::Solid;
 	}
-	virtual ~FCurvedStairGenerator() {}
+	virtual ~FCurvedStairGenerator() override
+	{
+	}
 
 protected:
 	typedef FStairGenerator Super;
@@ -236,7 +243,7 @@ protected:
 	FVector3f BackNormal = FVector3f::Zero();
 
 protected:
-	virtual void Reset() override;
+	virtual void ResetData() override;
 	virtual FVector3d GenerateVertex(ESide Side, int VertexColumn, int VertexRow) override;
 	virtual FVector3f GenerateNormal(ESide Side, int VertexId) override;
 	virtual FVector2f GenerateUV(ESide Side, int Step, int VertexId, float UVScale) override;
@@ -254,7 +261,9 @@ public:
 	{
 		StairStyle = EStairStyle::Floating;
 	}
-	virtual ~FSpiralStairGenerator() {}
+	virtual ~FSpiralStairGenerator() override
+	{
+	}
 
 protected:
 	typedef FCurvedStairGenerator Super;
