GeometryProcessing: remove contents of ExplicitUseGeometryMathTypes.h, this header is no longer necessary as type conflicts w/ LWC have largely been resolved. Emptying the header to prevent it from being used as an include-all for GeometryCore. Usage will be removed in future.
#rnx

diff --git a/Engine/Plugins/Experimental/MeshModelingToolsetExp/Source/ModelingComponents/Private/Mechanics/SpaceCurveDeformationMechanic.cpp b/Engine/Plugins/Experimental/MeshModelingToolsetExp/Source/ModelingComponents/Private/Mechanics/SpaceCurveDeformationMechanic.cpp
--- a/Engine/Plugins/Experimental/MeshModelingToolsetExp/Source/ModelingComponents/Private/Mechanics/SpaceCurveDeformationMechanic.cpp
+++ b/Engine/Plugins/Experimental/MeshModelingToolsetExp/Source/ModelingComponents/Private/Mechanics/SpaceCurveDeformationMechanic.cpp
@@ -16,6 +16,7 @@
 #include "ToolSetupUtil.h"
 #include "Transforms/MultiTransformer.h"
 #include "ToolDataVisualizer.h"
+#include "IntBoxTypes.h"
 
 #include "ExplicitUseGeometryMathTypes.h"		// using UE::Geometry::(math types)
 using namespace UE::Geometry;
diff --git a/Engine/Source/Runtime/GeometryCore/Public/ExplicitUseGeometryMathTypes.h b/Engine/Source/Runtime/GeometryCore/Public/ExplicitUseGeometryMathTypes.h
--- a/Engine/Source/Runtime/GeometryCore/Public/ExplicitUseGeometryMathTypes.h
+++ b/Engine/Source/Runtime/GeometryCore/Public/ExplicitUseGeometryMathTypes.h
@@ -2,123 +2,4 @@
 
 #pragma once
 
-//
-// This header exists to simplify usage of the core UE::Geometry:: math types, such as FVector3d/f, FTransform3d/f, etc. 
-// The GeometryProcessing plugin used these type names for several years before the UE Core math library was extended with
-// double-precision support, at which point it was decided to rename FVector as FVector3f/FVector3d, etc, creating a conflict. 
-// Unfortunately the UE Core versions are not a direct replacement for the GeometryProcessing variants, because:
-//   1) they have templated base classes, and significant other template code takes advantage of this. The UE Core Math types currently do not.
-//   2) they use more reliable (but potentially slower) standard-library math functions
-//   3) they have function naming schemes compatible with other standard vector math libraries (which simplifies porting/etc)
-//   4) they have internal member differences in some case (eg UE::Geometry::FTransform3d does not use SSE, Matrix2/3/4 stores as Vectors, etc)
-//
-// As a result the global ::FVector3d and UE::Geometry::FVector3d names are ambiguous. To avoid requiring explicit UE::Geometry:: disambiguation 
-// in all cases, which would make code harder to read and vastly complicate UE4/UE5 merging, in code that primarily uses GeometryProcessing
-// types, we prefer to explicitly specify - via C++ 'using' declarations - that we are using the UE::Geometry:: types at cpp file scope.
-// Rather than add a standard block of declarations into each file, this header can be included **IN THE CPP ONLY** to do so.
-//
-// Note that the header does not contain 'using namespace UE::Geometry;' - we leave that to the cpp file, so the standard boilerplate is:
-//
-//		#include "ExplicitUseGeometryMathTypes.h"		// using UE::Geometry::(math types)
-//		using namespace UE::Geometry;
-//
-// This header **MUST NOT BE INCLUDED IN OTHER HEADERS** as this could result in ambiguous definitions for downstream includers.
-// In headers, the GeometryProcessing core math types can be explicitly 'used' in other namespaces, or in class definitions, but not in global scope.
-//
-
-
-#include "IndexTypes.h"
-#include "VectorTypes.h"
-#include "IntVectorTypes.h"
-#include "MatrixTypes.h"
-#include "Quaternion.h"
-#include "BoxTypes.h"
-#include "IntBoxTypes.h"
-#include "PlaneTypes.h"
-#include "HalfspaceTypes.h"
-#include "LineTypes.h"
-#include "SegmentTypes.h"
-#include "RayTypes.h"
-#include "SphereTypes.h"
-#include "CapsuleTypes.h"
-#include "OrientedBoxTypes.h"
-#include "TriangleTypes.h"
-#include "CircleTypes.h"
-#include "FrameTypes.h"
-#include "TransformTypes.h"
-
-using UE::Geometry::FIndex2i;
-using UE::Geometry::FIndex3i;
-using UE::Geometry::FIndex4i;
-
-//using UE::Geometry::FVector2f;
-//using UE::Geometry::FVector2d;
-using UE::Geometry::FVector2i;
-//using UE::Geometry::FVector3f;
-//using UE::Geometry::FVector3d;
-using UE::Geometry::FVector3i;
-//using UE::Geometry::FVector4f;
-//using UE::Geometry::FVector4d;
-
-using UE::Geometry::FMatrix3f;
-using UE::Geometry::FMatrix3d;
-using UE::Geometry::FMatrix2f;
-using UE::Geometry::FMatrix2d;
-
-using UE::Geometry::FQuaternionf;
-using UE::Geometry::FQuaterniond;
-
-using UE::Geometry::FInterval1f;
-using UE::Geometry::FInterval1d;
-using UE::Geometry::FInterval1i;
-using UE::Geometry::FAxisAlignedBox2f;
-using UE::Geometry::FAxisAlignedBox2d;
-using UE::Geometry::FAxisAlignedBox2i;
-using UE::Geometry::FAxisAlignedBox3f;
-using UE::Geometry::FAxisAlignedBox3d;
-using UE::Geometry::FAxisAlignedBox3i;
-
-using UE::Geometry::FHalfspace3d;
-using UE::Geometry::FHalfspace3f;
-
-//using UE::Geometry::FPlane3d;
-//using UE::Geometry::FPlane3f;
-
-using UE::Geometry::FLine2f;
-using UE::Geometry::FLine2d;
-using UE::Geometry::FLine3f;
-using UE::Geometry::FLine3d;
-
-using UE::Geometry::FSegment2f;
-using UE::Geometry::FSegment2d;
-using UE::Geometry::FSegment3f;
-using UE::Geometry::FSegment3d;
-
-//using UE::Geometry::FRay3f;
-//using UE::Geometry::FRay3d;
-
-//using UE::Geometry::FSphere3d;
-//using UE::Geometry::FSphere3f;
-
-using UE::Geometry::FCapsule3d;
-using UE::Geometry::FCapsule3f;
-
-//using UE::Geometry::FOrientedBox3d;
-//using UE::Geometry::FOrientedBox3f;
-
-using UE::Geometry::FTriangle2d;
-using UE::Geometry::FTriangle2f;
-using UE::Geometry::FTriangle3d;
-using UE::Geometry::FTriangle3f;
-
-using UE::Geometry::FCircle2d;
-using UE::Geometry::FCircle2f;
-using UE::Geometry::FCircle3d;
-using UE::Geometry::FCircle3f;
-
-//using UE::Geometry::FTransform3d;
-//using UE::Geometry::FTransform3f;
-
-using UE::Geometry::FFrame3d;
-using UE::Geometry::FFrame3f;
-
+// this header is no longer necessary and will be removed in future
