Update TriangulateSimplePolygon to by-default orient its output triangles in the most-often-desired orientation (w/ triangle winding the opposite of the polygon winding, for example this is desirable when the polygon is generated from boundary edges on a hole to be filled)

Add test coverage for polygon triangulation functions to ensure they create triangulations with the expected orientations

#rnx

diff --git a/Engine/Plugins/Experimental/PlanarCutPlugin/Source/PlanarCut/Private/GeometryMeshConversion.cpp b/Engine/Plugins/Experimental/PlanarCutPlugin/Source/PlanarCut/Private/GeometryMeshConversion.cpp
--- a/Engine/Plugins/Experimental/PlanarCutPlugin/Source/PlanarCut/Private/GeometryMeshConversion.cpp
+++ b/Engine/Plugins/Experimental/PlanarCutPlugin/Source/PlanarCut/Private/GeometryMeshConversion.cpp
@@ -888,7 +888,7 @@ namespace AugmentedDynamicMesh
 			{
 				VertexPositions.Add(Mesh.GetVertex(AllBoundaryVerts[Idx]));
 			}
-			PolygonTriangulation::TriangulateSimplePolygon(VertexPositions, Triangles);
+			PolygonTriangulation::TriangulateSimplePolygon(VertexPositions, Triangles, true);
 			double HoleArea = 0;
 			FVector3d LastNormal(0, 0, 0);
 			const double FlippedThreshold = -.5; // Allow some angle deviation but reject hole fills that are too folded
@@ -1016,11 +1016,11 @@ namespace AugmentedDynamicMesh
 				FVector3d P1 = Mesh.GetVertex(V1);
 				FVector3d P2 = Mesh.GetVertex(V2);
 
-				int32 TID = Mesh.AppendTriangle(V0, V2, V1);
+				int32 TID = Mesh.AppendTriangle(V0, V1, V2);
 				// append can fail from a non-manifold edge; in that case add new vertices and try again
 				if (TID == FDynamicMesh3::NonManifoldID)
 				{
-					int32 SeparatedVIDs[]{ V0, V2, V1 };
+					int32 SeparatedVIDs[]{ V0, V1, V2 };
 					for (int32 SubIdx = 0; SubIdx < 3; SubIdx++)
 					{
 						if (SeparatedVIDs[SubIdx] < PrevMaxVID)
@@ -1584,7 +1584,7 @@ void FCellMeshes::CreateMeshesForBoundedPlanesWithNoise(int NumCells, const FPla
 			Triangulation.Triangulate();
 			if (Triangulation.Triangles.Num() == 0) // fall back to ear clipping if the triangulation came back empty
 			{
-				PolygonTriangulation::TriangulateSimplePolygon(Polygon.GetVertices(), Triangulation.Triangles);
+				PolygonTriangulation::TriangulateSimplePolygon(Polygon.GetVertices(), Triangulation.Triangles, false);
 			}
 			if (ensure(Triangulation.Triangles.Num() > 0))
 			{
diff --git a/Engine/Plugins/Experimental/Water/Source/Runtime/Private/WaterBooleanUtils.cpp b/Engine/Plugins/Experimental/Water/Source/Runtime/Private/WaterBooleanUtils.cpp
--- a/Engine/Plugins/Experimental/Water/Source/Runtime/Private/WaterBooleanUtils.cpp
+++ b/Engine/Plugins/Experimental/Water/Source/Runtime/Private/WaterBooleanUtils.cpp
@@ -161,8 +161,7 @@ void FWaterBooleanUtils::ExtractMesh(AVolume* Volume, FDynamicMesh3& Mesh)
 
 			for (FIndex3i Tri : PolyTriangles)
 			{
-				// flip orientation here...
-				Mesh.AppendTriangle(VertIndices[Tri.A], VertIndices[Tri.C], VertIndices[Tri.B]);
+				Mesh.AppendTriangle(VertIndices[Tri.A], VertIndices[Tri.B], VertIndices[Tri.C]);
 			}
 		}
 	}
diff --git a/Engine/Plugins/Runtime/GeometryProcessing/Source/DynamicMesh/Private/Operations/MeshBevel.cpp b/Engine/Plugins/Runtime/GeometryProcessing/Source/DynamicMesh/Private/Operations/MeshBevel.cpp
--- a/Engine/Plugins/Runtime/GeometryProcessing/Source/DynamicMesh/Private/Operations/MeshBevel.cpp
+++ b/Engine/Plugins/Runtime/GeometryProcessing/Source/DynamicMesh/Private/Operations/MeshBevel.cpp
@@ -891,7 +891,7 @@ void FMeshBevel::AppendJunctionVertexPolygon(FDynamicMesh3& Mesh, FBevelVertex&
 		int32 A = Vertex.Wedges[Tri.A].WedgeVertex;
 		int32 B = Vertex.Wedges[Tri.B].WedgeVertex;
 		int32 C = Vertex.Wedges[Tri.C].WedgeVertex;
-		int32 tid = Mesh.AppendTriangle(A, C, B, Vertex.NewGroupID);		// cycle orientation here because TriangulateSimplePolygon is wrong
+		int32 tid = Mesh.AppendTriangle(A, B, C, Vertex.NewGroupID);
 		if (Mesh.IsTriangle(tid))
 		{
 			Vertex.NewTriangles.Add(tid);
diff --git a/Engine/Plugins/Runtime/GeometryProcessing/Source/DynamicMesh/Private/Operations/SimpleHoleFiller.cpp b/Engine/Plugins/Runtime/GeometryProcessing/Source/DynamicMesh/Private/Operations/SimpleHoleFiller.cpp
--- a/Engine/Plugins/Runtime/GeometryProcessing/Source/DynamicMesh/Private/Operations/SimpleHoleFiller.cpp
+++ b/Engine/Plugins/Runtime/GeometryProcessing/Source/DynamicMesh/Private/Operations/SimpleHoleFiller.cpp
@@ -97,8 +97,8 @@ bool FSimpleHoleFiller::Fill_EarClip(int GroupID)
 	{
 		FIndex3i MeshTriangle(
 			Loop.Vertices[PolyTriangle.A],
-			Loop.Vertices[PolyTriangle.C],
-			Loop.Vertices[PolyTriangle.B]);  // Reversing orientation here!!
+			Loop.Vertices[PolyTriangle.B],
+			Loop.Vertices[PolyTriangle.C]);
 		int32 NewTriangle = Mesh->AppendTriangle(MeshTriangle, GroupID);
 		if (NewTriangle >= 0)
 		{
diff --git a/Engine/Plugins/Runtime/GeometryProcessing/Source/GeometryAlgorithms/Public/ConstrainedDelaunay2.h b/Engine/Plugins/Runtime/GeometryProcessing/Source/GeometryAlgorithms/Public/ConstrainedDelaunay2.h
--- a/Engine/Plugins/Runtime/GeometryProcessing/Source/GeometryAlgorithms/Public/ConstrainedDelaunay2.h
+++ b/Engine/Plugins/Runtime/GeometryProcessing/Source/GeometryAlgorithms/Public/ConstrainedDelaunay2.h
@@ -146,6 +146,7 @@ struct TConstrainedDelaunay2
 	bool GEOMETRYALGORITHMS_API Triangulate(TFunctionRef<bool(const TArray<TVector2<RealType>>&, const FIndex3i&)> KeepTriangle);
 };
 
+/** Convenience function; invokes TConstrainedDelaunay2::Triangulate with most-often-used options */
 template<typename RealType>
 TArray<FIndex3i> GEOMETRYALGORITHMS_API ConstrainedDelaunayTriangulate(const TGeneralPolygon2<RealType>& GeneralPolygon);
 
diff --git a/Engine/Source/Runtime/GeometryCore/Private/CompGeom/PolygonTriangulation.cpp b/Engine/Source/Runtime/GeometryCore/Private/CompGeom/PolygonTriangulation.cpp
--- a/Engine/Source/Runtime/GeometryCore/Private/CompGeom/PolygonTriangulation.cpp
+++ b/Engine/Source/Runtime/GeometryCore/Private/CompGeom/PolygonTriangulation.cpp
@@ -9,14 +9,14 @@ using namespace UE::Geometry;
 // explicit instantiations
 namespace PolygonTriangulation
 {
-	template GEOMETRYCORE_API void TriangulateSimplePolygon<float>(const TArray<TVector2<float>>& VertexPositions, TArray<FIndex3i>& OutTriangles);
-	template GEOMETRYCORE_API void TriangulateSimplePolygon<double>(const TArray<TVector2<double>>& VertexPositions, TArray<FIndex3i>& OutTriangles);
+	template GEOMETRYCORE_API void TriangulateSimplePolygon<float>(const TArray<TVector2<float>>& VertexPositions, TArray<FIndex3i>& OutTriangles, bool bOrientAsHoleFill);
+	template GEOMETRYCORE_API void TriangulateSimplePolygon<double>(const TArray<TVector2<double>>& VertexPositions, TArray<FIndex3i>& OutTriangles, bool bOrientAsHoleFill);
 
 	template GEOMETRYCORE_API void ComputePolygonPlane<float>(const TArray<TVector<float>>& VertexPositions, TVector<float>& NormalOut, TVector<float>& PlanePointOut);
 	template GEOMETRYCORE_API void ComputePolygonPlane<double>(const TArray<TVector<double>>& VertexPositions, TVector<double>& NormalOut, TVector<double>& PlanePointOut);
 
-	template GEOMETRYCORE_API void TriangulateSimplePolygon<float>(const TArray<TVector<float>>& VertexPositions, TArray<FIndex3i>& OutTriangles);
-	template GEOMETRYCORE_API void TriangulateSimplePolygon<double>(const TArray<TVector<double>>& VertexPositions, TArray<FIndex3i>& OutTriangles);
+	template GEOMETRYCORE_API void TriangulateSimplePolygon<float>(const TArray<TVector<float>>& VertexPositions, TArray<FIndex3i>& OutTriangles, bool bOrientAsHoleFill);
+	template GEOMETRYCORE_API void TriangulateSimplePolygon<double>(const TArray<TVector<double>>& VertexPositions, TArray<FIndex3i>& OutTriangles, bool bOrientAsHoleFill);
 }
 
 
@@ -27,7 +27,7 @@ namespace PolygonTriangulation
 // This is based on the 3D triangulation code from MeshDescription.cpp, simplified for 2D polygons
 // 
 template<typename T>
-void PolygonTriangulation::TriangulateSimplePolygon(const TArray<TVector2<T>>& VertexPositions, TArray<FIndex3i>& OutTriangles)
+void PolygonTriangulation::TriangulateSimplePolygon(const TArray<TVector2<T>>& VertexPositions, TArray<FIndex3i>& OutTriangles, bool bOrientAsHoleFill)
 {
 	struct Local
 	{
@@ -44,15 +44,15 @@ void PolygonTriangulation::TriangulateSimplePolygon(const TArray<TVector2<T>>& V
 
 
 	// compute signed area of polygon
-	double PolySignedArea = 0;
+	double PolySignedArea2 = 0;
 	for (int i = 0; i < PolygonVertexCount; ++i)
 	{
 		const TVector2<T>& v1 = VertexPositions[i];
 		const TVector2<T>& v2 = VertexPositions[(i + 1) % PolygonVertexCount];
-		PolySignedArea += v1.X*v2.Y - v1.Y*v2.X;
+		PolySignedArea2 += v1.X*v2.Y - v1.Y*v2.X;
 	}
-	PolySignedArea *= 0.5;
-	bool bIsClockwise = PolySignedArea < 0;
+
+	bool bIsClockwise = PolySignedArea2 < 0;
 	double OrientationSign = (bIsClockwise) ? -1.0 : 1.0;
 
 
@@ -62,7 +62,7 @@ void PolygonTriangulation::TriangulateSimplePolygon(const TArray<TVector2<T>>& V
 	// If perimeter has 3 vertices, just copy content of perimeter out 
 	if (PolygonVertexCount == 3)
 	{
-		OutTriangles.Add(FIndex3i(0, 1, 2));
+		OutTriangles.Add(bOrientAsHoleFill ? FIndex3i(0, 2, 1) : FIndex3i(0, 1, 2));
 		return;
 	}
 
@@ -129,11 +129,10 @@ void PolygonTriangulation::TriangulateSimplePolygon(const TArray<TVector2<T>>& V
 		{
 			// OK, we found an ear!  Let's save this triangle in our output buffer.
 			{
-				OutTriangles.Emplace();
-				FIndex3i& Triangle = OutTriangles.Last();
-				Triangle.A = PrevVertexNumbers[EarVertexNumber];
-				Triangle.B = EarVertexNumber;
-				Triangle.C = NextVertexNumbers[EarVertexNumber];
+				int32 A = PrevVertexNumbers[EarVertexNumber]
+					, B = EarVertexNumber
+					, C = NextVertexNumbers[EarVertexNumber];
+				OutTriangles.Add(bOrientAsHoleFill ? FIndex3i(A, C, B) : FIndex3i(A, B, C));
 			}
 
 			// Update our linked list.  We're effectively cutting off the ear by pointing the ear vertex's neighbors to
@@ -229,7 +228,7 @@ static bool PointInTriangle(const FVector& A, const FVector& B, const FVector& C
 // This is based on the 3D triangulation code from MeshDescription.cpp, simplified for 2D polygons
 // 
 template<typename T>
-void PolygonTriangulation::TriangulateSimplePolygon(const TArray<TVector<T>>& VertexPositions, TArray<FIndex3i>& OutTriangles)
+void PolygonTriangulation::TriangulateSimplePolygon(const TArray<TVector<T>>& VertexPositions, TArray<FIndex3i>& OutTriangles, bool bOrientAsHoleFill)
 {
 	struct Local3
 	{
@@ -264,7 +263,7 @@ void PolygonTriangulation::TriangulateSimplePolygon(const TArray<TVector<T>>& Ve
 	// If perimeter has 3 vertices, just copy content of perimeter out 
 	if (PolygonVertexCount == 3)
 	{
-		OutTriangles.Add(FIndex3i(0, 1, 2));
+		OutTriangles.Add(bOrientAsHoleFill ? FIndex3i(0, 2, 1) : FIndex3i(0, 1, 2));
 		return;
 	}
 
@@ -336,11 +335,10 @@ void PolygonTriangulation::TriangulateSimplePolygon(const TArray<TVector<T>>& Ve
 		{
 			// OK, we found an ear!  Let's save this triangle in our output buffer.
 			{
-				OutTriangles.Emplace();
-				FIndex3i& Triangle = OutTriangles.Last();
-				Triangle.A = PrevVertexNumbers[EarVertexNumber];
-				Triangle.B = EarVertexNumber;
-				Triangle.C = NextVertexNumbers[EarVertexNumber];
+				int32 A = PrevVertexNumbers[EarVertexNumber]
+					, B = EarVertexNumber
+					, C = NextVertexNumbers[EarVertexNumber];
+				OutTriangles.Add(bOrientAsHoleFill ? FIndex3i(A, C, B) : FIndex3i(A, B, C));
 			}
 
 			// Update our linked list.  We're effectively cutting off the ear by pointing the ear vertex's neighbors to
diff --git a/Engine/Source/Runtime/GeometryCore/Private/Generators/PlanarPolygonMeshGenerator.cpp b/Engine/Source/Runtime/GeometryCore/Private/Generators/PlanarPolygonMeshGenerator.cpp
--- a/Engine/Source/Runtime/GeometryCore/Private/Generators/PlanarPolygonMeshGenerator.cpp
+++ b/Engine/Source/Runtime/GeometryCore/Private/Generators/PlanarPolygonMeshGenerator.cpp
@@ -36,7 +36,7 @@ FMeshShapeGenerator& FPlanarPolygonMeshGenerator::Generate()
 	}
 
 	TArray<FIndex3i> TriangleList;
-	PolygonTriangulation::TriangulateSimplePolygon(Polygon.GetVertices(), TriangleList);
+	PolygonTriangulation::TriangulateSimplePolygon(Polygon.GetVertices(), TriangleList, false);
 
 	int NumTriangles = TriangleList.Num();
 	SetBufferSizes(NumVertices, NumTriangles, NumVertices, NumVertices);
diff --git a/Engine/Source/Runtime/GeometryCore/Private/Generators/RevolveGenerator.cpp b/Engine/Source/Runtime/GeometryCore/Private/Generators/RevolveGenerator.cpp
--- a/Engine/Source/Runtime/GeometryCore/Private/Generators/RevolveGenerator.cpp
+++ b/Engine/Source/Runtime/GeometryCore/Private/Generators/RevolveGenerator.cpp
@@ -279,7 +279,7 @@ FDynamicMesh3 FRevolvePlanarPathGenerator::GenerateMesh()
 			PolygonTriangulation::TriangulateSimplePolygon(VertexPositions, Triangles);
 			for (FIndex3i PolyTriangle : Triangles)
 			{
-				FIndex3i MeshTriangle( VertexLoop[PolyTriangle.A], VertexLoop[PolyTriangle.C], VertexLoop[PolyTriangle.B]);  // Reversing orientation here!!
+				FIndex3i MeshTriangle( VertexLoop[PolyTriangle.A], VertexLoop[PolyTriangle.B], VertexLoop[PolyTriangle.C]);
 				int32 NewTriangle = ResultMesh.AppendTriangle(MeshTriangle, NewGroupID);
 				if (NewTriangle >= 0)
 				{
@@ -469,7 +469,7 @@ FDynamicMesh3 FRevolvePlanarPolygonGenerator::GenerateMesh()
 			PolygonTriangulation::TriangulateSimplePolygon(VertexPositions, Triangles);
 			for (FIndex3i PolyTriangle : Triangles)
 			{
-				FIndex3i MeshTriangle( VertexLoop[PolyTriangle.A], VertexLoop[PolyTriangle.C], VertexLoop[PolyTriangle.B]);  // Reversing orientation here!!
+				FIndex3i MeshTriangle( VertexLoop[PolyTriangle.A], VertexLoop[PolyTriangle.B], VertexLoop[PolyTriangle.C]);
 				int32 NewTriangle = ResultMesh.AppendTriangle(MeshTriangle, NewGroupID);
 				if (NewTriangle >= 0)
 				{
@@ -657,7 +657,7 @@ FDynamicMesh3 FSpiralRevolvePlanarPolygonGenerator::GenerateMesh()
 			PolygonTriangulation::TriangulateSimplePolygon(VertexPositions, Triangles);
 			for (FIndex3i PolyTriangle : Triangles)
 			{
-				FIndex3i MeshTriangle( VertexLoop[PolyTriangle.A], VertexLoop[PolyTriangle.C], VertexLoop[PolyTriangle.B]);  // Reversing orientation here!!
+				FIndex3i MeshTriangle( VertexLoop[PolyTriangle.A], VertexLoop[PolyTriangle.B], VertexLoop[PolyTriangle.C]);
 				int32 NewTriangle = ResultMesh.AppendTriangle(MeshTriangle, NewGroupID);
 				if (NewTriangle >= 0)
 				{
diff --git a/Engine/Source/Runtime/GeometryCore/Public/CompGeom/PolygonTriangulation.h b/Engine/Source/Runtime/GeometryCore/Public/CompGeom/PolygonTriangulation.h
--- a/Engine/Source/Runtime/GeometryCore/Public/CompGeom/PolygonTriangulation.h
+++ b/Engine/Source/Runtime/GeometryCore/Public/CompGeom/PolygonTriangulation.h
@@ -14,9 +14,10 @@ namespace PolygonTriangulation
 	 * Compute triangulation of simple 2D polygon using ear-clipping
 	 * @param VertexPositions ordered vertices of 2D polygon
 	 * @param OutTriangles computed triangulation. Each triangle is a tuple of indices into VertexPositions.
+	 * @param bOrientAsHoleFill if true, output triangles are wound opposite of the input polygon; this is typically desired when the polygon comes from the boundary of a hole
 	 */
 	template<typename RealType>
-	void GEOMETRYCORE_API TriangulateSimplePolygon(const TArray<TVector2<RealType>>& VertexPositions, TArray<FIndex3i>& OutTriangles);
+	void GEOMETRYCORE_API TriangulateSimplePolygon(const TArray<TVector2<RealType>>& VertexPositions, TArray<FIndex3i>& OutTriangles, bool bOrientAsHoleFill = true);
 
 
 	template<typename RealType>
@@ -27,8 +28,9 @@ namespace PolygonTriangulation
 	 * Compute triangulation of 3D simple polygon using ear-clipping
 	 * @param VertexPositions ordered vertices of 3D polygon
 	 * @param OutTriangles computed triangulation. Each triangle is a tuple of indices into VertexPositions.
+	 * @param bOrientAsHoleFill if true, output triangles are wound opposite of the input polygon; this is typically desired when the polygon comes from the boundary of a hole
 	 */
 	template<typename RealType>
-	void GEOMETRYCORE_API TriangulateSimplePolygon(const TArray<TVector<RealType>>& VertexPositions, TArray<FIndex3i>& OutTriangles);
+	void GEOMETRYCORE_API TriangulateSimplePolygon(const TArray<TVector<RealType>>& VertexPositions, TArray<FIndex3i>& OutTriangles, bool bOrientAsHoleFill = true);
 
 }
\ No newline at end of file
diff --git a/Engine/Source/Runtime/GeometryCore/Public/Generators/SweepGenerator.h b/Engine/Source/Runtime/GeometryCore/Public/Generators/SweepGenerator.h
--- a/Engine/Source/Runtime/GeometryCore/Public/Generators/SweepGenerator.h
+++ b/Engine/Source/Runtime/GeometryCore/Public/Generators/SweepGenerator.h
@@ -163,7 +163,6 @@ protected:
 					for (const FIndex3i& Triangle : OutTriangles)
 					{
 						bool Flipped = CapIdx == 0;
-						Flipped = !Flipped;
 						SetTriangle(TriIdx,
 							Triangle.A + VertOffset, Triangle.B + VertOffset, Triangle.C + VertOffset,
 							Flipped);
