MeshModelingTools: Add bAllowElementFreeing to dynamic mesh overlay UnsetTriangle method to make it easier to do certain kinds of retriangulations. Add a unit test.

#rnx

diff --git a/Engine/Source/Runtime/GeometryCore/Private/DynamicMesh/DynamicMeshOverlay.cpp b/Engine/Source/Runtime/GeometryCore/Private/DynamicMesh/DynamicMeshOverlay.cpp
--- a/Engine/Source/Runtime/GeometryCore/Private/DynamicMesh/DynamicMeshOverlay.cpp
+++ b/Engine/Source/Runtime/GeometryCore/Private/DynamicMesh/DynamicMeshOverlay.cpp
@@ -413,7 +413,7 @@ void UE::Geometry::TDynamicMeshOverlay<RealType, ElementSize>::FreeUnusedElement
 }
 
 template<typename RealType, int ElementSize>
-void TDynamicMeshOverlay<RealType, ElementSize>::UnsetTriangle(int TriangleID)
+void TDynamicMeshOverlay<RealType, ElementSize>::UnsetTriangle(int TriangleID, bool bAllowElementFreeing)
 {
 	int i = 3 * TriangleID;
 	if (ElementTriangles[i] == FDynamicMesh3::InvalidID)
@@ -424,7 +424,7 @@ void TDynamicMeshOverlay<RealType, ElementSize>::UnsetTriangle(int TriangleID)
 	{
 		ElementsRefCounts.Decrement(ElementTriangles[i + SubIdx]);
 
-		if (ElementsRefCounts.GetRefCount(ElementTriangles[i + SubIdx]) == 1)
+		if (bAllowElementFreeing && ElementsRefCounts.GetRefCount(ElementTriangles[i + SubIdx]) == 1)
 		{
 			ElementsRefCounts.Decrement(ElementTriangles[i + SubIdx]);
 			ParentVertices[ElementTriangles[i + SubIdx]] = FDynamicMesh3::InvalidID;
diff --git a/Engine/Source/Runtime/GeometryCore/Public/DynamicMesh/DynamicMeshOverlay.h b/Engine/Source/Runtime/GeometryCore/Public/DynamicMesh/DynamicMeshOverlay.h
--- a/Engine/Source/Runtime/GeometryCore/Public/DynamicMesh/DynamicMeshOverlay.h
+++ b/Engine/Source/Runtime/GeometryCore/Public/DynamicMesh/DynamicMeshOverlay.h
@@ -290,9 +290,12 @@ public:
 
 	/**
 	 * Set the triangle to have InvalidID element IDs, decrementing element reference counts if needed.
-	 * Deletes elements that are no longer used after the triangle is unset.
+	 * 
+	 * @param bAllowElementFreeing If true, then any elements that were only referenced by this triangle
+	 *  become immediately unallocated. This can be set to false as part of a remeshing, but then it 
+	 *  should eventually be followed by a call to FreeUnusedElements().
 	 */
-	void UnsetTriangle(int TriangleID);
+	void UnsetTriangle(int TriangleID, bool bAllowElementFreeing = true);
 
 	/** @return true if this triangle was set */
 	bool IsSetTriangle(int TID) const
