speculative CIS fix for mac link error (not finding TExtremePoints3::Init)

diff --git a/Engine/Source/Runtime/GeometryCore/Private/CompGeom/ConvexHull3.cpp b/Engine/Source/Runtime/GeometryCore/Private/CompGeom/ConvexHull3.cpp
--- a/Engine/Source/Runtime/GeometryCore/Private/CompGeom/ConvexHull3.cpp
+++ b/Engine/Source/Runtime/GeometryCore/Private/CompGeom/ConvexHull3.cpp
@@ -644,7 +644,8 @@ bool TConvexHull3<RealType>::Solve(int32 NumPoints, TFunctionRef<TVector<RealTyp
 	return true;
 }
 
-
+template struct TExtremePoints3<float>;
+template struct TExtremePoints3<double>;
 template class TConvexHull3<float>;
 template class TConvexHull3<double>;
 
diff --git a/Engine/Source/Runtime/GeometryCore/Public/CompGeom/ConvexHull3.h b/Engine/Source/Runtime/GeometryCore/Public/CompGeom/ConvexHull3.h
--- a/Engine/Source/Runtime/GeometryCore/Public/CompGeom/ConvexHull3.h
+++ b/Engine/Source/Runtime/GeometryCore/Public/CompGeom/ConvexHull3.h
@@ -40,14 +40,6 @@ struct GEOMETRYCORE_API TExtremePoints3
 		Init(NumPoints, GetPointFunc, FilterFunc, Epsilon);
 	}
 
-	TExtremePoints3(TArrayView<TVector<RealType>> Points, TFunctionRef<bool(int32)> FilterFunc = [](int32 Idx) {return true; })
-	{
-		Init(Points.Num(), [&Points](int32 Idx)
-			{
-				return Points[Idx];
-			}, FilterFunc);
-	}
-
 private:
 	void Init(int32 NumPoints, TFunctionRef<TVector<RealType>(int32)> GetPointFunc, TFunctionRef<bool(int32)> FilterFunc, double Epsilon = 0);
 };
