make segment2 intersection more robust and make both implementations consistent
 + add test coverage

#rnx

diff --git a/Engine/Source/Runtime/GeometryCore/Public/Intersection/IntrSegment2Segment2.h b/Engine/Source/Runtime/GeometryCore/Public/Intersection/IntrSegment2Segment2.h
--- a/Engine/Source/Runtime/GeometryCore/Public/Intersection/IntrSegment2Segment2.h
+++ b/Engine/Source/Runtime/GeometryCore/Public/Intersection/IntrSegment2Segment2.h
@@ -149,15 +149,67 @@ public:
 			return (Result == EIntersectionResult::Intersects);
 		}
 
-		// if either segment direction is not a normalized vector, 
-		//   results are garbage, so fail query
-		if (IsNormalized(Segment1.Direction) == false || IsNormalized(Segment2.Direction) == false)
+		// Special handling of degenerate segments; by convention if Direction was too small to normalize, Extent and Direction will be zero
+		bool bIsPoint1 = Segment1.Extent == 0;
+		bool bIsPoint2 = Segment2.Extent == 0;
+		int IsPointCount = (int)bIsPoint1 + (int)bIsPoint2;
+		auto SetPointIntersection = [this](TVector2<RealType> Point, RealType Param) -> void
 		{
+			Result = EIntersectionResult::Intersects;
+			Type = EIntersectionType::Point;
+			Quantity = 1;
+			Point0 = Point;
+			Parameter0 = Param;
+		};
+		auto SetNoIntersection = [this]() -> void
+		{
+			Quantity = 0;
+			Result = EIntersectionResult::NoIntersection;
 			Type = EIntersectionType::Empty;
-			Result = EIntersectionResult::InvalidQuery;
-			return false;
+		};
+		if (IsPointCount == 2) // both degenerate -- check if they're within IntervalThreshold of each other
+		{
+			bool bIntersects = UE::Geometry::DistanceSquared(Segment1.Center, Segment2.Center) <= IntervalThreshold * IntervalThreshold;
+			if (bIntersects)
+			{
+				SetPointIntersection((Segment1.Center + Segment2.Center) * .5, 0);
+			}
+			else
+			{
+				SetNoIntersection();
+			}
+			return bIntersects;
+		}
+		else if (IsPointCount == 1) // one degenerate -- check if it's within IntervalThreshold of the non-degenerate segment
+		{
+			RealType DistSq;
+			RealType Param = 0;
+			TVector2<RealType> Point;
+			if (bIsPoint1)
+			{
+				DistSq = Segment2.DistanceSquared(Segment1.Center);
+				Point = Segment1.Center;
+			}
+			else // bIsPoint2
+			{
+				DistSq = Segment1.DistanceSquared(Segment2.Center);
+				Param = Segment1.Project(Segment2.Center);
+				Point = Segment2.Center;
+			}
+			bool bIntersects = DistSq <= IntervalThreshold * IntervalThreshold;
+			if (bIntersects)
+			{
+				SetPointIntersection(Point, Param);
+			}
+			else
+			{
+				SetNoIntersection();
+			}
+			return bIntersects;
 		}
 
+		// If neither segment is zero, directions should always be normalized (will hold as long as segment is created/updated via the standard functions)
+		checkSlow(IsNormalized(Segment1.Direction) && IsNormalized(Segment2.Direction));
 
 		TVector2<RealType> s = TVector2<RealType>::Zero();
 		Type = TIntrLine2Line2<RealType>::Classify(Segment1.Center, Segment1.Direction,
@@ -170,14 +222,11 @@ public:
 			if (FMath::Abs(s[0]) <= Segment1.Extent + IntervalThreshold
 				&& FMath::Abs(s[1]) <= Segment2.Extent + IntervalThreshold)
 			{
-				Quantity = 1;
-				Point0 = Segment1.Center + s[0] * Segment1.Direction;
-				Parameter0 = s[0];
+				SetPointIntersection(Segment1.Center + s[0] * Segment1.Direction, s[0]);
 			}
 			else
 			{
-				Quantity = 0;
-				Type = EIntersectionType::Empty;
+				SetNoIntersection();
 			}
 		}
 		else if (Type == EIntersectionType::Line)
@@ -185,6 +234,8 @@ public:
 			// Compute the location of Segment2 endpoints relative to Segment1.
 			TVector2<RealType> diff = Segment2.Center - Segment1.Center;
 			RealType t1 = Segment1.Direction.Dot(diff);
+			// Note: IntervalThreshold not used here; this is a bit inconsistent but simplifies the code
+			// If you add it, make sure that segments that intersect without it don't end up with too-wide intersection intervals
 			RealType tmin = t1 - Segment2.Extent;
 			RealType tmax = t1 + Segment2.Extent;
 			TIntersector1<RealType> calc(-Segment1.Extent, Segment1.Extent, tmin, tmax);
diff --git a/Engine/Source/Runtime/GeometryCore/Public/SegmentTypes.h b/Engine/Source/Runtime/GeometryCore/Public/SegmentTypes.h
--- a/Engine/Source/Runtime/GeometryCore/Public/SegmentTypes.h
+++ b/Engine/Source/Runtime/GeometryCore/Public/SegmentTypes.h
@@ -217,6 +217,28 @@ public:
 	{
 		// see IntrLine2Line2 and IntrSegment2Segment2 for details on this code
 
+		// Special handling of degenerate segments; by convention if Direction was too small to normalize, Extent and Direction will be zero
+		bool bIsPoint = Extent == 0;
+		bool bIsPointOther = OtherSegment.Extent == 0;
+		int IsPointCount = (int)bIsPoint + (int)bIsPointOther;
+		if (IsPointCount == 2)
+		{
+			return UE::Geometry::DistanceSquared(Center, OtherSegment.Center) <= IntervalThresh * IntervalThresh;
+		}
+		else if (IsPointCount == 1)
+		{
+			T DistSq;
+			if (bIsPoint)
+			{
+				DistSq = OtherSegment.DistanceSquared(Center);
+			}
+			else
+			{
+				DistSq = DistanceSquared(OtherSegment.Center);
+			}
+			return DistSq <= IntervalThresh * IntervalThresh;
+		}
+
 		TVector2<T> diff = OtherSegment.Center - Center;
 		T D0DotPerpD1 = DotPerp(Direction, OtherSegment.Direction);
 		if (TMathUtil<T>::Abs(D0DotPerpD1) > DotThresh)      // Lines intersect in a single point.
@@ -238,6 +260,8 @@ public:
 			// Compute the location of OtherSegment endpoints relative to our Segment
 			diff = OtherSegment.Center - Center;
 			T t1 = Direction.Dot(diff);
+			// Note: IntervalThresh not used here; this is a bit inconsistent but simplifies the code
+			// If you add it, make sure that segments that intersect without it don't end up with too-wide intersection intervals
 			T tmin = t1 - OtherSegment.Extent;
 			T tmax = t1 + OtherSegment.Extent;
 			TInterval1<T> extents(-Extent, Extent);
