GeometryProcessing:
Add FColliderMesh, minimal 3D triangle mesh representation w/ a buit-in AABBTree, intended to be used to avoid heavier FDynamicMesh3 copies in places where only a BVH is needed.
Add Ray transform/inverse functions to Transform3d

#rnx

diff --git a/Engine/Source/Runtime/GeometryCore/Private/DynamicMesh/ColliderMesh.cpp b/Engine/Source/Runtime/GeometryCore/Private/DynamicMesh/ColliderMesh.cpp
new file mode 100644
--- /dev/null
+++ b/Engine/Source/Runtime/GeometryCore/Private/DynamicMesh/ColliderMesh.cpp
@@ -0,0 +1,133 @@
+// Copyright Epic Games, Inc. All Rights Reserved.
+
+#include "DynamicMesh/ColliderMesh.h"
+#include "DynamicMesh/DynamicMesh3.h"
+
+using namespace UE::Geometry;
+
+
+FColliderMesh::FColliderMesh(const FDynamicMesh3& SourceMesh, const FBuildOptions& BuildOptions)
+{
+	Initialize(SourceMesh, BuildOptions);
+}
+
+FColliderMesh::FColliderMesh()
+{
+	Reset(false);
+}
+
+void FColliderMesh::Reset(bool bReleaseMemory)
+{
+	Vertices.SetNum(0, bReleaseMemory);
+	SourceVertexIDs.SetNum(0, bReleaseMemory);
+	Triangles.SetNum(0, bReleaseMemory);
+	SourceTriangleIDs.SetNum(0, bReleaseMemory);
+
+	bSourceWasCompactV = bSourceWasCompactT = true;
+
+	AABBTree.SetMesh(this, true);
+}
+
+
+void FColliderMesh::Initialize(const FDynamicMesh3& SourceMesh, const FBuildOptions& BuildOptions)
+{
+	bSourceWasCompactV = SourceMesh.IsCompactV();
+	bSourceWasCompactT = SourceMesh.IsCompactT();
+
+	int32 NumVertices = SourceMesh.VertexCount();
+	Vertices.Reserve(NumVertices);
+	if (BuildOptions.bBuildVertexMap && bSourceWasCompactV == false)
+	{
+		SourceVertexIDs.Reserve(NumVertices);
+	}
+
+	int32 NumTriangles = SourceMesh.TriangleCount();
+	Triangles.Reserve(NumTriangles);
+	if (BuildOptions.bBuildTriangleMap && bSourceWasCompactT == false)
+	{
+		SourceTriangleIDs.Reserve(NumTriangles);
+	}
+
+	TArray<int32> SourceToCompactMapV;
+
+	if (bSourceWasCompactV)
+	{
+		for (FVector3d VertexPos : SourceMesh.VerticesItr())
+		{
+			Vertices.Add(VertexPos);
+		}
+	}
+	else
+	{
+		SourceToCompactMapV.SetNum(SourceMesh.MaxVertexID());
+		for (int32 VertexID : SourceMesh.VertexIndicesItr())
+		{
+			SourceToCompactMapV[VertexID] = Vertices.Num();
+			Vertices.Add(SourceMesh.GetVertex(VertexID));
+			if (BuildOptions.bBuildVertexMap)
+			{
+				SourceVertexIDs.Add(VertexID);
+			}
+		}
+	}
+
+	if (bSourceWasCompactT && bSourceWasCompactV)
+	{
+		for (FIndex3i TriIndices : SourceMesh.TrianglesItr())
+		{
+			Triangles.Add(TriIndices);
+		}
+	}
+	else
+	{
+		for (int32 TriangleID : SourceMesh.TriangleIndicesItr())
+		{
+			FIndex3i Tri = SourceMesh.GetTriangle(TriangleID);
+			if (!bSourceWasCompactV)
+			{
+				Tri.A = SourceToCompactMapV[Tri.A];
+				Tri.B = SourceToCompactMapV[Tri.B];
+				Tri.C = SourceToCompactMapV[Tri.C];
+			}
+			Triangles.Add(Tri);
+			if (BuildOptions.bBuildTriangleMap)
+			{
+				SourceTriangleIDs.Add(TriangleID);
+			}
+		}
+	}
+
+	AABBTree.SetMesh(this, BuildOptions.bBuildAABBTree);
+
+}
+
+
+TMeshAABBTree3<FColliderMesh>& FColliderMesh::GetAABBTree()
+{
+	return AABBTree;
+}
+
+
+int32 FColliderMesh::GetSourceVertexID(int32 VertexID) const 
+{ 
+	if (bSourceWasCompactV)
+	{
+		return VertexID;
+	}
+	else
+	{
+		return (VertexID >= 0 && VertexID < SourceVertexIDs.Num()) ? SourceVertexIDs[VertexID] : IndexConstants::InvalidID;
+	}
+}
+
+int32 FColliderMesh::GetSourceTriangleID(int32 TriangleID) const 
+{ 
+	if (bSourceWasCompactT)
+	{
+		return TriangleID;
+	}
+	else
+	{
+		return (TriangleID >= 0 && TriangleID < SourceTriangleIDs.Num()) ? SourceTriangleIDs[TriangleID] : IndexConstants::InvalidID;
+	}
+}
diff --git a/Engine/Source/Runtime/GeometryCore/Public/DynamicMesh/ColliderMesh.h b/Engine/Source/Runtime/GeometryCore/Public/DynamicMesh/ColliderMesh.h
new file mode 100644
--- /dev/null
+++ b/Engine/Source/Runtime/GeometryCore/Public/DynamicMesh/ColliderMesh.h
@@ -0,0 +1,120 @@
+// Copyright Epic Games, Inc. All Rights Reserved.
+
+#include "Spatial/MeshAABBTree3.h"
+
+#pragma once
+
+
+
+namespace UE
+{
+namespace Geometry
+{
+
+class FDynamicMesh3;
+
+/**
+ * FColliderMesh is a minimal representation of an Indexed Triangle Mesh suitable to use
+ * with a TMeshAABBTree3. This class is intended to be used in situations where a copy
+ * of a FDynamicMesh3 would be created/kept only to use with an AABBTree or FWNTree. In such 
+ * situations, the FDynamicMesh3 copy is (relatively) heavy memory-wise, particularly if a large 
+ * number of small mesh/AABBTree pairs are stored. In those cases a FColliderMesh can be used instead.
+ * 
+ * If FBuildOptions.bBuildAABBTree is true (default), then an AABBTree will automatically be
+ * built and available via GetAABBTree()
+ * 
+ * If the source FDynamicMesh3 is not compact, it will be compacted. ID mappings can optionally
+ * be stored if the FBuildOptions are configured to request it (default false), allowing 
+ * hit-triangle/vertex IDs to be mapped back to IDs on the source mesh.
+ * 
+ * Currently there is no support for UVs or Vertex Normals.
+ * 
+ * FColliderMesh (effectively) implements the TTriangleMeshAdapter interface, and so can be used 
+ * anywhere a TTriangleMeshAdapter could be used (eg with TMeshQueries functions). 
+ */
+class GEOMETRYCORE_API FColliderMesh
+{
+public:
+	struct FBuildOptions
+	{
+		FBuildOptions()
+		{
+			bBuildAABBTree = true;
+			bBuildVertexMap = false;
+			bBuildTriangleMap = false;
+		}
+
+		/** If true, AABBTree is automatically built for the mesh on construction */
+		bool bBuildAABBTree;
+		/** If true and SourceWasCompactV() == false, mapping from FColliderMesh to source mesh vertex IDs is constructed */
+		bool bBuildVertexMap;
+		/** If true and SourceWasCompactT() == false, mapping from FColliderMesh to source mesh triangle IDs is constructed */
+		bool bBuildTriangleMap;
+	};
+
+public:
+	FColliderMesh();
+	FColliderMesh(const FDynamicMesh3& SourceMesh, const FBuildOptions& BuildOptions = FBuildOptions());
+
+	void Initialize(const FDynamicMesh3& SourceMesh, const FBuildOptions& BuildOptions = FBuildOptions());
+
+	void Reset(bool bReleaseMemory);
+
+	// mesh API required for TMeshAABBTree3
+	bool IsTriangle(int32 TriangleID) const { return TriangleID >= 0 && TriangleID < Triangles.Num(); }
+	bool IsVertex(int32 VertexID) const { return VertexID >= 0 && VertexID < Vertices.Num(); }
+	int32 MaxTriangleID() const { return Triangles.Num(); }
+	int32 MaxVertexID() const { return Vertices.Num(); }
+	int32 TriangleCount() const { return Triangles.Num(); }
+	int32 VertexCount() const { return Vertices.Num(); }
+	uint64 GetChangeStamp() const { return 0; }
+	FIndex3i GetTriangle(int32 TriangleID) const { return Triangles[TriangleID]; }
+	FVector3d GetVertex(int32 VertexID) const { return Vertices[VertexID]; }
+
+	void GetTriVertices(int32 TriangleID, FVector3d& V0, FVector3d& V1, FVector3d& V2) const
+	{
+		FIndex3i Tri = Triangles[TriangleID];
+		V0 = Vertices[Tri.A];
+		V1 = Vertices[Tri.B];
+		V2 = Vertices[Tri.C];
+	}
+
+	FVector3d GetTriNormal(int32 TriangleID) const
+	{
+		FIndex3i Tri = Triangles[TriangleID];
+		return VectorUtil::Normal(Vertices[Tri.A], Vertices[Tri.B], Vertices[Tri.C]);
+	}
+
+	bool SourceWasCompactV() const { return bSourceWasCompactV; }
+	bool HasVertexIDMap() const { return SourceVertexIDs.Num() > 0 && SourceVertexIDs.Num() == Vertices.Num(); }
+	int32 GetSourceVertexID(int32 VertexID) const;
+	bool SourceWasCompactT() const { return bSourceWasCompactT; }
+	bool HasTriangleIDMap() const { return SourceTriangleIDs.Num() > 0 && SourceTriangleIDs.Num() == Triangles.Num(); }
+	int32 GetSourceTriangleID(int32 TriangleID) const;
+
+	// spatial APIs
+
+	TMeshAABBTree3<FColliderMesh>& GetAABBTree();
+
+protected:
+	// For large meshes it might be better to optionally use TDynamicVector here (but keeping the TArray for
+	// small meshes is still also desirable due to minimum-size). Adds a branch to every access, though.
+	TArray<FVector3d> Vertices;
+	TArray<FIndex3i> Triangles;
+
+	TArray<int32> SourceVertexIDs;
+	bool bSourceWasCompactV;
+	TArray<int32> SourceTriangleIDs;
+	bool bSourceWasCompactT;
+
+	TMeshAABBTree3<FColliderMesh> AABBTree;
+
+};
+
+} // end namespace UE::Geometry
+} // end namespace UE
+
+
+
+
+
diff --git a/Engine/Source/Runtime/GeometryCore/Public/TransformTypes.h b/Engine/Source/Runtime/GeometryCore/Public/TransformTypes.h
--- a/Engine/Source/Runtime/GeometryCore/Public/TransformTypes.h
+++ b/Engine/Source/Runtime/GeometryCore/Public/TransformTypes.h
@@ -233,6 +233,20 @@ public:
 	}
 
 
+	UE::Math::TRay<RealType> TransformRay(const UE::Math::TRay<RealType> &Ray) const
+	{
+		TVector<RealType> Origin = TransformPosition(Ray.Origin);
+		TVector<RealType> Direction = Normalized(TransformVector(Ray.Direction));
+		return UE::Math::TRay<RealType>(Origin, Direction);
+	}
+
+
+	UE::Math::TRay<RealType> InverseTransformRay(const UE::Math::TRay<RealType> &Ray) const
+	{
+		TVector<RealType> InvOrigin = InverseTransformPosition(Ray.Origin);
+		TVector<RealType> InvDirection = Normalized(InverseTransformVector(Ray.Direction));
+		return UE::Math::TRay<RealType>(InvOrigin, InvDirection);
+	}
 
 	/**
 	 * Clamp all scale components to a minimum value. Sign of scale components is preserved.
