GeometryCore
-Fix for namespace copllisions with UE::Math::FRay3d

diff --git a/Engine/Source/Runtime/GeometryCore/Private/Spatial/SparseDynamicOctree3.cpp b/Engine/Source/Runtime/GeometryCore/Private/Spatial/SparseDynamicOctree3.cpp
--- a/Engine/Source/Runtime/GeometryCore/Private/Spatial/SparseDynamicOctree3.cpp
+++ b/Engine/Source/Runtime/GeometryCore/Private/Spatial/SparseDynamicOctree3.cpp
@@ -225,7 +225,7 @@ bool FSparseDynamicOctree3::ReinsertObject(int32 ObjectID, const FAxisAlignedBox
 
 
 
-double FSparseDynamicOctree3::FindNearestRayCellIntersection(const FSparseOctreeCell& Cell, const FRay3d& Ray) const
+double FSparseDynamicOctree3::FindNearestRayCellIntersection(const FSparseOctreeCell& Cell, const UE::Geometry::FRay3d& Ray) const
 {
 	FAxisAlignedBox3d Box = GetCellBox(Cell, MaxExpandFactor);
 	double ray_t = TNumericLimits<double>::Max();
@@ -241,9 +241,9 @@ double FSparseDynamicOctree3::FindNearestRayCellIntersection(const FSparseOctree
 
 
 
-int32 FSparseDynamicOctree3::FindNearestHitObject(const FRay3d& Ray,
+int32 FSparseDynamicOctree3::FindNearestHitObject(const UE::Geometry::FRay3d& Ray,
 	TFunctionRef<FAxisAlignedBox3d(int)> GetObjectBoundsFunc,
-	TFunctionRef<double(int, const FRay3d&)> HitObjectDistFunc,
+	TFunctionRef<double(int, const UE::Geometry::FRay3d&)> HitObjectDistFunc,
 	double MaxDistance) const
 {
 	// this should take advantage of raster!
