Geometry Core:  Oriented bounding box, DiTO algorithm.   Expose option for client code to provide sample direction vectors, and add simple unit test.

diff --git a/Engine/Source/Runtime/GeometryCore/Private/CompGeom/DiTOrientedBox.cpp b/Engine/Source/Runtime/GeometryCore/Private/CompGeom/DiTOrientedBox.cpp
--- a/Engine/Source/Runtime/GeometryCore/Private/CompGeom/DiTOrientedBox.cpp
+++ b/Engine/Source/Runtime/GeometryCore/Private/CompGeom/DiTOrientedBox.cpp
@@ -152,11 +152,9 @@ namespace
 
 	// implementation of the DiTO algorithm
 	template <typename RealType>
-	UE::Geometry::TOrientedBox3<RealType> ComputeDiTOImpl(const UE::Geometry::EDiTO DiTOType, const int32 NumPoints, TFunctionRef<UE::Math::TVector<RealType>(int32)> GetPointFunc)
+	UE::Geometry::TOrientedBox3<RealType> ComputeDiTOImpl(const TArray<UE::Math::TVector<RealType>>& SampleDirections, const int32 NumPoints, TFunctionRef<UE::Math::TVector<RealType>(int32)> GetPointFunc)
 	{
 		typedef UE::Math::TVector<RealType> VectorType;
-
-		const TArray<VectorType>  SampleDirections = DiTODirections<RealType>(DiTOType);
 		
 		const int32 NumSampleDir    = SampleDirections.Num();
 		const int32 NumSamplePoints = 2 * NumSampleDir;
@@ -169,6 +167,12 @@ namespace
 																			 ExpandBoxToContain(AABB, GetPointFunc, NumPoints);
 																			 return AABB;
 																		 }();
+		// if no sample directions were provided, just return the AABB
+		if (NumSampleDir == 0)
+		{
+			return UE::Geometry::TOrientedBox3<RealType>(ReferenceAABBox);
+		}
+
 		struct FIntervalBounds
 		{
 			RealType Upper = -TMathUtilConstants<RealType>::MaxReal;
@@ -399,16 +403,28 @@ namespace
 template <typename RealType>
 UE::Geometry::TOrientedBox3<RealType> UE::Geometry::ComputeOrientedBBox(const UE::Geometry::EDiTO DiTOType, const int32 NumPoints, TFunctionRef<UE::Math::TVector<RealType>(int32)> GetPointFunc)
 {
-	return ComputeDiTOImpl(DiTOType, NumPoints, GetPointFunc);
+	typedef UE::Math::TVector<RealType> VectorType;
+	
+	const TArray<VectorType> SampleDirections = DiTODirections<RealType>(DiTOType);
+	return ComputeDiTOImpl(SampleDirections, NumPoints, GetPointFunc);
 }
 
+template <typename RealType>
+UE::Geometry::TOrientedBox3<RealType> UE::Geometry::ComputeOrientedBBox(const TArray<UE::Math::TVector<RealType>>& SampleDirections, const int32 NumPoints, TFunctionRef<UE::Math::TVector<RealType>(int32)> GetPointFunc)
+{
+	return ComputeDiTOImpl(SampleDirections, NumPoints, GetPointFunc);
+}
 
 // explicit instantiations
 namespace UE
 {
 	namespace Geometry
 	{
-		template TOrientedBox3<float> GEOMETRYCORE_API ComputeOrientedBBox(const UE::Geometry::EDiTO DiTOType, const int32 NumPoints, TFunctionRef<UE::Math::TVector<float>(int32)> GetPointFunc);
-		template TOrientedBox3<double> GEOMETRYCORE_API ComputeOrientedBBox(const UE::Geometry::EDiTO DiTOType, const int32 NumPoints, TFunctionRef<UE::Math::TVector<double>(int32)> GetPointFunc);
+		template TOrientedBox3<float> GEOMETRYCORE_API ComputeOrientedBBox<float>(const UE::Geometry::EDiTO DiTOType, const int32 NumPoints, TFunctionRef<UE::Math::TVector<float>(int32)> GetPointFunc);
+		template TOrientedBox3<double> GEOMETRYCORE_API ComputeOrientedBBox<double>(const UE::Geometry::EDiTO DiTOType, const int32 NumPoints, TFunctionRef<UE::Math::TVector<double>(int32)> GetPointFunc);
+		
+		template TOrientedBox3<float> GEOMETRYCORE_API ComputeOrientedBBox<float>(const TArray<UE::Math::TVector<float>>& SampleDirections, const int32 NumPoints, TFunctionRef<UE::Math::TVector<float>(int32)> GetPointFunc);
+		template TOrientedBox3<double> GEOMETRYCORE_API ComputeOrientedBBox<double>(const TArray<UE::Math::TVector<double>>& SampleDirections, const int32 NumPoints, TFunctionRef<UE::Math::TVector<double>(int32)> GetPointFunc);
+	
 	}
 }
\ No newline at end of file
diff --git a/Engine/Source/Runtime/GeometryCore/Public/CompGeom/DiTOrientedBox.h b/Engine/Source/Runtime/GeometryCore/Public/CompGeom/DiTOrientedBox.h
--- a/Engine/Source/Runtime/GeometryCore/Public/CompGeom/DiTOrientedBox.h
+++ b/Engine/Source/Runtime/GeometryCore/Public/CompGeom/DiTOrientedBox.h
@@ -37,6 +37,24 @@ namespace Geometry
 	template <typename RealType>
 	TOrientedBox3<RealType>  GEOMETRYCORE_API ComputeOrientedBBox(const EDiTO DiTO_K, const int32 NumPoints, TFunctionRef<UE::Math::TVector<RealType>(int32)>  GetPointFunc);
 
+
+	/**
+	* @return OrientedBox3 that contains the points provided by the given GetPointFunc
+	* 
+	* Heuristic-based computation of an object oriented bounding box that utilizes 
+	* a small number of extremal vertices and an internal Di-Tetrahedron to generate the box orientation.
+	* Based on "Fast Computation of Tight-Fitting Oriented Bounding Boxes" by Larson and Kallberg ( cf., Game Engine Gems 2)
+	*
+	* @param SampleDirections       - array of sample directions, ideally the direction vectors should be of the same length but no length checks are done internally. 
+	* @param NumPoints              - number points used in constructing the OrientedBox
+	* @param GetPointFunc           - returns a point for every index i = [0, NumPoints)  
+	* 
+	* The number of extreme vertices used by the method will be twice the number of sample directions. 
+	*
+	*/
+	template <typename RealType>
+	TOrientedBox3<RealType>  GEOMETRYCORE_API ComputeOrientedBBox(const TArray<TVector<RealType>>& SampleDirections, const int32 NumPoints, TFunctionRef<UE::Math::TVector<RealType>(int32)>  GetPointFunc);
+
 } // end namespace Geometry
 }// end namespace UE
 
