Add structured down-sampling options to custom voronoi point sets, to allow less-clumped / sharp-edge-favoring down-sampled point sets from meshes and bones
Add a FPriorityOrderPoints class to help with downsampling, implementing a similar algorithm to what is used by physics to help choose collision sample points as a subset of mesh vertices but more generally accessible

diff --git a/Engine/Plugins/Experimental/ChaosEditor/Source/FractureEditor/FractureEditor.Build.cs b/Engine/Plugins/Experimental/ChaosEditor/Source/FractureEditor/FractureEditor.Build.cs
--- a/Engine/Plugins/Experimental/ChaosEditor/Source/FractureEditor/FractureEditor.Build.cs
+++ b/Engine/Plugins/Experimental/ChaosEditor/Source/FractureEditor/FractureEditor.Build.cs
@@ -52,6 +52,7 @@ public class FractureEditor : ModuleRules
                 "GeometryCollectionEngine",
                 "GeometryCollectionEditor",
 				"ModelingComponents",
+				"GeometryCore",
 				"MeshDescription",
 				"StaticMeshDescription",
 				"PlanarCut",
diff --git a/Engine/Plugins/Experimental/ChaosEditor/Source/FractureEditor/Private/FractureToolCustomVoronoi.cpp b/Engine/Plugins/Experimental/ChaosEditor/Source/FractureEditor/Private/FractureToolCustomVoronoi.cpp
--- a/Engine/Plugins/Experimental/ChaosEditor/Source/FractureEditor/Private/FractureToolCustomVoronoi.cpp
+++ b/Engine/Plugins/Experimental/ChaosEditor/Source/FractureEditor/Private/FractureToolCustomVoronoi.cpp
@@ -8,6 +8,7 @@
 #include "Engine/StaticMeshActor.h"
 #include "StaticMeshAttributes.h"
 #include "MeshDescription.h"
+#include "Spatial/PriorityOrderPoints.h"
 
 #define LOCTEXT_NAMESPACE "FractureCustomVoronoi"
 
@@ -127,6 +128,9 @@ void UFractureToolCustomVoronoi::GenerateLivePattern(const TArray<FFractureToolC
 
 	FBox& Bounds = CombinedWorldBounds; // shorter name for convenience
 	TArray<FVector>& Sites = LiveSites;
+	TArray<float> ImportanceWeights;
+	bool bUseOrderedSkip = CustomVoronoiSettings->SkipFraction > 0 
+		&& (CustomVoronoiSettings->SkipMode == EDownsamplingMode::UniformSpacing || CustomVoronoiSettings->SkipMode == EDownsamplingMode::KeepSharp);
 
 	const FVector Extent(Bounds.Max - Bounds.Min);
 	int32 NumSites = CustomVoronoiSettings->SitesToAdd;
@@ -193,7 +197,7 @@ void UFractureToolCustomVoronoi::GenerateLivePattern(const TArray<FFractureToolC
 			TVertexAttributesConstRef<FVector3f> VertexPositions = MeshDescription->GetVertexPositions();
 			TArray<FVector3f> Normals; // Per-vertex normals for offsets
 			float NormalOffset = CustomVoronoiSettings->NormalOffset;
-			if (NormalOffset > 0)
+			if (NormalOffset > 0 || bUseOrderedSkip)
 			{
 				Normals.SetNumZeroed(VertexPositions.GetNumElements());
 
@@ -209,11 +213,21 @@ void UFractureToolCustomVoronoi::GenerateLivePattern(const TArray<FFractureToolC
 				{
 					Normal.Normalize();
 				}
+				if (bUseOrderedSkip)
+				{
+					ImportanceWeights.SetNumZeroed(Normals.Num());
+					for (const FVertexInstanceID InstanceID : MeshDescription->VertexInstances().GetElementIDs())
+					{
+						FVector3f Normal = InstanceNormals.Get(InstanceID);
+						FVertexID VertexID = MeshDescription->GetVertexInstanceVertex(InstanceID);
+						ImportanceWeights[VertexID] = FMath::Max(ImportanceWeights[VertexID], 1 - Normal.Dot(Normals[VertexID]));
+					}
+				}
 			}
 			for (const FVertexID VertexID : MeshDescription->Vertices().GetElementIDs())
 			{
 				FVector Position = VertexPositions.Get(VertexID);
-				if (Normals.Num() > 0)
+				if (NormalOffset > 0)
 				{
 					Position += Normals[VertexID] * NormalOffset;
 				}
@@ -250,7 +264,7 @@ void UFractureToolCustomVoronoi::GenerateLivePattern(const TArray<FFractureToolC
 					Count = MaxUsedIdx - Start + 1;
 				}
 				float NormalOffset = CustomVoronoiSettings->NormalOffset;
-				if (NormalOffset > 0)
+				if (NormalOffset > 0 || bUseOrderedSkip)
 				{
 					TArray<FVector3f> Normals;
 					Normals.SetNumZeroed(Count);
@@ -277,6 +291,24 @@ void UFractureToolCustomVoronoi::GenerateLivePattern(const TArray<FFractureToolC
 						const FVector LocalPos = Collection.Transform[Bone].TransformPosition(OffsetPosition);
 						Sites.Add(Context.GetTransform().TransformPosition(LocalPos));
 					}
+					if (bUseOrderedSkip)
+					{
+						TArray<float> WeightsOnSourceVertices;
+						WeightsOnSourceVertices.SetNumZeroed(Count);
+						for (int32 VIdx = Start; VIdx < Start + Count; VIdx++)
+						{
+							const FVector& Position = Collection.Vertex[VIdx];
+							int32 NIdx = VertexHash[Position] - 1;
+							float& Wt = WeightsOnSourceVertices[NIdx];
+							Wt = FMath::Max(Wt, 1 - Normals[NIdx].Dot(Collection.Normal[VIdx]));
+						}
+						// Now re-order the weights to be 1:1 with the Sites (i.e., order by VertexHash iterator)
+						ImportanceWeights.Reserve(ImportanceWeights.Num() + VertexHash.Num());
+						for (TPair<FVector, int32>& Point : VertexHash)
+						{
+							ImportanceWeights.Add(WeightsOnSourceVertices[Point.Value - 1]);
+						}
+					}
 				}
 				else
 				{
@@ -322,10 +354,44 @@ void UFractureToolCustomVoronoi::GenerateLivePattern(const TArray<FFractureToolC
 
 	// randomly remove points based on the skip fraction
 	int32 TargetNumSites = Sites.Num() - int32(Sites.Num() * CustomVoronoiSettings->SkipFraction);
-	while (TargetNumSites < Sites.Num())
+	if (CustomVoronoiSettings->SkipFraction > 0)
 	{
-		int32 ToRemoveIdx = RandStream.RandHelper(Sites.Num());
-		Sites.RemoveAtSwap(ToRemoveIdx, 1, false);
+		if (bUseOrderedSkip)
+		{
+			UE::Geometry::FPriorityOrderPoints Ordering;
+			if (CustomVoronoiSettings->SkipMode == EDownsamplingMode::UniformSpacing)
+			{
+				Ordering.ComputeUniformSpaced(Sites, ImportanceWeights, TargetNumSites);
+			}
+			else // CustomVoronoiSettings->SkipMode == EDownsamplingMode::KeepSharpEDownsamplingMode::KeepSharp
+			{
+				Ordering.ComputeDescendingImportance(ImportanceWeights);
+			}
+
+			// empty ordering can happen if we're in a mode w/ no ImportanceWeights but mode is KeepSharp; just leave default order in that case
+			if (!Ordering.Order.IsEmpty())
+			{
+				TArray<FVector3d> OrderedSites;
+				OrderedSites.SetNumUninitialized(TargetNumSites);
+				for (int32 Idx = 0; Idx < TargetNumSites; Idx++)
+				{
+					OrderedSites[Idx] = Sites[Ordering.Order[Idx]];
+				}
+				Sites = MoveTemp(OrderedSites);
+			}
+			else
+			{
+				Sites.SetNum(TargetNumSites);
+			}
+		}
+		else
+		{
+			while (TargetNumSites < Sites.Num())
+			{
+				int32 ToRemoveIdx = RandStream.RandHelper(Sites.Num());
+				Sites.RemoveAtSwap(ToRemoveIdx, 1, false);
+			}
+		}
 	}
 
 	// Convert newly generated points from world space to local (unscaled) gizmo space
diff --git a/Engine/Plugins/Experimental/ChaosEditor/Source/FractureEditor/Private/FractureToolCustomVoronoi.h b/Engine/Plugins/Experimental/ChaosEditor/Source/FractureEditor/Private/FractureToolCustomVoronoi.h
--- a/Engine/Plugins/Experimental/ChaosEditor/Source/FractureEditor/Private/FractureToolCustomVoronoi.h
+++ b/Engine/Plugins/Experimental/ChaosEditor/Source/FractureEditor/Private/FractureToolCustomVoronoi.h
@@ -24,6 +24,17 @@ enum class EVoronoiPattern
 	SelectedBones
 };
 
+UENUM()
+enum class EDownsamplingMode
+{
+	// Downsample points by randomly removing points, without considering spacing
+	Random,
+	// Downsample points so they're spread evenly across space, favoring points on sharp features
+	UniformSpacing,
+	// Downsample points so points on sharp features are the last to be removed
+	KeepSharp
+};
+
 UCLASS(config = EditorPerProjectUserSettings)
 class UFractureCustomVoronoiSettings : public UFractureToolSettings
 {
@@ -67,10 +78,14 @@ public:
 		EditCondition = "VoronoiPattern == EVoronoiPattern::Grid", EditConditionHides))
 	int32 GridZ = 5;
 
-	/** Fraction of points to randomly skip */
+	/** Fraction of points to skip */
 	UPROPERTY(EditAnywhere, Category = LiveVoronoiSites, meta = (ClampMin = "0", ClampMax = "1"))
 	float SkipFraction = 0.0;
 
+	/** Strategy used for skipping points; only used if SkipFraction is greater than 0 */
+	UPROPERTY(EditAnywhere, Category = LiveVoronoiSites)
+	EDownsamplingMode SkipMode;
+
 	//~ TODO: Have reference mesh fall back to the input fracture mesh if not provided here
 	/** Static mesh actor to be used as a reference when adding Voronoi sites */
 	UPROPERTY(EditAnywhere, Category = LiveVoronoiSites, meta = (DisplayName = "Reference Mesh", EditCondition = "VoronoiPattern == EVoronoiPattern::MeshVertices", EditConditionHides))
diff --git a/Engine/Source/Runtime/GeometryCore/Private/Spatial/PriorityOrderPoints.cpp b/Engine/Source/Runtime/GeometryCore/Private/Spatial/PriorityOrderPoints.cpp
new file mode 100644
--- /dev/null
+++ b/Engine/Source/Runtime/GeometryCore/Private/Spatial/PriorityOrderPoints.cpp
@@ -0,0 +1,344 @@
+// Copyright Epic Games, Inc. All Rights Reserved.
+
+
+#include "Spatial/PriorityOrderPoints.h"
+#include "BoxTypes.h"
+
+using namespace UE::Geometry;
+
+
+namespace PriorityOrderPointsLocal
+{
+/**
+ * Binary predicate for sorting indices according to a secondary array of values to sort by.
+ * Puts values into descending order.
+ */
+template<class T>
+class DescendingPredicate
+{
+public:
+	DescendingPredicate(const TArrayView<const T>& CompValues)
+		: CompValues(CompValues)
+	{}
+
+	bool operator()(const int i, const int j) const
+	{
+		// Indices must be in range
+		return CompValues[i] > CompValues[j];
+	}
+
+private:
+	const TArrayView<const T> CompValues;
+};
+
+// Spatial ordering algorithm
+// Note we support two sets of importance weights because it is sometimes helpful to also sort by distance from center
+// However this feature is not yet exposed in the outer interface
+template<typename RealType>
+void OrderPoints(TArray<int32>& PointOrder, TArrayView<const TVector<RealType>> Points, TArrayView<TArrayView<const float>> ImportanceWeights, int32 EarlyStop, int32 SpatialLevels)
+{
+	const int32 NumPoints = Points.Num();
+	PointOrder.SetNumUninitialized(NumPoints);
+	if (NumPoints == 0)
+	{
+		return;
+	}
+
+	TAxisAlignedBox3<RealType> Bounds;
+	for (int32 Idx = 0; Idx < NumPoints; Idx++)
+	{
+		PointOrder[Idx] = Idx;
+		Bounds.Contain(Points[Idx]);
+	}
+
+	TVector<RealType> Center = Bounds.Center();
+	TArray<TVector<RealType>> LocalPoints;
+	LocalPoints.AddUninitialized(NumPoints);
+	for (int i = 0; i < NumPoints; i++)
+	{
+		LocalPoints[i] = Points[i] - Center;
+	}
+	TAxisAlignedBox3<RealType> LocalBBox(-Bounds.Extents(), Bounds.Extents());
+	LocalBBox.Expand(1.0e-3f);
+	RealType MaxBBoxDim = LocalBBox.MaxDim();
+
+	auto ToFlatIdx = [&LocalPoints](int32 PointIdx, const TVector<RealType>& Center, RealType CellSizeIn, int64 Res) -> int64
+	{
+		const TVector<RealType>& Pos = LocalPoints[PointIdx];
+		// grid center co-located at bbox center:
+		int64 X =
+			static_cast<int64>(FMath::Floor((Pos[0] - Center[0]) / CellSizeIn)) + Res / 2;
+		int64 Y =
+			static_cast<int64>(FMath::Floor((Pos[1] - Center[1]) / CellSizeIn)) + Res / 2;
+		int64 Z =
+			static_cast<int64>(FMath::Floor((Pos[2] - Center[2]) / CellSizeIn)) + Res / 2;
+		return ((X * Res + Y) * Res) + Z;
+	};
+
+	constexpr int MaxImportanceWts = 2;
+	struct FBestPtData
+	{
+		int32 Idx[MaxImportanceWts]{ -1, -1 };
+		float Wt[MaxImportanceWts]{ -FMathf::MaxReal, -FMathf::MaxReal };
+	};
+	int32 NumWeights = ImportanceWeights.Num();
+	for (int Idx = 0; Idx < NumWeights; Idx++)
+	{
+		if (ImportanceWeights[Idx].IsEmpty())
+		{
+			NumWeights = Idx;
+			break;
+		}
+	}
+
+	// Update FBestPointData with a new point; return true if it's the first point in the cell, false otherwise
+	auto UpdateBestPtData = [NumWeights, &ImportanceWeights, &PointOrder](FBestPtData& BestData, int32 OrderIdx) -> bool
+	{
+		int32 PtIdx = PointOrder[OrderIdx];
+		if (NumWeights <= 0)
+		{
+			if (BestData.Idx[0] == -1)
+			{
+				BestData.Idx[0] = OrderIdx;
+				return true;
+			}
+			return false;
+		}
+
+		if (BestData.Idx[0] == -1)
+		{
+			BestData.Idx[0] = BestData.Idx[1] = OrderIdx;
+			BestData.Wt[0] = ImportanceWeights[0][PtIdx];
+			if (NumWeights > 1)
+			{
+				BestData.Wt[1] = ImportanceWeights[1][PtIdx];
+			}
+			return true;
+		}
+
+		if (ImportanceWeights[0][PtIdx] > BestData.Wt[0])
+		{
+			BestData.Wt[0] = ImportanceWeights[0][PtIdx];
+			BestData.Idx[0] = OrderIdx;
+		}
+
+		// TODO: if we want to support more than two importance weights, turn this into a loop
+		if (NumWeights > 1 && ImportanceWeights[1][PtIdx] > BestData.Wt[1])
+		{
+			BestData.Wt[1] = ImportanceWeights[1][PtIdx];
+			BestData.Idx[1] = OrderIdx;
+		}
+
+		return false;
+	};
+
+	// points before this offset have already been moved forward, and don't need further consideration
+	int32 MovedPtsOffset = 0;
+	TArray<uint8> Rank;
+	TMap<int64, FBestPtData> BestPoints;
+	TSet<int64> CoveredCells, OffsetCoveredCells;
+	TArray<int32> MoveToFront;
+	// traverse power of 2 cell resolutions (~octree levels), moving forward "best" points from the cells at each level
+	for (int32 Resolution = 2; MovedPtsOffset < NumPoints && Resolution <= 1024; Resolution *= 2)
+	{
+		BestPoints.Reset();
+		CoveredCells.Reset();
+		OffsetCoveredCells.Reset();
+		const RealType CellSize = MaxBBoxDim / static_cast<RealType>(Resolution);
+
+		// make quarter-sized cells w/ a half-cell offset to help detect points that are close, but across cell boundaries
+		int32 OffsetResolution = Resolution * 4;
+		const RealType OffsetCellSize = MaxBBoxDim / static_cast<RealType>(OffsetResolution);
+		OffsetResolution += 2; // allow for offset center
+		TVector<RealType> OffsetCenter = -TVector<RealType>(OffsetCellSize / 2);
+		int32 FoundPointsCount = 0;
+
+		// Use the Rank array to mark which points should move to the front of the remaining list in this iteration
+		Rank.Reset();
+		Rank.AddZeroed(NumPoints);
+
+		auto ToMainFlatIdx = [&ToFlatIdx, &CellSize, &Resolution](int32 PointIdx) -> int64
+		{
+			return ToFlatIdx(PointIdx, TVector<RealType>::ZeroVector, CellSize, Resolution);
+		};
+		auto ToOffsetFlatIdx = [&ToFlatIdx, &OffsetCenter, &OffsetCellSize, &OffsetResolution](int32 PointIdx) -> int64
+		{
+			return ToFlatIdx(PointIdx, OffsetCenter, OffsetCellSize, OffsetResolution);
+		};
+
+		// Mark cells that are already covered by points in the already-considered front of the list
+		for (int32 OrderIdx = 0; OrderIdx < MovedPtsOffset; OrderIdx++)
+		{
+			int32 PointIdx = PointOrder[OrderIdx];
+			const int64 FlatIdx = ToMainFlatIdx(PointIdx);
+			CoveredCells.Add(FlatIdx);
+			OffsetCoveredCells.Add(ToOffsetFlatIdx(PointIdx));
+		}
+
+		// find the best points in cells that aren't already covered
+		bool bAllSeparate = true;
+		for (int32 OrderIdx = MovedPtsOffset; OrderIdx < NumPoints; OrderIdx++)
+		{
+			int32 PointIdx = PointOrder[OrderIdx];
+			const int64 FlatIdx = ToMainFlatIdx(PointIdx);
+			if (CoveredCells.Contains(FlatIdx) || OffsetCoveredCells.Contains(ToOffsetFlatIdx(PointIdx)))
+			{
+				bAllSeparate = false;
+				continue;
+			}
+			FBestPtData& Best = BestPoints.FindOrAdd(FlatIdx);
+			bAllSeparate = UpdateBestPtData(Best, OrderIdx) & bAllSeparate;
+		}
+
+		if (bAllSeparate)
+		{
+			// all the points were in separate cells, no need to continue promoting the 'best' for each cell beyond here
+			break;
+		}
+
+		auto ConsiderMovingPt = [&Rank, MovedPtsOffset, &OffsetCoveredCells, &ToOffsetFlatIdx](int32 PointIdx) -> bool
+		{
+			uint8& R = Rank[PointIdx];
+			if (!R)
+			{
+				// try to avoid clumping by skipping points whose offset cells were already covered
+				int64 OffsetFlatIdx = ToOffsetFlatIdx(PointIdx);
+				bool bAlreadyInSet = false;
+				OffsetCoveredCells.Add(OffsetFlatIdx, &bAlreadyInSet);
+				if (bAlreadyInSet)
+				{
+					return false;
+				}
+
+				R = 1;
+				return true;
+			}
+			else
+			{
+				return false;
+			}
+		};
+
+		bool bShouldStopEarly = false;
+
+		// Decide which points to move into the 'front' section
+		int32 NumToMoveToFront = 0;
+		MoveToFront.Reset();
+		// Add points favored by the first metric
+		for (const TPair<int64, FBestPtData>& BestDataPair : BestPoints)
+		{
+			int PtIdx = PointOrder[BestDataPair.Value.Idx[0]];
+			if (ConsiderMovingPt(PtIdx))
+			{
+				NumToMoveToFront++;
+				MoveToFront.Add(BestDataPair.Value.Idx[0]);
+				bShouldStopEarly |= (EarlyStop > -1 && (MovedPtsOffset + NumToMoveToFront) >= EarlyStop);
+				if (bShouldStopEarly)
+				{
+					break;
+				}
+			}
+		}
+		if (NumWeights > 1 && !bShouldStopEarly) // TODO: make this a loop if we support more than two importance weights
+		{
+			// Add points favored by the curvature metric second
+			//  (so they may be skipped where a first-metric-favored point already covered a region)
+			for (const TPair<int64, FBestPtData>& BestDataPair : BestPoints)
+			{
+				int PtIdx = PointOrder[BestDataPair.Value.Idx[1]];
+				if (ConsiderMovingPt(PtIdx))
+				{
+					NumToMoveToFront++;
+					MoveToFront.Add(BestDataPair.Value.Idx[1]);
+					bShouldStopEarly |= (EarlyStop > -1 && (MovedPtsOffset + NumToMoveToFront) >= EarlyStop);
+					if (bShouldStopEarly)
+					{
+						break;
+					}
+				}
+			}
+		}
+
+		// Do the move-to-front operations by swapping
+		int MoveFrontIdx = 0;
+		for (int MoveBackIdx = 0; MoveBackIdx < NumToMoveToFront; MoveBackIdx++)
+		{
+			if (Rank[PointOrder[MovedPtsOffset + MoveBackIdx]])
+			{
+				// Point was marked for moving to front, and is already in front; skip it
+				continue;
+			}
+			// Swap point with one that was marked for the front section
+			bool bDidSwap = false;
+			for (; MoveFrontIdx < MoveToFront.Num(); MoveFrontIdx++)
+			{
+				if (MoveToFront[MoveFrontIdx] < MovedPtsOffset + NumToMoveToFront)
+				{
+					// Don't swap this one back because it's already in the 'front' zone
+					continue;
+				}
+				int32 PtIdx = PointOrder[MoveToFront[MoveFrontIdx]];
+				checkSlow(Rank[PtIdx]);
+				Swap(PointOrder[MovedPtsOffset + MoveBackIdx], PointOrder[MoveToFront[MoveFrontIdx]]);
+				bDidSwap = true;
+				MoveFrontIdx++;
+				break;
+			}
+			// The above for loop should always find a point to swap back
+			checkSlow(bDidSwap);
+		}
+
+		// Note the above swapping method could have been done via a Sort() using the Rank array, but swapping is more direct / faster
+
+		// Sort the just-added points by first metric (they were in arbitrary order before)
+		if (NumWeights > 0)
+		{
+			DescendingPredicate<float> DescendingCurvaturePred(ImportanceWeights[0]);
+			Sort(&PointOrder[MovedPtsOffset], NumToMoveToFront, DescendingCurvaturePred);
+		}
+
+		MovedPtsOffset += NumToMoveToFront;
+
+		if (bShouldStopEarly)
+		{
+			break;
+		}
+	}
+
+	// Sort remaining non-coincident points by first metric
+	if (NumWeights > 0 && MovedPtsOffset != EarlyStop)
+	{
+		DescendingPredicate<float> DescendingCurvaturePred(ImportanceWeights[0]);
+		Sort(&PointOrder[MovedPtsOffset], NumPoints - MovedPtsOffset, DescendingCurvaturePred);
+	}
+}
+}
+
+
+void FPriorityOrderPoints::ComputeUniformSpaced(TArrayView<const FVector3d> Points, TArrayView<const float> ImportanceWeights, int32 EarlyStop)
+{
+	TArrayView<TArrayView<const float>> ImportanceWeightsWeights(&ImportanceWeights, 1);
+	PriorityOrderPointsLocal::OrderPoints<double>(Order, Points, ImportanceWeightsWeights, EarlyStop, SpatialLevels);
+}
+
+void FPriorityOrderPoints::ComputeUniformSpaced(TArrayView<const FVector3f> Points, TArrayView<const float> ImportanceWeights, int32 EarlyStop)
+{
+	TArrayView<TArrayView<const float>> ImportanceWeightsWeights(&ImportanceWeights, 1);
+	PriorityOrderPointsLocal::OrderPoints<float>(Order, Points, ImportanceWeightsWeights, EarlyStop, SpatialLevels);
+}
+
+void FPriorityOrderPoints::ComputeDescendingImportance(TArrayView<const float> ImportanceWeights)
+{
+	int32 NumPoints = ImportanceWeights.Num();
+	if (NumPoints == 0)
+	{
+		return;
+	}
+	Order.SetNum(NumPoints);
+	for (int32 Idx = 0; Idx < NumPoints; Idx++)
+	{
+		Order[Idx] = Idx;
+	}
+	PriorityOrderPointsLocal::DescendingPredicate<float> DescendingCurvaturePred(ImportanceWeights);
+	Sort(&Order[0], NumPoints, DescendingCurvaturePred);
+}
\ No newline at end of file
diff --git a/Engine/Source/Runtime/GeometryCore/Public/Spatial/PriorityOrderPoints.h b/Engine/Source/Runtime/GeometryCore/Public/Spatial/PriorityOrderPoints.h
new file mode 100644
--- /dev/null
+++ b/Engine/Source/Runtime/GeometryCore/Public/Spatial/PriorityOrderPoints.h
@@ -0,0 +1,62 @@
+// Copyright Epic Games, Inc. All Rights Reserved.
+
+#pragma once
+
+#include "CoreMinimal.h"
+#include "Spatial/PointHashGrid3.h"
+
+namespace UE
+{
+namespace Geometry
+{
+
+using namespace UE::Math;
+
+/**
+ * Generate a "priority ordering" for a point set
+ *
+ * Either order by descending ImportanceWeight, or try to additionally un-clump the points --
+ * such that if we only look at the first N points, they will cover the overall space as well as possible
+ * while still favoring higher ImportanceWeight points in each local region.
+ */
+struct GEOMETRYCORE_API FPriorityOrderPoints
+{
+
+	/// Inputs
+
+	// Octree levels to populate incrementally when spacing out points via ComputeUniformSpaced
+	// We will consider a maximum of (2^SpatialLevels)^3 "cells" at the finest granularity
+	int32 SpatialLevels = 10;
+
+
+	/// Outputs
+
+	TArray<int32> Order;
+
+public:
+
+	FPriorityOrderPoints(int32 SpatialLevels = 10) : SpatialLevels(SpatialLevels)
+	{
+	}
+
+	/**
+	 * Compute an ordering for the input points that attempts to keep points 'well spaced' / un-clumped
+	 * 
+	 * @param Points				Points to re-order
+	 * @param ImportanceWeights		Points with a higher ImportanceWeight in their local region of space will appear earlier in the ordering
+	 * @param EarlyStop				If >= 0, stop re-ordering one we've sorted this many points; beyond this the ordering will be arbitrary
+	 */
+	void ComputeUniformSpaced(TArrayView<const FVector3d> Points, TArrayView<const float> ImportanceWeights, int32 EarlyStop = -1);
+	void ComputeUniformSpaced(TArrayView<const FVector3f> Points, TArrayView<const float> ImportanceWeights, int32 EarlyStop = -1);
+
+	/**
+	 * Compute an ordering that only sorts by descending ImportanceWeights, and does not attempt to spatially distribute the points
+	 * 
+	 * @param ImportanceWeights		Points with a higher ImportanceWeight will appear earlier in the ordering
+	 */
+	void ComputeDescendingImportance(TArrayView<const float> ImportanceWeights);
+
+};
+
+} // end namespace UE::Geometry
+} // end namespace UE
