PathExt tool: enable optional rounded corners. Replaces polygon corners with circular arcs. User can set desired uniform arc radius.

diff --git a/Engine/Plugins/Experimental/MeshModelingToolsetExp/Source/MeshModelingToolsExp/Private/DrawPolyPathTool.cpp b/Engine/Plugins/Experimental/MeshModelingToolsetExp/Source/MeshModelingToolsExp/Private/DrawPolyPathTool.cpp
--- a/Engine/Plugins/Experimental/MeshModelingToolsetExp/Source/MeshModelingToolsExp/Private/DrawPolyPathTool.cpp
+++ b/Engine/Plugins/Experimental/MeshModelingToolsetExp/Source/MeshModelingToolsExp/Private/DrawPolyPathTool.cpp
@@ -42,7 +42,17 @@ namespace DrawPolyPathToolLocals
 	}
 
 
-	void GeneratePathMesh(FDynamicMesh3& Mesh, const TArray<FFrame3d>& InPathPoints, const TArray<double>& InOffsetScaleFactors, double OffsetDistance, bool bPathIsClosed, bool bRampMode, bool bSinglePolyGroup)
+	void GeneratePathMesh(FDynamicMesh3& Mesh, 
+		const TArray<FFrame3d>& InPathPoints, 
+		const TArray<double>& InOffsetScaleFactors, 
+		double OffsetDistance, 
+		bool bPathIsClosed, 
+		bool bRampMode, 
+		bool bSinglePolyGroup,
+		bool bRoundedCorners,
+		double CornerRadiusFraction,
+		bool bLimitCornerRadius,
+		int NumCornerArcPoints)
 	{
 		Mesh.Clear();
 
@@ -62,11 +72,12 @@ namespace DrawPolyPathToolLocals
 
 		TArray<double> ArcLengths;
 		ComputeArcLengths(UsePathPoints, ArcLengths);
-		const double PathLength = ArcLengths.Last();
+		double PathLength = ArcLengths.Last();
 
 		if (bPathIsClosed)
 		{
-			FPolygonEdgeMeshGenerator MeshGen(UsePathPoints, UseOffsetScaleFactors, OffsetDistance, FVector3d::UnitZ());
+			double CornerRadius = 2.0 * CornerRadiusFraction * OffsetDistance;
+			FPolygonEdgeMeshGenerator MeshGen(UsePathPoints, UseOffsetScaleFactors, OffsetDistance, FVector3d::UnitZ(), bRoundedCorners, CornerRadius, bLimitCornerRadius, NumCornerArcPoints);
 			MeshGen.bSinglePolyGroup = bSinglePolyGroup;
 			MeshGen.UVWidth = PathLength;
 			MeshGen.UVHeight = 2 * OffsetDistance;
@@ -74,19 +85,37 @@ namespace DrawPolyPathToolLocals
 			Mesh.Copy(&MeshGen);
 
 			Mesh.EnableVertexUVs(FVector2f::Zero());
-			for (int k = 0; k < NumPoints; ++k)
-			{
-				// Temporarily set vertex UVs to arclengths, for use in interpolating height in ramp mode
-				const float Alpha = (float)ArcLengths[k] / PathLength;
-				Mesh.SetVertexUV(2 * k, FVector2f(Alpha, (float)k));
-				Mesh.SetVertexUV(2 * k + 1, FVector2f(Alpha, (float)k));
-			}
 
 			if (bRampMode)
 			{
-				int NumMeshVertices = 2 * NumPoints;
-				ensure(NumMeshVertices == Mesh.VertexCount());
+				// Temporarily set vertex UVs to arclengths, for use in interpolating height in ramp mode
+
+				if (bRoundedCorners)
+				{
+					// If we added arcs to the corners, recompute arc lengths
+					const int N = Mesh.VertexCount() / 2;
+					ArcLengths.Init(0.0, N);
+					double CurPathLength = 0;
+					for (int k = 1; k < N; ++k)
+					{
+						CurPathLength += Distance(Mesh.GetVertex(2 * k), Mesh.GetVertex(2 * (k - 1)));
+						ArcLengths[k] = CurPathLength;
+					}
+					PathLength = ArcLengths.Last();
+				}
+
+				int NumMeshVertices = Mesh.VertexCount();
 				ensure(NumMeshVertices == Mesh.MaxVertexID());
+				ensure(NumMeshVertices == 2 * ArcLengths.Num());
+
+				for (int k = 0; k < NumMeshVertices/2; ++k)
+				{
+					const float Alpha = (float)ArcLengths[k] / PathLength;
+					Mesh.SetVertexUV(2 * k, FVector2f(Alpha, (float)k));
+					Mesh.SetVertexUV(2 * k + 1, FVector2f(Alpha, (float)k));
+				}
+
+				// Set last vertex positions to match first vertex locations so we can construct the vertical wall
 				Mesh.SetVertex(NumMeshVertices - 2, Mesh.GetVertex(0));
 				Mesh.SetVertex(NumMeshVertices - 1, Mesh.GetVertex(1));
 			}
@@ -305,23 +334,29 @@ void UDrawPolyPathTool::OnClicked(const FInputDeviceRay& ClickPos)
 			if (SurfacePathMechanic->IsDone())
 			{
 				bPathIsClosed = SurfacePathMechanic->LoopWasClosed();
-				GetToolManager()->EmitObjectChange(this, MakeUnique<FDrawPolyPathStateChange>(CurrentCurveTimestamp), LOCTEXT("DrawPolyPathBeginOffset", "Set path width"));
+				GetToolManager()->EmitObjectChange(this, MakeUnique<FDrawPolyPathStateChange>(CurrentCurveTimestamp), LOCTEXT("DrawPolyPathFinishPath", "Finish path"));
 				OnCompleteSurfacePath();
 			}
 			else
 			{
-				GetToolManager()->EmitObjectChange(this, MakeUnique<FDrawPolyPathStateChange>(CurrentCurveTimestamp), LOCTEXT("DrawPolyPathBeginPath", "Begin path"));
+				GetToolManager()->EmitObjectChange(this, MakeUnique<FDrawPolyPathStateChange>(CurrentCurveTimestamp), LOCTEXT("DrawPolyPathAddToPath", "Add point to path"));
 			}
 		}
 	}
-	else if (CurveDistMechanic != nullptr)
+	else if (CurveDistMechanic != nullptr && !bSpecifyingRadius)
+	{
+		GetToolManager()->EmitObjectChange(this, MakeUnique<FDrawPolyPathStateChange>(CurrentCurveTimestamp), LOCTEXT("DrawPolyPathBeginWidth", "Set path width"));
+		OnCompleteWidth();
+	}
+	else if (CurveDistMechanic != nullptr && bSpecifyingRadius)
 	{
-		GetToolManager()->EmitObjectChange(this, MakeUnique<FDrawPolyPathStateChange>(CurrentCurveTimestamp), LOCTEXT("DrawPolyPathBeginHeight", "Set extrude height"));
-		OnCompleteOffsetDistance();
+		GetToolManager()->EmitObjectChange(this, MakeUnique<FDrawPolyPathStateChange>(CurrentCurveTimestamp), LOCTEXT("DrawPolyPathBeginRadius", "Set corner radius"));
+		OnCompleteRadius();
 	}
 	else if (ExtrudeHeightMechanic != nullptr)
 	{
 		CurHeight = TransformProps->ExtrudeHeight;
+		GetToolManager()->EmitObjectChange(this, MakeUnique<FDrawPolyPathStateChange>(CurrentCurveTimestamp), LOCTEXT("DrawPolyPathBeginExtrude", "Set extrude height"));
 		OnCompleteExtrudeHeight();
 	}
 }
@@ -352,16 +387,22 @@ bool UDrawPolyPathTool::OnUpdateHover(const FInputDeviceRay& DevicePos)
 	{
 		CurveDistMechanic->UpdateCurrentDistance(DevicePos.WorldRay);
 
+		double CurveDistance = CurveDistMechanic->CurrentDistance;
+
 		if (TransformProps->bSnapToWorldGrid)
 		{
-			double QuantizedDistance = ToolSceneQueriesUtil::SnapDistanceToWorldGridSize(this, CurveDistMechanic->CurrentDistance);
-			TransformProps->Width = QuantizedDistance * 2.0;
-			CurOffsetDistance = QuantizedDistance;
+			CurveDistance = ToolSceneQueriesUtil::SnapDistanceToWorldGridSize(this, CurveDistance);
+		}
+
+		if (bSpecifyingRadius)
+		{
+			TransformProps->CornerRadius = FMath::Clamp(CurveDistance / CurWidth, 0.0, 2.0);
+			CurRadius = TransformProps->CornerRadius;
 		}
 		else
 		{
-			TransformProps->Width = CurveDistMechanic->CurrentDistance * 2.0;
-			CurOffsetDistance = CurveDistMechanic->CurrentDistance;
+			TransformProps->Width = CurveDistance;
+			CurWidth = 0.5 * CurveDistance;
 		}
 		UpdatePathPreview();
 	}
@@ -505,7 +546,7 @@ void UDrawPolyPathTool::OnCompleteSurfacePath()
 		Line1.Origin += DistOffsetDelta * PlaneNormal.Cross(Line1.Direction);
 		Line2.Origin += DistOffsetDelta * PlaneNormal.Cross(Line2.Direction);
 
-		if (Line1.Direction.Dot(Line2.Direction) > 0.999 )
+		if (FMath::Abs(Line1.Direction.Dot(Line2.Direction)) > 0.999 )
 		{
 			CurPathPoints[j].ConstrainedAlignAxis(0, UE::Geometry::Normalized(Next-Prev), PlaneNormal);
 			OffsetScaleFactors[j] = 1.0;
@@ -530,17 +571,18 @@ void UDrawPolyPathTool::OnCompleteSurfacePath()
 	SurfacePathMechanic = nullptr;
 	if (TransformProps->WidthMode == EDrawPolyPathWidthMode::Fixed)
 	{
-		BeginConstantOffsetDistance();
+		BeginConstantWidth();
 	}
 	else
 	{
-		BeginInteractiveOffsetDistance();
+		BeginInteractiveWidth();
 	}
 }
 
 
-void UDrawPolyPathTool::BeginInteractiveOffsetDistance()
+void UDrawPolyPathTool::BeginInteractiveWidth()
 {
+	bSpecifyingRadius = false;
 	bHasSavedWidth = true;
 	SavedWidth = TransformProps->Width;
 	
@@ -555,17 +597,57 @@ void UDrawPolyPathTool::BeginInteractiveOffsetDistance()
 }
 
 
-void UDrawPolyPathTool::BeginConstantOffsetDistance()
+void UDrawPolyPathTool::BeginConstantWidth()
+{
+	bSpecifyingRadius = false;
+	InitializePreviewMesh();
+	CurWidth = TransformProps->Width * 0.5;
+	UpdatePathPreview();
+	OnCompleteWidth();
+}
+
+void UDrawPolyPathTool::OnCompleteWidth()
+{
+	if (TransformProps->bRoundedCorners)
+	{
+		if (TransformProps->WidthMode == EDrawPolyPathWidthMode::Fixed)
+		{
+			BeginConstantRadius();
+		}
+		else
+		{
+			BeginInteractiveRadius();
+		}
+	}
+	else
+	{
+		BeginConstantRadius();
+	}
+}
+
+
+void UDrawPolyPathTool::BeginConstantRadius()
 {
+	bSpecifyingRadius = true;
 	InitializePreviewMesh();
-	CurOffsetDistance = TransformProps->Width * 0.5;
+	CurRadius = TransformProps->CornerRadius;
 	UpdatePathPreview();
-	OnCompleteOffsetDistance();
+	OnCompleteRadius();
 }
 
+void UDrawPolyPathTool::BeginInteractiveRadius()
+{
+	bSpecifyingRadius = true;
 
+	// begin setting offset distance
+	CurveDistMechanic = NewObject<USpatialCurveDistanceMechanic>(this);
+	CurveDistMechanic->Setup(this);
+	CurveDistMechanic->InitializePolyCurve(CurPolyLine, UE::Geometry::FTransform3d::Identity());
+
+	InitializePreviewMesh();
+}
 
-void UDrawPolyPathTool::OnCompleteOffsetDistance()
+void UDrawPolyPathTool::OnCompleteRadius()
 {
 	CurveDistMechanic = nullptr;
 
@@ -599,7 +681,6 @@ void UDrawPolyPathTool::OnCompleteExtrudeHeight()
 }
 
 
-
 void UDrawPolyPathTool::UpdatePathPreview()
 {
 	check(EditPreview != nullptr);
@@ -617,7 +698,18 @@ void UDrawPolyPathTool::GeneratePathMesh(FDynamicMesh3& Mesh)
 	SecondPolyLoop.Reset();
 
 	const bool bRampMode = (TransformProps->ExtrudeMode == EDrawPolyPathExtrudeMode::RampFixed) || (TransformProps->ExtrudeMode == EDrawPolyPathExtrudeMode::RampInteractive);
-	DrawPolyPathToolLocals::GeneratePathMesh(Mesh, CurPathPoints, OffsetScaleFactors, CurOffsetDistance, bPathIsClosed, bRampMode, TransformProps->bSinglePolyGroup);
+	constexpr bool bLimitCornerRadius = true;
+	DrawPolyPathToolLocals::GeneratePathMesh(Mesh, 
+		CurPathPoints, 
+		OffsetScaleFactors, 
+		CurWidth, 
+		bPathIsClosed, 
+		bRampMode, 
+		TransformProps->bSinglePolyGroup, 
+		TransformProps->bRoundedCorners, 
+		CurRadius,
+		bLimitCornerRadius,
+		TransformProps->RadialSlices);
 
 	FMeshNormals::QuickRecomputeOverlayNormals(Mesh);
 
@@ -824,15 +916,42 @@ void UDrawPolyPathTool::UndoCurrentOperation()
 	}
 	else if (CurveDistMechanic != nullptr)
 	{
-		CurveDistMechanic = nullptr;
-		ClearPreview();
-		InitializeNewSurfacePath();
-		SurfacePathMechanic->HitPath = CurPathPoints;
+		if (!bSpecifyingRadius)
+		{
+			CurveDistMechanic = nullptr;
+			ClearPreview();
+			InitializeNewSurfacePath();
+			SurfacePathMechanic->HitPath = CurPathPoints;
+		}
+		else
+		{
+			BeginInteractiveWidth();
+		}
 	}
 	else if (ExtrudeHeightMechanic != nullptr)
 	{
 		ExtrudeHeightMechanic = nullptr;
-		BeginInteractiveOffsetDistance();
+
+		if (TransformProps->bRoundedCorners && TransformProps->WidthMode == EDrawPolyPathWidthMode::Interactive)
+		{
+			BeginInteractiveRadius();
+		}
+		else
+		{
+			if (TransformProps->WidthMode == EDrawPolyPathWidthMode::Fixed)
+			{
+				// Pop all the way back to initial path creation
+				CurveDistMechanic = nullptr;
+				ClearPreview();
+				InitializeNewSurfacePath();
+				SurfacePathMechanic->HitPath = CurPathPoints;
+			}
+			else
+			{
+				BeginInteractiveWidth();
+			}
+		}
+	
 	}
 }
 
diff --git a/Engine/Plugins/Experimental/MeshModelingToolsetExp/Source/MeshModelingToolsExp/Public/DrawPolyPathTool.h b/Engine/Plugins/Experimental/MeshModelingToolsetExp/Source/MeshModelingToolsExp/Public/DrawPolyPathTool.h
--- a/Engine/Plugins/Experimental/MeshModelingToolsetExp/Source/MeshModelingToolsExp/Public/DrawPolyPathTool.h
+++ b/Engine/Plugins/Experimental/MeshModelingToolsetExp/Source/MeshModelingToolsExp/Public/DrawPolyPathTool.h
@@ -75,6 +75,19 @@ public:
 	UPROPERTY(EditAnywhere, Category = Path, meta = (EditCondition = "WidthMode == EDrawPolyPathWidthMode::Fixed", UIMin = "0.0001", UIMax = "1000", ClampMin = "0", ClampMax = "999999"))
 	float Width = 10.0f;
 
+	/** Use arc segments instead of straight lines in corners */
+	UPROPERTY(EditAnywhere, Category = Path)
+	bool bRoundedCorners = false;
+
+	/** Radius of the corner arcs, as a fraction of path width. This is only available if Rounded Corners is enabled */
+	UPROPERTY(EditAnywhere, Category = Path, meta = (EditCondition = "WidthMode == EDrawPolyPathWidthMode::Fixed && bRoundedCorners", UIMin = "0.01", UIMax = "2.0", ClampMin = "0.01", ClampMax = "999999"))
+	float CornerRadius = 0.01f;
+
+	/** Number of radial subdivisions for rounded corners */
+	UPROPERTY(EditAnywhere, NonTransactional, Category = Path, meta = (UIMin = "3", UIMax = "100", ClampMin = "3", ClampMax = "10000",
+		EditCondition = "bRoundedCorners"))
+	int RadialSlices = 16;
+
 	/** If true, all quads on the path will belong to the same polygon. If false, each quad gets its own polygon. */
 	UPROPERTY(EditAnywhere, Category = Path)
 	bool bSinglePolyGroup = false;
@@ -199,7 +212,8 @@ protected:
 	TArray<double> OffsetScaleFactors;
 	TArray<FVector3d> CurPolyLine;
 	double CurPathLength;
-	double CurOffsetDistance;
+	double CurWidth;
+	double CurRadius;
 	double CurHeight;
 	bool bHasSavedWidth = false;
 	float SavedWidth;
@@ -220,19 +234,28 @@ protected:
 	UPROPERTY()
 	TObjectPtr<UCollectSurfacePathMechanic> SurfacePathMechanic = nullptr;
 
+	bool bSpecifyingRadius = false;
+
 	void InitializeNewSurfacePath();
 	void UpdateSurfacePathPlane();
 	void OnCompleteSurfacePath();
-	void OnCompleteOffsetDistance();
-	void OnCompleteExtrudeHeight();
 
-	void BeginInteractiveOffsetDistance();
-	void BeginConstantOffsetDistance();
-	void UpdatePathPreview();
+	void BeginInteractiveWidth();
+	void BeginConstantWidth();
+	void OnCompleteWidth();
+
+	void BeginInteractiveRadius();
+	void BeginConstantRadius();
+	void OnCompleteRadius();
+
 	void BeginInteractiveExtrudeHeight();
 	void UpdateExtrudePreview();
+	void OnCompleteExtrudeHeight();
+
+	void UpdatePathPreview();
 	void InitializePreviewMesh();
 	void ClearPreview();
+
 	void GeneratePathMesh(UE::Geometry::FDynamicMesh3& Mesh);
 	void GenerateExtrudeMesh(UE::Geometry::FDynamicMesh3& PathMesh);
 	void EmitNewObject();
diff --git a/Engine/Source/Runtime/GeometryCore/Private/Generators/PolygonEdgeMeshGenerator.cpp b/Engine/Source/Runtime/GeometryCore/Private/Generators/PolygonEdgeMeshGenerator.cpp
--- a/Engine/Source/Runtime/GeometryCore/Private/Generators/PolygonEdgeMeshGenerator.cpp
+++ b/Engine/Source/Runtime/GeometryCore/Private/Generators/PolygonEdgeMeshGenerator.cpp
@@ -4,14 +4,278 @@
 
 using namespace UE::Geometry;
 
+namespace PolygonEdgeMeshGeneratorLocal
+{
+	// Given two non-collinear rays and a circle radius, there is a circle that both rays touch (tangentially) at the same distance along each ray. Return this ray distance.
+	double TwoRayCircleTangentTest(const FVector2d& RayOrigin, const FVector2d& RayDirection0, const FVector2d& RayDirection1, double CircleRadius)
+	{
+		const double InteriorRayAngle = AngleR(RayDirection0, RayDirection1);
+		const double Denom = FMath::Tan(0.5 * InteriorRayAngle);
+		if (!ensure(!FMath::IsNearlyEqual(Denom, 0.0)))		// Should check for ray collinearity before this function
+		{
+			return 0.0;
+		}
+		return CircleRadius / Denom;
+	}
+
+	// Given a circle radius, two points on the circle, and a tangent to the first point, find the circle center.
+	FVector2d FindCircleCenter(double CircleRadius, const FVector2d& PointOnCircle0, const FVector2d& PointOnCircle1, const FVector2d& Tangent0)
+	{
+		const FVector2d ZeroToMid = 0.5 * (PointOnCircle1 - PointOnCircle0);
+		const double LengthToMid = ZeroToMid.Length();
+		const double MidToCenterDistance = FMath::Sqrt(CircleRadius * CircleRadius - LengthToMid * LengthToMid);
+		const double Scalar = MidToCenterDistance / LengthToMid;
+
+		// Two possible centers given two points and a radius
+		const FVector2d C0 = PointOnCircle0 + ZeroToMid + Scalar * PerpCW(ZeroToMid);
+		const FVector2d C1 = PointOnCircle0 + ZeroToMid - Scalar * PerpCW(ZeroToMid);
+
+		// Find the center that produces a diameter cord perpendicular to the input segment
+		const double Dot0 = FMath::Abs((C0 - PointOnCircle0).Dot(Tangent0));
+		const double Dot1 = FMath::Abs((C1 - PointOnCircle0).Dot(Tangent0));
+		if (Dot0 < Dot1)
+		{
+			return C0;
+		}
+		else
+		{
+			return C1;
+		}
+	}
+
+
+	// Given two points on a circle, discretize the shortest arc between them using the given number of points
+	void GenerateArc(const FVector2d& PointOnCircle0, 
+		const FVector2d& PointOnCircle1, 
+		const FVector2d& CircleCenter, 
+		double CircleRadius, 
+		int NumArcPoints, 
+		TArray<FVector2d>& OutArc)
+	{
+		double Param0 = 2.0 * FMath::Atan2(PointOnCircle0[1] - CircleCenter[1], PointOnCircle0[0] - CircleCenter[0] + CircleRadius);
+		double Param1 = 2.0 * FMath::Atan2(PointOnCircle1[1] - CircleCenter[1], PointOnCircle1[0] - CircleCenter[0] + CircleRadius);
+
+		// There are two arcs connecting the two parameters -- choose the shorter one. Param1 - Param0 can be negative.
+		if (FMath::Abs(Param1 - Param0) > PI)
+		{
+			if (Param0 < Param1)
+			{
+				Param0 += TWO_PI;
+			}
+			else
+			{
+				Param1 += TWO_PI;
+			}
+		}
+
+		const double DeltaParam = (Param1 - Param0) / (double)(NumArcPoints - 1);
+
+		for (int I = 0; I < NumArcPoints; ++I)
+		{
+			const double Param = Param0 + I * DeltaParam;
+			const FVector2d ArcPoint = CircleCenter + CircleRadius * FVector2d{ FMath::Cos(Param), FMath::Sin(Param) };
+			OutArc.Add(ArcPoint);
+		}
+	}
+
+	// Given a corner P0-P1-P2, generate an arc with given radius (discretized into NumArcPoints-1 segments)
+	void RoundCorner(const FVector2d& P0, const FVector2d& P1, const FVector2d& P2, double CircleRadius, int NumArcPoints, TArray<FVector2d>& OutCorner)
+	{
+		if (!ensure(NumArcPoints > 1))
+		{
+			OutCorner.Add(P1);
+			return;
+		}
+
+		const FVector2d& RayOrigin = P1;
+		const FVector2d RayDirection0 = Normalized(P0 - RayOrigin);
+		const FVector2d RayDirection1 = Normalized(P2 - RayOrigin);
+
+		if (FMath::IsNearlyEqual(FMath::Abs(RayDirection0.Dot(RayDirection1)), 1.0))
+		{
+			OutCorner.Add(P1);
+			return;
+		}
+
+		const double RayIntersectLength = TwoRayCircleTangentTest(RayOrigin, RayDirection0, RayDirection1, CircleRadius);
+
+		const FVector2d PointOnCircle0 = RayOrigin + RayIntersectLength * RayDirection0;
+		const FVector2d PointOnCircle1 = RayOrigin + RayIntersectLength * RayDirection1;
+		const FVector2d Tangent0 = Normalized(P0 - PointOnCircle0);
+		const FVector2d CircleCenter = FindCircleCenter(CircleRadius, PointOnCircle0, PointOnCircle1, Tangent0);
+
+		GenerateArc(PointOnCircle0, PointOnCircle1, CircleCenter, CircleRadius, NumArcPoints, OutCorner);
+	}
+
+	FVector2d ProjectPathVertex(const FVector3d& InVertex, const FFrame3d& PolygonFrame)
+	{
+		return PolygonFrame.ToPlaneUV(InVertex);
+	}
+
+	FVector3d UnprojectPathVertex(const FVector2d& InVertex, const FFrame3d& PolygonFrame)
+	{
+		return PolygonFrame.FromPlaneUV(InVertex);
+	}
+
+
+	// For each corner in the path, compute the maximum radius of an arc that could be used to round out the corner. 
+	// Arc radius is limited such that each "corner arc" would only extend to the midway point of each incident path edge.
+	void ComputeMaxCornerRadii(const TArray<FVector3d>& Path, const FFrame3d& PolygonFrame, TArray<double>& MaxCircleRadii)
+	{
+		int N = Path.Num();
+		MaxCircleRadii.Init(BIG_NUMBER, N);
+
+		for (int CornerID = 0; CornerID < N; ++CornerID)
+		{
+			const int MinusOne = (CornerID + N - 1) % N;
+			const int PlusOne = (CornerID + 1) % N;
+			const FVector2d M1 = ProjectPathVertex(Path[MinusOne], PolygonFrame);
+			const FVector2d C = ProjectPathVertex(Path[CornerID], PolygonFrame);
+			const FVector2d P1 = ProjectPathVertex(Path[PlusOne], PolygonFrame);
+
+			const double CornerAngle = AngleR(Normalized(M1 - C), Normalized(P1 - C));
+
+			// Don't let an arc go more than half the edge length (conservative limit but easy to compute)
+			const double L1 = Distance(M1, C);
+			const double RMax1 = 0.5 * L1 * FMath::Tan(0.5 * CornerAngle);
+			const double L2 = Distance(C, P1);
+			const double RMax2 = 0.5 * L2 * FMath::Tan(0.5 * CornerAngle);
+
+			MaxCircleRadii[CornerID] = FMath::Min(RMax1, RMax2);
+		}
+	}
+
+	// Given a polygonal path, replace each corner with a discretized arc
+	void CurvePath(const TArray<FVector3d>& InPath, 
+		TArray<bool> InteriorAngleFlag, 
+		double Radius, 
+		double Width, 
+		bool bLimitCornerRadius,
+		int NumArcVertices,
+		TArray<double> MaxCornerRadii,
+		TArray<double> OtherSideMaxCornerRadii,
+		const FFrame3d& PolygonFrame,
+		TArray<FVector3d>& OutPath)
+	{
+		const int N = InPath.Num();
+		if (N < 3)
+		{
+			return;
+		}
+
+		for (int CornerID = 0; CornerID < N; ++CornerID)
+		{
+			const int Prev = (CornerID + N - 1) % N;
+			const int Next = (CornerID + 1) % N;
+
+			const FVector2d PrevVertex = ProjectPathVertex(InPath[Prev], PolygonFrame);
+			const FVector2d CurrVertex = ProjectPathVertex(InPath[CornerID], PolygonFrame);
+			const FVector2d NextVertex = ProjectPathVertex(InPath[Next], PolygonFrame);
+
+			TArray<FVector2d> ArcVertices;
+
+			double EffectiveRadius = Radius;
+			if (bLimitCornerRadius)
+			{
+				// The input Radius is taken as the desired radius for exterior corners. So the desired interior radius is the input
+				// radius minus width.
+				// Each corner has a max interior and max exterior radius. We want to respect these limits but still have:
+				//		interior radius = extrerior radius - width
+
+				if (InteriorAngleFlag[CornerID])
+				{
+					double ClampedExteriorRadius = FMath::Min(Radius, OtherSideMaxCornerRadii[CornerID]);
+					double InteriorMax = MaxCornerRadii[CornerID];
+					EffectiveRadius = FMath::Clamp(ClampedExteriorRadius - Width, 0.0, InteriorMax );
+				}
+				else
+				{
+					double InteriorMax = OtherSideMaxCornerRadii[CornerID];
+					double ExteriorMax = MaxCornerRadii[CornerID];
+					EffectiveRadius = FMath::Min(Radius, FMath::Min(ExteriorMax, InteriorMax + Width));
+				}
+			}
+			else
+			{
+				if (InteriorAngleFlag[CornerID])
+				{
+					EffectiveRadius -= Width;
+				}
+			}
+
+			if (EffectiveRadius > KINDA_SMALL_NUMBER)
+			{
+				RoundCorner(PrevVertex, CurrVertex, NextVertex, EffectiveRadius, NumArcVertices, ArcVertices);
+			}
+			else if (OtherSideMaxCornerRadii[CornerID] <= KINDA_SMALL_NUMBER && MaxCornerRadii[CornerID] <= KINDA_SMALL_NUMBER)
+			{
+				// Both interior and exterior radii are close to zero -- don't insert any arc here at all
+				ArcVertices.Add(CurrVertex);
+			}
+			else
+			{
+				// TODO: Make a proper triangle fan, not a bunch of vertices on top of each other
+				for (int I = 0; I < NumArcVertices; ++I)
+				{
+					ArcVertices.Add(CurrVertex);
+				}
+			}
+
+			for (const FVector2d& V : ArcVertices)
+			{
+				OutPath.Add(UnprojectPathVertex(V, PolygonFrame));
+			}
+		}
+	}
+
+
+	double ArcLength(const TArray<FVector3d>& Path, double MinSegmentArcLength = 0.1 )
+	{
+		double Length = 0;
+		for (int Index = 1; Index < Path.Num(); ++Index)
+		{
+			Length += FMath::Max(Distance(Path[Index], Path[Index - 1]), MinSegmentArcLength);
+		}
+
+		// loop
+		Length += FMath::Max(Distance(Path.Last(), Path[0]), MinSegmentArcLength);
+
+		return Length;
+	}
+
+	double PartialArcLength(const TArray<FVector3d>& Path, int MaxIndex, double MinSegmentArcLength = 0.1)
+	{
+		if (MaxIndex == Path.Num())
+		{
+			return ArcLength(Path);
+		}
+
+		double Length = 0.0;
+		for (int Index = 1; Index <= MaxIndex; ++Index)
+		{
+			Length += FMath::Max(Distance(Path[Index], Path[Index - 1]), MinSegmentArcLength);
+		}
+		return Length;
+	}
+}
+
+
 FPolygonEdgeMeshGenerator::FPolygonEdgeMeshGenerator(const TArray<FFrame3d>& InPolygon,
 	const TArray<double>& InOffsetScaleFactors,
 	double InWidth,
-	FVector3d InNormal) :
+	FVector3d InNormal,
+	bool bInRoundedCorners,
+	double InCornerRadius,
+	bool bInLimitCornerRadius,
+	int InNumArcVertices) :
 	Polygon(InPolygon),
 	OffsetScaleFactors(InOffsetScaleFactors),
 	Width(InWidth),
-	Normal(InNormal)
+	Normal(InNormal),
+	bRoundedCorners(bInRoundedCorners),
+	CornerRadius(InCornerRadius),
+	bLimitCornerRadius(bInLimitCornerRadius),
+	NumArcVertices(InNumArcVertices)
 {
 	check(Polygon.Num() == OffsetScaleFactors.Num());
 }
@@ -21,48 +285,109 @@ FPolygonEdgeMeshGenerator::FPolygonEdgeMeshGenerator(const TArray<FFrame3d>& InP
 // TODO: Enable more subdivisions along the width and length dimensions if requested
 FMeshShapeGenerator& FPolygonEdgeMeshGenerator::Generate()
 {
-	const int NumInputVertices = Polygon.Num();
+	using namespace PolygonEdgeMeshGeneratorLocal;
+
+	const int32 NumInputVertices = Polygon.Num();
 	check(NumInputVertices >= 3);
 	if (NumInputVertices < 3)
 	{
 		return *this;
 	}
-	const int NumVertices = 2 * NumInputVertices;
-	const int NumUVs = 2 * NumInputVertices + 2;
-	const int NumTriangles = 2 * NumInputVertices;
-	SetBufferSizes(NumVertices, NumTriangles, NumUVs, NumVertices);
+
+	TArray<FVector3d> CenterPath;
+	TArray<FVector3d> LeftSidePath;
+	TArray<FVector3d> RightSidePath;
 
 	// Trace the input path, placing vertices on either side of each input vertex 
 	const FVector3d LeftVertex{ 0, -Width, 0 };
 	const FVector3d RightVertex{ 0, Width, 0 };
-	for (int CurrentInputVertex = 0; CurrentInputVertex < NumInputVertices; ++CurrentInputVertex)
+	for (int32 CurrentInputVertex = 0; CurrentInputVertex < NumInputVertices; ++CurrentInputVertex)
 	{
 		const FFrame3d& CurrentFrame = Polygon[CurrentInputVertex];
-		const int NewVertexAIndex = 2 * CurrentInputVertex;
-		const int NewVertexBIndex = NewVertexAIndex + 1;
-		Vertices[NewVertexAIndex] = CurrentFrame.FromFramePoint(OffsetScaleFactors[CurrentInputVertex] * LeftVertex);
-		Vertices[NewVertexBIndex] = CurrentFrame.FromFramePoint(OffsetScaleFactors[CurrentInputVertex] * RightVertex);
+		const int32 NewVertexAIndex = 2 * CurrentInputVertex;
+		const int32 NewVertexBIndex = NewVertexAIndex + 1;
+		const FVector3d NewVertexA = CurrentFrame.FromFramePoint(OffsetScaleFactors[CurrentInputVertex] * LeftVertex);
+		const FVector3d NewVertexB = CurrentFrame.FromFramePoint(OffsetScaleFactors[CurrentInputVertex] * RightVertex);
+
+		CenterPath.Add(CurrentFrame.Origin);
+		LeftSidePath.Add(NewVertexA);
+		RightSidePath.Add(NewVertexB);
+	}
+
+	TArray<FVector3d> NewLeftSidePath;
+	TArray<FVector3d> NewRightSidePath;
+	
+	const FFrame3d PolygonFrame = Polygon[0];
+
+	if (bRoundedCorners)
+	{
+		// Compute the maximum allowable radius for each corner, on each path side
+		TArray<double> LeftMaxCornerRadii, RightMaxCornerRadii;
+		ComputeMaxCornerRadii(LeftSidePath, PolygonFrame, LeftMaxCornerRadii);
+		ComputeMaxCornerRadii(RightSidePath, PolygonFrame, RightMaxCornerRadii);
+
+		// For each point on the input path, determine whether the corner turns "left" or "right". This is used to define
+		// the interior and exterior radii at each path corner.
+		// TODO: Support a "round outwards" mode where we replace the turning flags with an "interior to the polygon" flag
+		TArray<bool> TurnsLeft;
+		TArray<bool> TurnsRight;
+		for (int32 I = 0; I < CenterPath.Num(); ++I)
+		{
+			const int32 Prev = I == 0 ? CenterPath.Num() - 1 : I - 1;
+			const int32 Next = (I + 1) % CenterPath.Num();
+			const FVector2d A = ProjectPathVertex(CenterPath[Prev], PolygonFrame);
+			const FVector2d B = ProjectPathVertex(CenterPath[I], PolygonFrame);
+			const FVector2d C = ProjectPathVertex(CenterPath[Next], PolygonFrame);
+			TurnsLeft.Add(Orient(A, B, C) < 0);
+			TurnsRight.Add(!TurnsLeft.Last());
+		}
+
+		CurvePath(LeftSidePath,  TurnsLeft,  CornerRadius, 2.0 * Width, bLimitCornerRadius, NumArcVertices, LeftMaxCornerRadii, RightMaxCornerRadii, PolygonFrame, NewLeftSidePath);
+		CurvePath(RightSidePath, TurnsRight, CornerRadius, 2.0 * Width, bLimitCornerRadius, NumArcVertices, RightMaxCornerRadii, LeftMaxCornerRadii, PolygonFrame, NewRightSidePath);
+	}
+	else
+	{
+		NewLeftSidePath = LeftSidePath;
+		NewRightSidePath = RightSidePath;
+	}
+
+	check(NewLeftSidePath.Num() == NewRightSidePath.Num());
+
+
+	//
+	// Set vertices, triangles, polygroups in the output mesh
+	//
+
+	const int32 FinalNumVertices = 2 * NewLeftSidePath.Num();
+	const int32 FinalNumUVs = 2 * NewLeftSidePath.Num() + 2;		// Extra pair of UVs for the final loop polygon
+	const int32 FinalNumTriangles = 2 * NewLeftSidePath.Num();
+	SetBufferSizes(FinalNumVertices, FinalNumTriangles, FinalNumUVs, FinalNumVertices);
+
+	Vertices.Empty();
+	for (int32 CurrentVertex = 0; CurrentVertex < NewLeftSidePath.Num(); ++CurrentVertex)
+	{
+		Vertices.Add(NewLeftSidePath[CurrentVertex]);
+		Vertices.Add(NewRightSidePath[CurrentVertex]);
 	}
 
-	// Triangulate the vertices we just placed
 	int PolyIndex = 0;
-	for (int CurrentVertex = 0; CurrentVertex < NumInputVertices; ++CurrentVertex)
+	for (int32 CurrentVertex = 0; CurrentVertex < NewLeftSidePath.Num(); ++CurrentVertex)
 	{
-		const int NewVertexA = 2 * CurrentVertex;
-		const int NewVertexB = NewVertexA + 1;
-		const int NewVertexC = (NewVertexA + 2) % NumVertices;
-		const int NewVertexD = (NewVertexA + 3) % NumVertices;
+		const int32 NewVertexA = 2 * CurrentVertex;
+		const int32 NewVertexB = NewVertexA + 1;
+		const int32 NewVertexC = (NewVertexA + 2) % FinalNumVertices;
+		const int32 NewVertexD = (NewVertexA + 3) % FinalNumVertices;
 
-		FIndex3i NewTriA{ NewVertexA, NewVertexB, NewVertexC };
-		const int NewTriAIndex = 2 * CurrentVertex;
+		const FIndex3i NewTriA{ NewVertexA, NewVertexB, NewVertexC };
+		const int32 NewTriAIndex = 2 * CurrentVertex;
 
 		SetTriangle(NewTriAIndex, NewTriA);
 		SetTriangleUVs(NewTriAIndex, NewTriA);
 		SetTriangleNormals(NewTriAIndex, NewTriA);
 		SetTrianglePolygon(NewTriAIndex, PolyIndex);
 
-		const int NewTriBIndex = NewTriAIndex + 1;
-		FIndex3i NewTriB{ NewVertexC, NewVertexB, NewVertexD };
+		const int32 NewTriBIndex = NewTriAIndex + 1;
+		const FIndex3i NewTriB{ NewVertexC, NewVertexB, NewVertexD };
 
 		SetTriangle(NewTriBIndex, NewTriB);
 		SetTriangleUVs(NewTriBIndex, NewTriB);
@@ -75,12 +400,16 @@ FMeshShapeGenerator& FPolygonEdgeMeshGenerator::Generate()
 		}
 	}
 
-	for (int NewVertexIndex = 0; NewVertexIndex < NumVertices; ++NewVertexIndex)
+	for (int32 NewVertexIndex = 0; NewVertexIndex < FinalNumVertices; ++NewVertexIndex)
 	{
 		Normals[NewVertexIndex] = FVector3f(Normal);
 		NormalParentVertex[NewVertexIndex] = NewVertexIndex;
 	}
 
+	//
+	// Generate UVs 
+	//
+	 
 	// Create a UV strip for the path
 	const float UVLeft = 0.0f, UVBottom = 0.0f;
 	float UVRight = 1.0f, UVTop = 1.0f;
@@ -100,27 +429,34 @@ FMeshShapeGenerator& FPolygonEdgeMeshGenerator::Generate()
 	FVector2f UV11 = FVector2f(UVRight, UVTop);
 	FVector2f UV10 = FVector2f(UVLeft, UVTop);
 
-	for (int CurrentInputVertex = 0; CurrentInputVertex < NumInputVertices; ++CurrentInputVertex)
+	const int32 FinalNumVerticesOneSide = NewLeftSidePath.Num();
+
+	double TotalArcLength = ArcLength(NewLeftSidePath);
+	const double MinSegmentArcLength = 0.1 * TotalArcLength / (double)FinalNumVerticesOneSide;
+	TotalArcLength = ArcLength(NewLeftSidePath, MinSegmentArcLength);
+
+	for (int32 NewVertexIndex = 0; NewVertexIndex < FinalNumVerticesOneSide; ++NewVertexIndex)
 	{
-		const int NewVertexAIndex = 2 * CurrentInputVertex;
-		const int NewVertexBIndex = NewVertexAIndex + 1;
+		const int32 NewVertexAIndex = 2 * NewVertexIndex;
+		const int32 NewVertexBIndex = NewVertexAIndex + 1;
+
+		float UParam = PartialArcLength(NewLeftSidePath, NewVertexIndex, MinSegmentArcLength) / TotalArcLength;
 
-		float UParam = (float)CurrentInputVertex / (float)NumInputVertices;	// TODO: Parametrize based on arc length
 		UVs[NewVertexAIndex] = BilinearInterp(UV00, UV01, UV11, UV10, UParam, 0.0f);
 		UVs[NewVertexBIndex] = BilinearInterp(UV00, UV01, UV11, UV10, UParam, 1.0f);
 		UVParentVertex[NewVertexAIndex] = NewVertexAIndex;
 		UVParentVertex[NewVertexBIndex] = NewVertexBIndex;
 	}
 
-	// Final UVs
+	// UVs for the final quad	
 	{
-		const int NewTriAIndex = 2 * (NumInputVertices - 1);
-		const int NewTriBIndex = NewTriAIndex + 1;
+		const int32 NewTriAIndex = 2 * (FinalNumVerticesOneSide - 1);
+		const int32 NewTriBIndex = NewTriAIndex + 1;
 
-		const int NewVertexA = 2 * (NumInputVertices - 1);
-		const int NewVertexB = NewVertexA + 1;
-		const int NewVertexC = NewVertexA + 2;
-		const int NewVertexD = NewVertexA + 3;
+		const int32 NewVertexA = 2 * (FinalNumVerticesOneSide - 1);
+		const int32 NewVertexB = NewVertexA + 1;
+		const int32 NewVertexC = NewVertexA + 2;
+		const int32 NewVertexD = NewVertexA + 3;
 
 		ensure(NewVertexD < UVs.Num());
 
@@ -130,10 +466,10 @@ FMeshShapeGenerator& FPolygonEdgeMeshGenerator::Generate()
 		FIndex3i NewTriB{ NewVertexC, NewVertexB, NewVertexD };
 		SetTriangleUVs(NewTriBIndex, NewTriB);
 
-		UVs[2 * NumInputVertices] = { UVRight, UVBottom };  // BilinearInterp(UV00, UV01, UV11, UV10, 1.0, 0.0);		// == (UVRight, UVBottom)?
-		UVs[2 * NumInputVertices + 1] = { UVRight, UVTop }; // BilinearInterp(UV00, UV01, UV11, UV10, 1.0, 1.0);   // == (UVRight, UVTop)?
-		UVParentVertex[2 * NumInputVertices] = 0;
-		UVParentVertex[2 * NumInputVertices + 1] = 1;
+		UVs[2 * FinalNumVerticesOneSide] = { UVRight, UVBottom }; 
+		UVs[2 * FinalNumVerticesOneSide + 1] = { UVRight, UVTop };
+		UVParentVertex[2 * FinalNumVerticesOneSide] = 0;
+		UVParentVertex[2 * FinalNumVerticesOneSide + 1] = 1;
 	}
 
 	return *this;
diff --git a/Engine/Source/Runtime/GeometryCore/Public/Generators/PolygonEdgeMeshGenerator.h b/Engine/Source/Runtime/GeometryCore/Public/Generators/PolygonEdgeMeshGenerator.h
--- a/Engine/Source/Runtime/GeometryCore/Public/Generators/PolygonEdgeMeshGenerator.h
+++ b/Engine/Source/Runtime/GeometryCore/Public/Generators/PolygonEdgeMeshGenerator.h
@@ -16,12 +16,12 @@ class GEOMETRYCORE_API FPolygonEdgeMeshGenerator : public FMeshShapeGenerator
 
 private:
 
-	// Polygon to triangulate. Assumed to be closed (i.e. last edge is (LastVertex, FirstVertex)). If Polygon has 
+	// Polygon to thicken and triangulate. Assumed to be closed (i.e. last edge is (LastVertex, FirstVertex)). If Polygon has 
 	// self-intersections or degenerate edges, result is undefined.
 	TArray<FFrame3d> Polygon;
 
-	// For each polygon vertex, a scale factor for the patch width at that vertex. Helps keep the width constant
-	// going around acute corners. 
+	// For each polygon vertex, a scale factor for the patch width at that vertex. It's a variable offset in order to keep the overall 
+	// width constant going around acute corners. 
 	TArray<double> OffsetScaleFactors;
 
 	// Width of quads to generate.
@@ -39,10 +39,26 @@ public:
 	/** If true, output mesh has a single polygroup, otherwise each quad gets a separate group */
 	bool bSinglePolyGroup = false;
 
+	/** Use arc segments instead of straight lines in corners */
+	bool bRoundedCorners = false;
+
+	/** Radius of the corner arcs; this is only available if Rounded Corners is enabled */
+	double CornerRadius = 0.0;
+
+	/** Limit corner radius to prevent overlapping */
+	bool bLimitCornerRadius = false;
+
+	/** Number of radial subdivisions for rounded corners */
+	int NumArcVertices;
+
 	FPolygonEdgeMeshGenerator(const TArray<FFrame3d>& InPolygon,
 		const TArray<double>& InOffsetScaleFactors,
 		double InWidth = 1.0,
-		FVector3d InNormal = FVector3d::UnitZ());
+		FVector3d InNormal = FVector3d::UnitZ(),
+		bool bInRoundedCorners = false,
+		double InCornerRadius = 0.0,
+		bool bInLimitCornerRadius = false,
+		int InNumArcVertices = 3);
 
 	// Generate triangulation
 	// TODO: Enable more subdivisions along the width and length dimensions if requested
