GeometryProcessing: add MeshAttributeUtil file with CompactAttributeValues() function that compacts the values in an integer triangle attrib set (eg MaterialIDs)
#rnx

diff --git a/Engine/Source/Runtime/GeometryCore/Private/DynamicMesh/MeshAttributeUtil.cpp b/Engine/Source/Runtime/GeometryCore/Private/DynamicMesh/MeshAttributeUtil.cpp
new file mode 100644
--- /dev/null
+++ b/Engine/Source/Runtime/GeometryCore/Private/DynamicMesh/MeshAttributeUtil.cpp
@@ -0,0 +1,62 @@
+// Copyright Epic Games, Inc. All Rights Reserved.
+
+#include "DynamicMesh/MeshAttributeUtil.h"
+#include "IntBoxTypes.h"
+
+using namespace UE::Geometry;
+
+
+bool UE::Geometry::CompactAttributeValues(
+	const FDynamicMesh3& Mesh,
+	TDynamicMeshScalarTriangleAttribute<int32>& TriangleAttrib,
+	TArray<int32>& OldToNewMap,
+	TArray<int32>& NewToOldMap,
+	bool& bWasCompact)
+{
+	bWasCompact = false;
+
+	// compute range of values in attribute set
+	FInterval1i IndexRange = FInterval1i::Empty();
+	for (int32 TriangleID : Mesh.TriangleIndicesItr())
+	{
+		IndexRange.Contain(TriangleAttrib.GetValue(TriangleID));
+	}
+	if (IndexRange.Min < 0)
+	{
+		return false;
+	}
+
+	int32 MaxValue = IndexRange.Max;
+	OldToNewMap.Init(IndexConstants::InvalidID, MaxValue+1);
+
+	// generate remapping and set new values
+	int32 NewValueCount = 0;
+	for (int32 TriangleID : Mesh.TriangleIndicesItr())
+	{
+		int32 Value = TriangleAttrib.GetValue(TriangleID);
+		if (OldToNewMap[Value] == IndexConstants::InvalidID)
+		{
+			OldToNewMap[Value] = NewValueCount;
+			NewValueCount++;
+		}
+		int32 NewValue = OldToNewMap[Value];
+		if (NewValue != Value)
+		{
+			TriangleAttrib.SetValue(TriangleID, NewValue);
+		}
+	}
+
+	// construct inverse mapping
+	NewToOldMap.Init(IndexConstants::InvalidID, NewValueCount);
+	for (int32 k = 0; k <= MaxValue; ++k)
+	{
+		if (OldToNewMap[k] >= 0)
+		{
+			NewToOldMap[ OldToNewMap[k] ] = k;
+		}
+	}
+
+	bWasCompact = ( NewToOldMap.Num() == OldToNewMap.Num() );
+
+	return true;
+}
\ No newline at end of file
diff --git a/Engine/Source/Runtime/GeometryCore/Public/DynamicMesh/MeshAttributeUtil.h b/Engine/Source/Runtime/GeometryCore/Public/DynamicMesh/MeshAttributeUtil.h
new file mode 100644
--- /dev/null
+++ b/Engine/Source/Runtime/GeometryCore/Public/DynamicMesh/MeshAttributeUtil.h
@@ -0,0 +1,34 @@
+// Copyright Epic Games, Inc. All Rights Reserved.
+
+#pragma once
+
+#include "DynamicMesh/DynamicMesh3.h"
+#include "DynamicMesh/DynamicMeshAttributeSet.h"
+
+
+/**
+ * Utility functions for manipulating DynamicMesh attribute sets
+ */
+namespace UE
+{
+namespace Geometry
+{
+
+	/**
+	 * Compact the values of an integer Triangle Attribute, ie so that the attribute values are dense in range 0..N.
+	 * Useful for (eg) compacting MaterialIDs or Polygroups.
+	 * @param OldToNewMapOut generated mapping from previous IDs to new IDs (ie size is maximum input attrib value (+1), will contain InvalidID for any compacted values )
+	 * @param NewToOldMapOut generated mapping from new IDs to previous IDs (ie size is maximum output attribute value (+1), no invalid values)
+	 * @param bWasCompactOut set to true if the attribute set was already compact, IE was not modified by the operation
+	 * @return true on success (including if already compact). false if any Attribute Values were negative/InvalidID, cannot compact in this case
+	 */
+	GEOMETRYCORE_API bool CompactAttributeValues(
+		const FDynamicMesh3& Mesh,
+		TDynamicMeshScalarTriangleAttribute<int32>& TriangleAttrib,
+		TArray<int32>& OldToNewMapOut,
+		TArray<int32>& NewToOldMapOut,
+		bool& bWasCompactOut);
+
+
+}
+}
