ModelingTools: Fix BakeVertex alpha channel preview.

#rnx

diff --git a/Engine/Source/Runtime/GeometryCore/Public/VectorTypes.h b/Engine/Source/Runtime/GeometryCore/Public/VectorTypes.h
--- a/Engine/Source/Runtime/GeometryCore/Public/VectorTypes.h
+++ b/Engine/Source/Runtime/GeometryCore/Public/VectorTypes.h
@@ -400,7 +400,8 @@ constexpr FColor ToFColor(const UE::Math::TVector4<T>& Vector)
 	return FColor(
 		FMathf::Clamp((int)((float)Vector.X * 255.0f), 0, 255),
 		FMathf::Clamp((int)((float)Vector.Y * 255.0f), 0, 255),
-		FMathf::Clamp((int)((float)Vector.Z * 255.0f), 0, 255));
+		FMathf::Clamp((int)((float)Vector.Z * 255.0f), 0, 255),
+		FMathf::Clamp((int)((float)Vector.W * 255.0f), 0, 255));
 }
 
 
