MeshModelingToolset: Some minor comment/code cleanup

#rnx

diff --git a/Engine/Source/Runtime/GeometryCore/Private/DynamicMesh/DynamicMesh3.cpp b/Engine/Source/Runtime/GeometryCore/Private/DynamicMesh/DynamicMesh3.cpp
--- a/Engine/Source/Runtime/GeometryCore/Private/DynamicMesh/DynamicMesh3.cpp
+++ b/Engine/Source/Runtime/GeometryCore/Private/DynamicMesh/DynamicMesh3.cpp
@@ -579,7 +579,6 @@ void FDynamicMesh3::EnumerateVertexTriangles(int32 VertexID, TFunctionRef<void(i
 		{
 			ApplyFunc(Edge.Tri[0]);
 		}
-		const int et1 = Edge.Tri[1];
 		if (Edge.Tri[1] != InvalidID && TriHasSequentialVertices(Edge.Tri[1], VertexID, vOther))
 		{
 			ApplyFunc(Edge.Tri[1]);
diff --git a/Engine/Source/Runtime/GeometryCore/Private/Util/SmallListSet.cpp b/Engine/Source/Runtime/GeometryCore/Private/Util/SmallListSet.cpp
--- a/Engine/Source/Runtime/GeometryCore/Private/Util/SmallListSet.cpp
+++ b/Engine/Source/Runtime/GeometryCore/Private/Util/SmallListSet.cpp
@@ -111,7 +111,7 @@ bool FSmallListSet::Remove(int32 ListIndex, int32 Value)
 			{
 				ListBlocks[j - 1] = ListBlocks[j];
 			}
-			//block_store[iEnd] = -2;     // OPTIONAL
+			//ListBlocks[iEnd] = -2;     // OPTIONAL
 
 			if (N > BLOCKSIZE) 
 			{
diff --git a/Engine/Source/Runtime/GeometryCore/Public/Util/SmallListSet.h b/Engine/Source/Runtime/GeometryCore/Public/Util/SmallListSet.h
--- a/Engine/Source/Runtime/GeometryCore/Public/Util/SmallListSet.h
+++ b/Engine/Source/Runtime/GeometryCore/Public/Util/SmallListSet.h
@@ -14,7 +14,7 @@ namespace Geometry
 /**
  * FSmallListSet stores a set of short integer-valued variable-size lists.
  * The lists are encoded into a few large TDynamicVector buffers, with internal pooling,
- * so adding/removing lists usually does not involve any or delete ops.
+ * so adding/removing lists usually does not involve any new or delete ops.
  * 
  * The lists are stored in two parts. The first N elements are stored in a linear
  * subset of a TDynamicVector. If the list spills past these N elements, the extra elements
@@ -36,7 +36,7 @@ protected:
 	/** offset from start of linear-memory portion of list that contains pointer to head of variable-length linked list */
 	static constexpr int32 BLOCK_LIST_OFFSET = BLOCKSIZE + 1;
 
-	/** mapping from list index to offset into block_store that contains list data */
+	/** mapping from list index to offset into ListBlocks that contains list data */
 	TDynamicVector<int32> ListHeads{};
 
 	/** 
@@ -45,7 +45,7 @@ protected:
 	 */
 	TDynamicVector<int32> ListBlocks{};
 
-	/** list of free blocks as indices/offsets into block_store */
+	/** list of free blocks as indices/offsets into ListBlocks */
 	TDynamicVector<int32> FreeBlocks{};
 
 	/** number of allocated lists */
