Path extrude tool: add support for open curves in PolygonEdgeMeshGenerator

#rnx

diff --git a/Engine/Plugins/Experimental/MeshModelingToolsetExp/Source/MeshModelingToolsExp/Private/DrawPolyPathTool.cpp b/Engine/Plugins/Experimental/MeshModelingToolsetExp/Source/MeshModelingToolsExp/Private/DrawPolyPathTool.cpp
--- a/Engine/Plugins/Experimental/MeshModelingToolsetExp/Source/MeshModelingToolsExp/Private/DrawPolyPathTool.cpp
+++ b/Engine/Plugins/Experimental/MeshModelingToolsetExp/Source/MeshModelingToolsExp/Private/DrawPolyPathTool.cpp
@@ -11,7 +11,6 @@
 #include "Util/ColorConstants.h"
 #include "ToolSetupUtil.h"
 #include "DynamicMesh/MeshIndexUtil.h"
-#include "Generators/RectangleMeshGenerator.h"
 #include "Generators/PolygonEdgeMeshGenerator.h"
 #include "Distance/DistLine3Line3.h"
 #include "ModelingObjectsCreationAPI.h"
@@ -74,87 +73,54 @@ namespace DrawPolyPathToolLocals
 		ComputeArcLengths(UsePathPoints, ArcLengths);
 		double PathLength = ArcLengths.Last();
 
-		if (bPathIsClosed)
-		{
-			double CornerRadius = 2.0 * CornerRadiusFraction * OffsetDistance;
-			FPolygonEdgeMeshGenerator MeshGen(UsePathPoints, UseOffsetScaleFactors, OffsetDistance, FVector3d::UnitZ(), bRoundedCorners, CornerRadius, bLimitCornerRadius, NumCornerArcPoints);
-			MeshGen.bSinglePolyGroup = bSinglePolyGroup;
-			MeshGen.UVWidth = PathLength;
-			MeshGen.UVHeight = 2 * OffsetDistance;
-			MeshGen.Generate();
-			Mesh.Copy(&MeshGen);
+		const double PathWidth = 2.0 * OffsetDistance;
+		const double CornerRadius = CornerRadiusFraction * PathWidth;
+		
+		FPolygonEdgeMeshGenerator MeshGen(UsePathPoints, bPathIsClosed, UseOffsetScaleFactors, PathWidth, FVector3d::UnitZ(), bRoundedCorners, CornerRadius, bLimitCornerRadius, NumCornerArcPoints);
+		MeshGen.bSinglePolyGroup = bSinglePolyGroup;
+		MeshGen.UVWidth = PathLength;
+		MeshGen.UVHeight = 2 * OffsetDistance;
+		MeshGen.Generate();
+		Mesh.Copy(&MeshGen);
 
-			Mesh.EnableVertexUVs(FVector2f::Zero());
+		Mesh.EnableVertexUVs(FVector2f::Zero());
 
-			if (bRampMode)
-			{
-				// Temporarily set vertex UVs to arclengths, for use in interpolating height in ramp mode
+		if (bRampMode)
+		{
+			// Temporarily set vertex UVs to arclengths, for use in interpolating height in ramp mode
 
-				if (bRoundedCorners)
+			if (bRoundedCorners)
+			{
+				// If we added arcs to the corners, recompute arc lengths
+				const int N = Mesh.VertexCount() / 2;
+				ArcLengths.Init(0.0, N);
+				double CurPathLength = 0;
+				for (int k = 1; k < N; ++k)
 				{
-					// If we added arcs to the corners, recompute arc lengths
-					const int N = Mesh.VertexCount() / 2;
-					ArcLengths.Init(0.0, N);
-					double CurPathLength = 0;
-					for (int k = 1; k < N; ++k)
-					{
-						CurPathLength += Distance(Mesh.GetVertex(2 * k), Mesh.GetVertex(2 * (k - 1)));
-						ArcLengths[k] = CurPathLength;
-					}
-					PathLength = ArcLengths.Last();
+					CurPathLength += Distance(Mesh.GetVertex(2 * k), Mesh.GetVertex(2 * (k - 1)));
+					ArcLengths[k] = CurPathLength;
 				}
+				PathLength = ArcLengths.Last();
+			}
 
-				int NumMeshVertices = Mesh.VertexCount();
-				ensure(NumMeshVertices == Mesh.MaxVertexID());
-				ensure(NumMeshVertices == 2 * ArcLengths.Num());
+			int NumMeshVertices = Mesh.VertexCount();
+			ensure(NumMeshVertices == Mesh.MaxVertexID());
+			ensure(NumMeshVertices == 2 * ArcLengths.Num());
 
-				for (int k = 0; k < NumMeshVertices/2; ++k)
-				{
-					const float Alpha = (float)ArcLengths[k] / PathLength;
-					Mesh.SetVertexUV(2 * k, FVector2f(Alpha, (float)k));
-					Mesh.SetVertexUV(2 * k + 1, FVector2f(Alpha, (float)k));
-				}
+			for (int k = 0; k < NumMeshVertices/2; ++k)
+			{
+				const float Alpha = static_cast<float>(ArcLengths[k] / PathLength);
+				Mesh.SetVertexUV(2 * k, FVector2f(Alpha, static_cast<float>(k)));
+				Mesh.SetVertexUV(2 * k + 1, FVector2f(Alpha, static_cast<float>(k)));
+			}
 
+			if (bPathIsClosed)
+			{
 				// Set last vertex positions to match first vertex locations so we can construct the vertical wall
 				Mesh.SetVertex(NumMeshVertices - 2, Mesh.GetVertex(0));
 				Mesh.SetVertex(NumMeshVertices - 1, Mesh.GetVertex(1));
 			}
 		}
-		else
-		{
-			FRectangleMeshGenerator MeshGen;
-			MeshGen.bSinglePolyGroup = bSinglePolyGroup;
-			MeshGen.Width = PathLength;
-			MeshGen.Height = 2 * OffsetDistance;
-			MeshGen.Normal = FVector3f::UnitZ();
-			MeshGen.Origin = FVector3d(PathLength / 2, 0, 0);
-			MeshGen.HeightVertexCount = 2;
-			MeshGen.WidthVertexCount = NumPoints;
-			MeshGen.Generate();
-			Mesh.Copy(&MeshGen);
-			Mesh.EnableVertexUVs(FVector2f::Zero());		// we will store arc length for each vtx in VertexUV
-
-			double ShiftX = 0;
-			double DeltaX = PathLength / (double)(NumPoints - 1);
-			for (int32 k = 0; k < NumPoints; ++k)
-			{
-				FFrame3d PathFrame = UsePathPoints[k];
-				FVector3d V0 = Mesh.GetVertex(k);
-				V0.X -= ShiftX;
-				V0.Y *= UseOffsetScaleFactors[k];
-				V0 = PathFrame.FromFramePoint(V0);
-				Mesh.SetVertex(k, V0);
-				const float Alpha = (float)ArcLengths[k] / PathLength;
-				Mesh.SetVertexUV(k, FVector2f(Alpha, (float)k));
-				FVector3d V1 = Mesh.GetVertex(NumPoints + k);
-				V1.X -= ShiftX;
-				V1.Y *= UseOffsetScaleFactors[k];
-				V1 = PathFrame.FromFramePoint(V1);
-				Mesh.SetVertex(NumPoints + k, V1);
-				Mesh.SetVertexUV(NumPoints + k, FVector2f(Alpha, (float)k));
-				ShiftX += DeltaX;
-			}
-		}
 	}
 
 }	// namespace DrawPolyPathToolLocals
diff --git a/Engine/Source/Runtime/GeometryCore/Private/Generators/PolygonEdgeMeshGenerator.cpp b/Engine/Source/Runtime/GeometryCore/Private/Generators/PolygonEdgeMeshGenerator.cpp
--- a/Engine/Source/Runtime/GeometryCore/Private/Generators/PolygonEdgeMeshGenerator.cpp
+++ b/Engine/Source/Runtime/GeometryCore/Private/Generators/PolygonEdgeMeshGenerator.cpp
@@ -107,14 +107,14 @@ namespace PolygonEdgeMeshGeneratorLocal
 		GenerateArc(PointOnCircle0, PointOnCircle1, CircleCenter, CircleRadius, NumArcPoints, OutCorner);
 	}
 
-	FVector2d ProjectPathVertex(const FVector3d& InVertex, const FFrame3d& PolygonFrame)
+	FVector2d ProjectPathVertex(const FVector3d& InVertex, const FFrame3d& ProjectionFrame)
 	{
-		return PolygonFrame.ToPlaneUV(InVertex);
+		return ProjectionFrame.ToPlaneUV(InVertex);
 	}
 
-	FVector3d UnprojectPathVertex(const FVector2d& InVertex, const FFrame3d& PolygonFrame)
+	FVector3d UnprojectPathVertex(const FVector2d& InVertex, const FFrame3d& ProjectionFrame)
 	{
-		return PolygonFrame.FromPlaneUV(InVertex);
+		return ProjectionFrame.FromPlaneUV(InVertex);
 	}
 
 
@@ -145,90 +145,6 @@ namespace PolygonEdgeMeshGeneratorLocal
 		}
 	}
 
-	// Given a polygonal path, replace each corner with a discretized arc
-	void CurvePath(const TArray<FVector3d>& InPath, 
-		TArray<bool> InteriorAngleFlag, 
-		double Radius, 
-		double Width, 
-		bool bLimitCornerRadius,
-		int NumArcVertices,
-		TArray<double> MaxCornerRadii,
-		TArray<double> OtherSideMaxCornerRadii,
-		const FFrame3d& PolygonFrame,
-		TArray<FVector3d>& OutPath)
-	{
-		const int N = InPath.Num();
-		if (N < 3)
-		{
-			return;
-		}
-
-		for (int CornerID = 0; CornerID < N; ++CornerID)
-		{
-			const int Prev = (CornerID + N - 1) % N;
-			const int Next = (CornerID + 1) % N;
-
-			const FVector2d PrevVertex = ProjectPathVertex(InPath[Prev], PolygonFrame);
-			const FVector2d CurrVertex = ProjectPathVertex(InPath[CornerID], PolygonFrame);
-			const FVector2d NextVertex = ProjectPathVertex(InPath[Next], PolygonFrame);
-
-			TArray<FVector2d> ArcVertices;
-
-			double EffectiveRadius = Radius;
-			if (bLimitCornerRadius)
-			{
-				// The input Radius is taken as the desired radius for exterior corners. So the desired interior radius is the input
-				// radius minus width.
-				// Each corner has a max interior and max exterior radius. We want to respect these limits but still have:
-				//		interior radius = extrerior radius - width
-
-				if (InteriorAngleFlag[CornerID])
-				{
-					double ClampedExteriorRadius = FMath::Min(Radius, OtherSideMaxCornerRadii[CornerID]);
-					double InteriorMax = MaxCornerRadii[CornerID];
-					EffectiveRadius = FMath::Clamp(ClampedExteriorRadius - Width, 0.0, InteriorMax );
-				}
-				else
-				{
-					double InteriorMax = OtherSideMaxCornerRadii[CornerID];
-					double ExteriorMax = MaxCornerRadii[CornerID];
-					EffectiveRadius = FMath::Min(Radius, FMath::Min(ExteriorMax, InteriorMax + Width));
-				}
-			}
-			else
-			{
-				if (InteriorAngleFlag[CornerID])
-				{
-					EffectiveRadius -= Width;
-				}
-			}
-
-			if (EffectiveRadius > KINDA_SMALL_NUMBER)
-			{
-				RoundCorner(PrevVertex, CurrVertex, NextVertex, EffectiveRadius, NumArcVertices, ArcVertices);
-			}
-			else if (OtherSideMaxCornerRadii[CornerID] <= KINDA_SMALL_NUMBER && MaxCornerRadii[CornerID] <= KINDA_SMALL_NUMBER)
-			{
-				// Both interior and exterior radii are close to zero -- don't insert any arc here at all
-				ArcVertices.Add(CurrVertex);
-			}
-			else
-			{
-				// TODO: Make a proper triangle fan, not a bunch of vertices on top of each other
-				for (int I = 0; I < NumArcVertices; ++I)
-				{
-					ArcVertices.Add(CurrVertex);
-				}
-			}
-
-			for (const FVector2d& V : ArcVertices)
-			{
-				OutPath.Add(UnprojectPathVertex(V, PolygonFrame));
-			}
-		}
-	}
-
-
 	double ArcLength(const TArray<FVector3d>& Path, double MinSegmentArcLength = 0.1 )
 	{
 		double Length = 0;
@@ -261,6 +177,7 @@ namespace PolygonEdgeMeshGeneratorLocal
 
 
 FPolygonEdgeMeshGenerator::FPolygonEdgeMeshGenerator(const TArray<FFrame3d>& InPolygon,
+	bool bInClosed,
 	const TArray<double>& InOffsetScaleFactors,
 	double InWidth,
 	FVector3d InNormal,
@@ -269,6 +186,7 @@ FPolygonEdgeMeshGenerator::FPolygonEdgeMeshGenerator(const TArray<FFrame3d>& InP
 	bool bInLimitCornerRadius,
 	int InNumArcVertices) :
 	Polygon(InPolygon),
+	bClosed(bInClosed),
 	OffsetScaleFactors(InOffsetScaleFactors),
 	Width(InWidth),
 	Normal(InNormal),
@@ -281,6 +199,101 @@ FPolygonEdgeMeshGenerator::FPolygonEdgeMeshGenerator(const TArray<FFrame3d>& InP
 }
 
 
+
+// Given a piecewise linear path, replace each corner with a discretized arc
+void FPolygonEdgeMeshGenerator::CurvePath(const TArray<FVector3d>& InPath,
+	const TArray<bool>& InteriorAngleFlag,
+	const TArray<double>& MaxCornerRadii,
+	const TArray<double>& OtherSideMaxCornerRadii,
+	const FFrame3d& ProjectionFrame,
+	TArray<FVector3d>& OutPath) const
+{
+	using namespace PolygonEdgeMeshGeneratorLocal;
+
+	const int N = InPath.Num();
+	if (N < 3)
+	{
+		return;
+	}
+
+	int Begin = bClosed ? 0 : 1;
+	int End = bClosed ? N : N - 1;
+
+	if (!bClosed)
+	{
+		OutPath.Add(InPath[0]);
+	}
+
+	for (int CornerID = Begin; CornerID < End; ++CornerID)
+	{
+		const int Prev = bClosed ? (CornerID + N - 1) % N : CornerID - 1;
+		const int Next = bClosed ? (CornerID + 1) % N : CornerID + 1;
+
+		const FVector2d PrevVertex = ProjectPathVertex(InPath[Prev], ProjectionFrame);
+		const FVector2d CurrVertex = ProjectPathVertex(InPath[CornerID], ProjectionFrame);
+		const FVector2d NextVertex = ProjectPathVertex(InPath[Next], ProjectionFrame);
+
+		TArray<FVector2d> ArcVertices;
+
+		double EffectiveRadius = CornerRadius;
+		if (bLimitCornerRadius)
+		{
+			// The input Radius is taken as the desired radius for exterior corners. So the desired interior radius is the input
+			// radius minus width.
+			// Each corner has a max interior and max exterior radius. We want to respect these limits but still have:
+			//		interior radius = exterior radius - width
+
+			if (InteriorAngleFlag[CornerID])
+			{
+				double ClampedExteriorRadius = FMath::Min(CornerRadius, OtherSideMaxCornerRadii[CornerID]);
+				double InteriorMax = MaxCornerRadii[CornerID];
+				EffectiveRadius = FMath::Clamp(ClampedExteriorRadius - Width, 0.0, InteriorMax);
+			}
+			else
+			{
+				double InteriorMax = OtherSideMaxCornerRadii[CornerID];
+				double ExteriorMax = MaxCornerRadii[CornerID];
+				EffectiveRadius = FMath::Min(CornerRadius, FMath::Min(ExteriorMax, InteriorMax + Width));
+			}
+		}
+		else
+		{
+			if (InteriorAngleFlag[CornerID])
+			{
+				EffectiveRadius -= Width;
+			}
+		}
+
+		if (EffectiveRadius > KINDA_SMALL_NUMBER)
+		{
+			RoundCorner(PrevVertex, CurrVertex, NextVertex, EffectiveRadius, NumArcVertices, ArcVertices);
+		}
+		else if (OtherSideMaxCornerRadii[CornerID] <= KINDA_SMALL_NUMBER && MaxCornerRadii[CornerID] <= KINDA_SMALL_NUMBER)
+		{
+			// Both interior and exterior radii are close to zero -- don't insert any arc here at all
+			ArcVertices.Add(CurrVertex);
+		}
+		else
+		{
+			// TODO: Make a proper triangle fan, not a bunch of vertices on top of each other
+			for (int I = 0; I < NumArcVertices; ++I)
+			{
+				ArcVertices.Add(CurrVertex);
+			}
+		}
+
+		for (const FVector2d& V : ArcVertices)
+		{
+			OutPath.Add(UnprojectPathVertex(V, ProjectionFrame));
+		}
+	}
+
+	if (!bClosed)
+	{
+		OutPath.Add(InPath.Last());
+	}
+}
+
 // Generate triangulation
 // TODO: Enable more subdivisions along the width and length dimensions if requested
 FMeshShapeGenerator& FPolygonEdgeMeshGenerator::Generate()
@@ -299,8 +312,8 @@ FMeshShapeGenerator& FPolygonEdgeMeshGenerator::Generate()
 	TArray<FVector3d> RightSidePath;
 
 	// Trace the input path, placing vertices on either side of each input vertex 
-	const FVector3d LeftVertex{ 0, -Width, 0 };
-	const FVector3d RightVertex{ 0, Width, 0 };
+	const FVector3d LeftVertex{ 0, -0.5 * Width, 0 };
+	const FVector3d RightVertex{ 0, 0.5 * Width, 0 };
 	for (int32 CurrentInputVertex = 0; CurrentInputVertex < NumInputVertices; ++CurrentInputVertex)
 	{
 		const FFrame3d& CurrentFrame = Polygon[CurrentInputVertex];
@@ -331,19 +344,42 @@ FMeshShapeGenerator& FPolygonEdgeMeshGenerator::Generate()
 		// TODO: Support a "round outwards" mode where we replace the turning flags with an "interior to the polygon" flag
 		TArray<bool> TurnsLeft;
 		TArray<bool> TurnsRight;
-		for (int32 I = 0; I < CenterPath.Num(); ++I)
+
+		if (bClosed)
 		{
-			const int32 Prev = I == 0 ? CenterPath.Num() - 1 : I - 1;
-			const int32 Next = (I + 1) % CenterPath.Num();
-			const FVector2d A = ProjectPathVertex(CenterPath[Prev], PolygonFrame);
-			const FVector2d B = ProjectPathVertex(CenterPath[I], PolygonFrame);
-			const FVector2d C = ProjectPathVertex(CenterPath[Next], PolygonFrame);
-			TurnsLeft.Add(Orient(A, B, C) < 0);
-			TurnsRight.Add(!TurnsLeft.Last());
+			for (int32 I = 0; I < CenterPath.Num(); ++I)
+			{
+				const int32 Prev = I == 0 ? CenterPath.Num() - 1 : I - 1;
+				const int32 Next = (I + 1) % CenterPath.Num();
+				const FVector2d A = ProjectPathVertex(CenterPath[Prev], PolygonFrame);
+				const FVector2d B = ProjectPathVertex(CenterPath[I], PolygonFrame);
+				const FVector2d C = ProjectPathVertex(CenterPath[Next], PolygonFrame);
+				TurnsLeft.Add(Orient(A, B, C) < 0);
+				TurnsRight.Add(!TurnsLeft.Last());
+			}
 		}
+		else
+		{
+			TurnsLeft.Add(false);	// First point
+			TurnsRight.Add(false);
 
-		CurvePath(LeftSidePath,  TurnsLeft,  CornerRadius, 2.0 * Width, bLimitCornerRadius, NumArcVertices, LeftMaxCornerRadii, RightMaxCornerRadii, PolygonFrame, NewLeftSidePath);
-		CurvePath(RightSidePath, TurnsRight, CornerRadius, 2.0 * Width, bLimitCornerRadius, NumArcVertices, RightMaxCornerRadii, LeftMaxCornerRadii, PolygonFrame, NewRightSidePath);
+			for (int32 I = 1; I < CenterPath.Num()-1; ++I)
+			{
+				const int32 Prev = I - 1;
+				const int32 Next = I + 1;
+				const FVector2d A = ProjectPathVertex(CenterPath[Prev], PolygonFrame);
+				const FVector2d B = ProjectPathVertex(CenterPath[I], PolygonFrame);
+				const FVector2d C = ProjectPathVertex(CenterPath[Next], PolygonFrame);
+				TurnsLeft.Add(Orient(A, B, C) < 0);
+				TurnsRight.Add(!TurnsLeft.Last());
+			}
+
+			TurnsLeft.Add(false);	// Last point
+			TurnsRight.Add(false);
+		}
+
+		CurvePath(LeftSidePath,  TurnsLeft,  LeftMaxCornerRadii,  RightMaxCornerRadii, PolygonFrame, NewLeftSidePath);
+		CurvePath(RightSidePath, TurnsRight, RightMaxCornerRadii, LeftMaxCornerRadii,  PolygonFrame, NewRightSidePath);
 	}
 	else
 	{
@@ -360,7 +396,7 @@ FMeshShapeGenerator& FPolygonEdgeMeshGenerator::Generate()
 
 	const int32 FinalNumVertices = 2 * NewLeftSidePath.Num();
 	const int32 FinalNumUVs = 2 * NewLeftSidePath.Num() + 2;		// Extra pair of UVs for the final loop polygon
-	const int32 FinalNumTriangles = 2 * NewLeftSidePath.Num();
+	const int32 FinalNumTriangles = bClosed ?  2 * NewLeftSidePath.Num() : 2 * NewLeftSidePath.Num() - 2;
 	SetBufferSizes(FinalNumVertices, FinalNumTriangles, FinalNumUVs, FinalNumVertices);
 
 	Vertices.Empty();
@@ -371,7 +407,8 @@ FMeshShapeGenerator& FPolygonEdgeMeshGenerator::Generate()
 	}
 
 	int PolyIndex = 0;
-	for (int32 CurrentVertex = 0; CurrentVertex < NewLeftSidePath.Num(); ++CurrentVertex)
+	int EndVertex = bClosed ? NewLeftSidePath.Num() : NewLeftSidePath.Num() - 1;
+	for (int32 CurrentVertex = 0; CurrentVertex < EndVertex; ++CurrentVertex)
 	{
 		const int32 NewVertexA = 2 * CurrentVertex;
 		const int32 NewVertexB = NewVertexA + 1;
@@ -449,6 +486,7 @@ FMeshShapeGenerator& FPolygonEdgeMeshGenerator::Generate()
 	}
 
 	// UVs for the final quad	
+	if (bClosed)
 	{
 		const int32 NewTriAIndex = 2 * (FinalNumVerticesOneSide - 1);
 		const int32 NewTriBIndex = NewTriAIndex + 1;
diff --git a/Engine/Source/Runtime/GeometryCore/Public/Generators/PolygonEdgeMeshGenerator.h b/Engine/Source/Runtime/GeometryCore/Public/Generators/PolygonEdgeMeshGenerator.h
--- a/Engine/Source/Runtime/GeometryCore/Public/Generators/PolygonEdgeMeshGenerator.h
+++ b/Engine/Source/Runtime/GeometryCore/Public/Generators/PolygonEdgeMeshGenerator.h
@@ -20,6 +20,9 @@ private:
 	// self-intersections or degenerate edges, result is undefined.
 	TArray<FFrame3d> Polygon;
 
+	// True if Polygon represents a closed path
+	bool bClosed;
+
 	// For each polygon vertex, a scale factor for the patch width at that vertex. It's a variable offset in order to keep the overall 
 	// width constant going around acute corners. 
 	TArray<double> OffsetScaleFactors;
@@ -52,6 +55,7 @@ public:
 	int NumArcVertices;
 
 	FPolygonEdgeMeshGenerator(const TArray<FFrame3d>& InPolygon,
+		bool bInClosed,
 		const TArray<double>& InOffsetScaleFactors,
 		double InWidth = 1.0,
 		FVector3d InNormal = FVector3d::UnitZ(),
@@ -64,6 +68,15 @@ public:
 	// TODO: Enable more subdivisions along the width and length dimensions if requested
 	virtual FMeshShapeGenerator& Generate() final;
 
+private:
+
+	void CurvePath(const TArray<FVector3d>& InPath,
+		const TArray<bool>& InteriorAngleFlag,
+		const TArray<double>& MaxCornerRadii,
+		const TArray<double>& OtherSideMaxCornerRadii,
+		const FFrame3d& PolygonFrame,
+		TArray<FVector3d>& OutPath) const;
+
 };
 
 
