GeometryCore: Make top and front of steps generated by StairGenerator have different polygroup face ids

This slightly improves the UX when inserting edge loops because the retriangulation option works in a more intuitive way when face groups are planar

#rnx

diff --git a/Engine/Source/Runtime/GeometryCore/Private/Generators/StairGenerator.cpp b/Engine/Source/Runtime/GeometryCore/Private/Generators/StairGenerator.cpp
--- a/Engine/Source/Runtime/GeometryCore/Private/Generators/StairGenerator.cpp
+++ b/Engine/Source/Runtime/GeometryCore/Private/Generators/StairGenerator.cpp
@@ -307,39 +307,7 @@ FMeshShapeGenerator& FStairGenerator::GenerateSolidStairs()
 	for (int FaceId = 0; FaceId < NumQuads; FaceId++)
 	{
 		ESide Side = FaceToSide(FaceId);
-		if (bPolygroupPerQuad)
-		{
-			GroupId = FaceId;
-		}
-		else
-		{
-			switch (Side)
-			{
-			case ESide::Right:
-				GroupId = 0;
-				break;
-			case ESide::Left:
-				GroupId = 1;
-				break;
-			case ESide::Front:
-			case ESide::Top:
-			{
-				// Each actual step (front/top) face are grouped together
-				check(NumFrontFaceVertex == NumTopFaceVertex);
-				check(NumFrontFaceVertex % 4 == 0);
-				int StartFaceId = Side == ESide::Front ? FrontStartFaceId : TopStartFaceId;
-				int LocalFaceId = FaceId - StartFaceId / 4;
-				GroupId = LocalFaceId + 4;
-				break;
-			}
-			case ESide::Back:
-				GroupId = 2;
-				break;
-			case ESide::Bottom:
-				GroupId = 3;
-				break;
-			}
-		}
+		
 		const bool Flipped = (Side == ESide::Left);
 		for (int TriId = 0; TriId < 2; TriId++)
 		{
@@ -361,6 +329,39 @@ FMeshShapeGenerator& FStairGenerator::GenerateSolidStairs()
 			SetTrianglePolygon(CurrentTriId, GroupId);
 			CurrentTriId++;
 		}
+		
+		if (bPolygroupPerQuad)
+		{
+			GroupId += 1;
+		}
+		else
+		{
+			switch (Side)
+			{
+			// Right, Left, Back and Bottom sides are all in the same group. Since all faces within each side are
+			// consecutive we increment GroupId only when the side is going to change in the next iteration
+			case ESide::Right:
+			case ESide::Left:
+			case ESide::Back:
+			case ESide::Bottom:
+			{
+				ESide NextSide = FaceToSide(FaceId + 1 < NumQuads ? FaceId + 1 : FaceId);
+				GroupId += (NextSide != Side);
+				break;
+			}
+				
+			// Front and Top faces have incrementing GroupIds. Note: Previously faces on the same step were grouped
+			// together but we changed that because having non-planar groups made edge loop insertion with
+			// retriangulation screw up the step geometry
+			case ESide::Front:
+			case ESide::Top:
+			{
+				GroupId += 1;
+				break;
+			}
+			}
+		}
+
 	}
 
 	return *this;
@@ -714,52 +715,7 @@ FMeshShapeGenerator& FStairGenerator::GenerateFloatingStairs()
 	for (int FaceId = 0; FaceId < NumQuads; FaceId++)
 	{
 		ESide Side = FaceToSide(FaceId);
-		if (bPolygroupPerQuad)
-		{
-			GroupId = FaceId;
-		}
-		else
-		{
-			switch (Side)
-			{
-			case ESide::Right:
-				GroupId = 0;
-				break;
-			case ESide::Left:
-				GroupId = 1;
-				break;
-			case ESide::Front:
-			case ESide::Top:
-			{
-				// Each actual step (front/top) face are grouped together
-				check(NumFrontFaceVertex == NumTopFaceVertex);
-				check(NumFrontFaceVertex % 4 == 0);
-				int StartFaceId = Side == ESide::Front ? FrontStartFaceId : TopStartFaceId;
-				int LocalFaceId = FaceId - StartFaceId / 4;
-				GroupId = LocalFaceId + 2;
-				break;
-			}
-			case ESide::Back:
-			case ESide::Bottom:
-			{
-				// Each actual step (back/bottom) face are grouped together
-				check(NumBackFaceVertex == NumBackFaceVertex);
-				check(NumBackFaceVertex % 4 == 0);
-				int StartFaceId = Side == ESide::Back ? BackStartFaceId : BottomStartFaceId;
-				int LocalFaceId = FaceId - StartFaceId / 4;
-				if (Side == ESide::Back && LocalFaceId == NumSteps - 1)
-				{
-					LocalFaceId--;
-				}
-				else if (Side == ESide::Bottom)
-				{
-					LocalFaceId = LocalFaceId > 0 ? LocalFaceId - 1 : LocalFaceId;
-				}
-				GroupId = LocalFaceId + 2 + NumFrontFaceVertex / 4;
-				break;
-			}
-			}
-		}
+
 		const bool Flipped = (Side == ESide::Left);
 		for (int TriId = 0; TriId < 2; TriId++)
 		{
@@ -781,6 +737,39 @@ FMeshShapeGenerator& FStairGenerator::GenerateFloatingStairs()
 			SetTrianglePolygon(CurrentTriId, GroupId);
 			CurrentTriId++;
 		}
+
+		if (bPolygroupPerQuad)
+		{
+			GroupId += 1;
+		}
+		else
+		{
+			switch (Side)
+			{
+			// Right and Left sides are all in the same group. Since all faces within each side are
+			// consecutive we increment GroupId only when the side is going to change in the next iteration
+			case ESide::Right:
+			case ESide::Left:
+			{
+				ESide NextSide = FaceToSide(FaceId + 1 < NumQuads ? FaceId + 1 : FaceId);
+				GroupId += (NextSide != Side);
+				break;
+			}
+				
+			// Front/Top and Back/Bottom faces have incrementing GroupIds. Note: Previously faces on the same step were
+			// grouped together but we changed that because having non-planar groups made edge loop insertion with
+			// retriangulation screw up the step geometry
+			case ESide::Front:
+			case ESide::Top:
+			case ESide::Back:
+			case ESide::Bottom:
+			{
+				GroupId += 1;
+				break;
+			}
+			}
+		}
+		
 	}
 
 	return *this;
