GeometryCore: Fix crash in TMeshWindingNumberGrid for zero cell size

diff --git a/Engine/Source/Runtime/GeometryCore/Public/Spatial/MeshWindingNumberGrid.h b/Engine/Source/Runtime/GeometryCore/Public/Spatial/MeshWindingNumberGrid.h
--- a/Engine/Source/Runtime/GeometryCore/Public/Spatial/MeshWindingNumberGrid.h
+++ b/Engine/Source/Runtime/GeometryCore/Public/Spatial/MeshWindingNumberGrid.h
@@ -89,13 +89,16 @@ public:
 		int nj = (int)((max.Y - GridOrigin.Y) / (float)CellSize) + 1;
 		int nk = (int)((max.Z - GridOrigin.Z) / (float)CellSize) + 1;
 
-		if (ComputeMode == EComputeModes::FullGrid)
+		if (ni >= 0 && nj >= 0 && nk >= 0)
 		{
-			make_grid_dense(GridOrigin, (float)CellSize, ni, nj, nk, WindingGrid);
-		}
-		else
-		{
-			make_grid(GridOrigin, (float)CellSize, ni, nj, nk, WindingGrid);
+			if (ComputeMode == EComputeModes::FullGrid)
+			{
+				make_grid_dense(GridOrigin, (float)CellSize, ni, nj, nk, WindingGrid);
+			}
+			else
+			{
+				make_grid(GridOrigin, (float)CellSize, ni, nj, nk, WindingGrid);
+			}
 		}
 
 		if (!bWantMeshSDFGrid)
