Path Extrude Tool: support a single-segment path

#rnx

diff --git a/Engine/Source/Runtime/GeometryCore/Private/Generators/PolygonEdgeMeshGenerator.cpp b/Engine/Source/Runtime/GeometryCore/Private/Generators/PolygonEdgeMeshGenerator.cpp
--- a/Engine/Source/Runtime/GeometryCore/Private/Generators/PolygonEdgeMeshGenerator.cpp
+++ b/Engine/Source/Runtime/GeometryCore/Private/Generators/PolygonEdgeMeshGenerator.cpp
@@ -301,11 +301,6 @@ FMeshShapeGenerator& FPolygonEdgeMeshGenerator::Generate()
 	using namespace PolygonEdgeMeshGeneratorLocal;
 
 	const int32 NumInputVertices = Polygon.Num();
-	check(NumInputVertices >= 3);
-	if (NumInputVertices < 3)
-	{
-		return *this;
-	}
 
 	TArray<FVector3d> CenterPath;
 	TArray<FVector3d> LeftSidePath;
@@ -332,7 +327,7 @@ FMeshShapeGenerator& FPolygonEdgeMeshGenerator::Generate()
 	
 	const FFrame3d PolygonFrame = Polygon[0];
 
-	if (bRoundedCorners)
+	if (bRoundedCorners && NumInputVertices >= 3)
 	{
 		// Compute the maximum allowable radius for each corner, on each path side
 		TArray<double> LeftMaxCornerRadii, RightMaxCornerRadii;
