MeshModelingTools: Fix normals not being transferred to overlays in Remesh and Simplify when discarding attributes at start. Fix CompactCopy to properly copy vertex normals/uvs.

#rnx

diff --git a/Engine/Plugins/Experimental/MeshModelingToolsetExp/Source/MeshModelingToolsEditorOnly/Private/SimplifyMeshTool.cpp b/Engine/Plugins/Experimental/MeshModelingToolsetExp/Source/MeshModelingToolsEditorOnly/Private/SimplifyMeshTool.cpp
--- a/Engine/Plugins/Experimental/MeshModelingToolsetExp/Source/MeshModelingToolsEditorOnly/Private/SimplifyMeshTool.cpp
+++ b/Engine/Plugins/Experimental/MeshModelingToolsetExp/Source/MeshModelingToolsEditorOnly/Private/SimplifyMeshTool.cpp
@@ -191,6 +191,8 @@ TUniquePtr<FDynamicMeshOperator> USimplifyMeshTool::MakeNewOperator()
 	TUniquePtr<FSimplifyMeshOp> Op = MakeUnique<FSimplifyMeshOp>();
 
 	Op->bDiscardAttributes = SimplifyProperties->bDiscardAttributes;
+	// We always want attributes enabled on result even if we discard them initially
+	Op->bResultMustHaveAttributesEnabled = true;
 	Op->bPreventNormalFlips = SimplifyProperties->bPreventNormalFlips;
 	Op->bPreserveSharpEdges = SimplifyProperties->bPreserveSharpEdges;
 	Op->bAllowSeamCollapse = !SimplifyProperties->bPreserveSharpEdges;
diff --git a/Engine/Plugins/Experimental/MeshModelingToolsetExp/Source/MeshModelingToolsExp/Private/RemeshMeshTool.cpp b/Engine/Plugins/Experimental/MeshModelingToolsetExp/Source/MeshModelingToolsExp/Private/RemeshMeshTool.cpp
--- a/Engine/Plugins/Experimental/MeshModelingToolsetExp/Source/MeshModelingToolsExp/Private/RemeshMeshTool.cpp
+++ b/Engine/Plugins/Experimental/MeshModelingToolsetExp/Source/MeshModelingToolsExp/Private/RemeshMeshTool.cpp
@@ -208,6 +208,8 @@ TUniquePtr<FDynamicMeshOperator> URemeshMeshTool::MakeNewOperator()
 
 	Op->bCollapses = BasicProperties->bCollapses;
 	Op->bDiscardAttributes = BasicProperties->bDiscardAttributes;
+	// We always want attributes enabled on result even if we discard them initially
+	Op->bResultMustHaveAttributesEnabled = true;
 	Op->bFlips = BasicProperties->bFlips;
 	Op->bPreserveSharpEdges = BasicProperties->bPreserveSharpEdges;
 	Op->MeshBoundaryConstraint = (EEdgeRefineFlags)BasicProperties->MeshBoundaryConstraint;
diff --git a/Engine/Plugins/Runtime/MeshModelingToolset/Source/ModelingOperators/Private/CleaningOps/RemeshMeshOp.cpp b/Engine/Plugins/Runtime/MeshModelingToolset/Source/ModelingOperators/Private/CleaningOps/RemeshMeshOp.cpp
--- a/Engine/Plugins/Runtime/MeshModelingToolset/Source/ModelingOperators/Private/CleaningOps/RemeshMeshOp.cpp
+++ b/Engine/Plugins/Runtime/MeshModelingToolset/Source/ModelingOperators/Private/CleaningOps/RemeshMeshOp.cpp
@@ -5,6 +5,7 @@
 #include "DynamicMesh/DynamicMesh3.h"
 #include "DynamicMesh/DynamicMeshAABBTree3.h"
 #include "DynamicMesh/DynamicMeshAttributeSet.h"
+#include "DynamicMesh/MeshAttributeUtil.h"
 #include "Remesher.h"
 #include "QueueRemesher.h"
 #include "MeshConstraintsUtil.h"
@@ -195,4 +196,16 @@ void FRemeshMeshOp::CalculateResult(FProgressCancel* Progress)
 		}
 	}
 
+	if (!TargetMesh->HasAttributes() && bResultMustHaveAttributesEnabled)
+	{
+		TargetMesh->EnableAttributes();
+		if (TargetMesh->HasVertexUVs())
+		{
+			CopyVertexUVsToOverlay(*TargetMesh, *TargetMesh->Attributes()->PrimaryUV());
+		}
+		if (TargetMesh->HasVertexNormals())
+		{
+			CopyVertexNormalsToOverlay(*TargetMesh, *TargetMesh->Attributes()->PrimaryNormals());
+		}
+	}
 }
diff --git a/Engine/Plugins/Runtime/MeshModelingToolset/Source/ModelingOperators/Public/CleaningOps/RemeshMeshOp.h b/Engine/Plugins/Runtime/MeshModelingToolset/Source/ModelingOperators/Public/CleaningOps/RemeshMeshOp.h
--- a/Engine/Plugins/Runtime/MeshModelingToolset/Source/ModelingOperators/Public/CleaningOps/RemeshMeshOp.h
+++ b/Engine/Plugins/Runtime/MeshModelingToolset/Source/ModelingOperators/Public/CleaningOps/RemeshMeshOp.h
@@ -65,6 +65,9 @@ public:
 	float SmoothingStrength, TargetEdgeLength;
 	ERemeshSmoothingType SmoothingType;
 	bool bDiscardAttributes, bPreserveSharpEdges, bFlips, bSplits, bCollapses, bReproject, bPreventNormalFlips;
+	// When true, result will have attributes object regardless of whether attributes 
+	// were discarded or present initially.
+	bool bResultMustHaveAttributesEnabled = false;
 	EEdgeRefineFlags MeshBoundaryConstraint, GroupBoundaryConstraint, MaterialBoundaryConstraint;
 
 	FDynamicMesh3* ProjectionTarget = nullptr;
diff --git a/Engine/Plugins/Runtime/MeshModelingToolset/Source/ModelingOperatorsEditorOnly/Private/CleaningOps/SimplifyMeshOp.cpp b/Engine/Plugins/Runtime/MeshModelingToolset/Source/ModelingOperatorsEditorOnly/Private/CleaningOps/SimplifyMeshOp.cpp
--- a/Engine/Plugins/Runtime/MeshModelingToolset/Source/ModelingOperatorsEditorOnly/Private/CleaningOps/SimplifyMeshOp.cpp
+++ b/Engine/Plugins/Runtime/MeshModelingToolset/Source/ModelingOperatorsEditorOnly/Private/CleaningOps/SimplifyMeshOp.cpp
@@ -7,6 +7,7 @@
 
 #include "DynamicMesh/DynamicMesh3.h"
 #include "DynamicMesh/DynamicMeshAttributeSet.h"
+#include "DynamicMesh/MeshAttributeUtil.h"
 #include "MeshSimplification.h"
 #include "MeshConstraintsUtil.h"
 #include "ProjectionTargets.h"
@@ -329,4 +330,16 @@ void FSimplifyMeshOp::CalculateResult(FProgressCancel* Progress)
 		FMeshNormals::QuickComputeVertexNormals(*ResultMesh);
 	}
 
+	if (!TargetMesh->HasAttributes() && bResultMustHaveAttributesEnabled)
+	{
+		TargetMesh->EnableAttributes();
+		if (TargetMesh->HasVertexUVs())
+		{
+			CopyVertexUVsToOverlay(*TargetMesh, *TargetMesh->Attributes()->PrimaryUV());
+		}
+		if (TargetMesh->HasVertexNormals())
+		{
+			CopyVertexNormalsToOverlay(*TargetMesh, *TargetMesh->Attributes()->PrimaryNormals());
+		}
+	}
 }
diff --git a/Engine/Plugins/Runtime/MeshModelingToolset/Source/ModelingOperatorsEditorOnly/Public/CleaningOps/SimplifyMeshOp.h b/Engine/Plugins/Runtime/MeshModelingToolset/Source/ModelingOperatorsEditorOnly/Public/CleaningOps/SimplifyMeshOp.h
--- a/Engine/Plugins/Runtime/MeshModelingToolset/Source/ModelingOperatorsEditorOnly/Public/CleaningOps/SimplifyMeshOp.h
+++ b/Engine/Plugins/Runtime/MeshModelingToolset/Source/ModelingOperatorsEditorOnly/Public/CleaningOps/SimplifyMeshOp.h
@@ -75,6 +75,9 @@ public:
 	int TargetPercentage, TargetCount;
 	float TargetEdgeLength;
 	bool bDiscardAttributes, bReproject, bPreventNormalFlips, bPreserveSharpEdges, bAllowSeamCollapse;
+	// When true, result will have attributes object regardless of whether attributes 
+	// were discarded or present initially.
+	bool bResultMustHaveAttributesEnabled = false;
 	UE::Geometry::EEdgeRefineFlags MeshBoundaryConstraint, GroupBoundaryConstraint, MaterialBoundaryConstraint;
 	/** Angle threshold in degrees used for testing if two triangles should be considered coplanar, or two lines collinear */
 	float MinimalPlanarAngleThresh = 0.01f;
diff --git a/Engine/Source/Runtime/GeometryCore/Private/DynamicMesh/DynamicMesh3.cpp b/Engine/Source/Runtime/GeometryCore/Private/DynamicMesh/DynamicMesh3.cpp
--- a/Engine/Source/Runtime/GeometryCore/Private/DynamicMesh/DynamicMesh3.cpp
+++ b/Engine/Source/Runtime/GeometryCore/Private/DynamicMesh/DynamicMesh3.cpp
@@ -237,6 +237,18 @@ void FDynamicMesh3::CompactCopy(const FDynamicMesh3& copy, bool bNormals, bool b
 
 	// currently cannot re-use existing attribute buffers
 	Clear();
+	if (bNormals && copy.HasVertexNormals())
+	{
+		EnableVertexNormals(FVector3f::UnitY());
+	}
+	if (bColors && copy.HasVertexColors())
+	{
+		EnableVertexColors(FVector3f::One());
+	}
+	if (bUVs && copy.HasVertexUVs())
+	{
+		EnableVertexUVs(FVector2f::Zero());
+	}
 
 	// Use a triangle map if we have a CompactInfo or we need to copy attributes.
 	const bool bUseTriangleMap = CompactInfo != nullptr || (bAttributes && copy.HasAttributes());
diff --git a/Engine/Source/Runtime/GeometryCore/Private/DynamicMesh/MeshAttributeUtil.cpp b/Engine/Source/Runtime/GeometryCore/Private/DynamicMesh/MeshAttributeUtil.cpp
--- a/Engine/Source/Runtime/GeometryCore/Private/DynamicMesh/MeshAttributeUtil.cpp
+++ b/Engine/Source/Runtime/GeometryCore/Private/DynamicMesh/MeshAttributeUtil.cpp
@@ -62,4 +62,78 @@ bool UE::Geometry::CompactAttributeValues(
 	bWasCompact = ( NewToOldMap.Num() == OldToNewMap.Num() );
 
 	return true;
-}
\ No newline at end of file
+}
+
+bool UE::Geometry::CopyVertexUVsToOverlay(
+	const FDynamicMesh3& Mesh,
+	FDynamicMeshUVOverlay& UVOverlayOut,
+	bool bCompactElements)
+{
+	if (!Mesh.HasVertexUVs())
+	{
+		return false;
+	}
+
+	if (UVOverlayOut.ElementCount() > 0)
+	{
+		UVOverlayOut.ClearElements();
+	}
+	
+	UVOverlayOut.BeginUnsafeElementsInsert();
+	for (int32 Vid : Mesh.VertexIndicesItr())
+	{
+		FVector2f UV = Mesh.GetVertexUV(Vid);
+		UVOverlayOut.InsertElement(Vid, &UV.X, true);
+	}
+	UVOverlayOut.EndUnsafeElementsInsert();
+
+	for (int32 Tid : Mesh.TriangleIndicesItr())
+	{
+		UVOverlayOut.SetTriangle(Tid, Mesh.GetTriangle(Tid));
+	}
+
+	if (bCompactElements)
+	{
+		FCompactMaps CompactMaps;
+		UVOverlayOut.CompactInPlace(CompactMaps);
+	}
+
+	return true;
+}
+
+bool UE::Geometry::CopyVertexNormalsToOverlay(
+	const FDynamicMesh3& Mesh,
+	FDynamicMeshNormalOverlay& NormalOverlayOut,
+	bool bCompactElements)
+{
+	if (!Mesh.HasVertexNormals())
+	{
+		return false;
+	}
+
+	if (NormalOverlayOut.ElementCount() > 0)
+	{
+		NormalOverlayOut.ClearElements();
+	}
+
+	NormalOverlayOut.BeginUnsafeElementsInsert();
+	for (int32 Vid : Mesh.VertexIndicesItr())
+	{
+		FVector3f Normal = Mesh.GetVertexNormal(Vid);
+		NormalOverlayOut.InsertElement(Vid, &Normal.X, true);
+	}
+	NormalOverlayOut.EndUnsafeElementsInsert();
+
+	for (int32 Tid : Mesh.TriangleIndicesItr())
+	{
+		NormalOverlayOut.SetTriangle(Tid, Mesh.GetTriangle(Tid));
+	}
+
+	if (bCompactElements)
+	{
+		FCompactMaps CompactMaps;
+		NormalOverlayOut.CompactInPlace(CompactMaps);
+	}
+
+	return true;
+}
diff --git a/Engine/Source/Runtime/GeometryCore/Public/DynamicMesh/MeshAttributeUtil.h b/Engine/Source/Runtime/GeometryCore/Public/DynamicMesh/MeshAttributeUtil.h
--- a/Engine/Source/Runtime/GeometryCore/Public/DynamicMesh/MeshAttributeUtil.h
+++ b/Engine/Source/Runtime/GeometryCore/Public/DynamicMesh/MeshAttributeUtil.h
@@ -37,6 +37,26 @@ namespace Geometry
 		TArray<int32>& NewToOldMapOut,
 		bool& bWasCompactOut);
 
+	/**
+	 * Copies vertex UVs to a given overlay (clearing it before use if it has any elements).
+	 * 
+	 * @param bCompactElements If true, elements are compacted in the overlay
+	 * @return true if successful
+	 */
+	GEOMETRYCORE_API bool CopyVertexUVsToOverlay(
+		const FDynamicMesh3& Mesh,
+		FDynamicMeshUVOverlay& UVOverlayOut,
+		bool bCompactElements = false);
 
+	/**
+	 * Copies vertex normals to a given overlay (clearing it before use if it has any elements).
+	 *
+	 * @param bCompactElements If true, elements are compacted in the overlay
+	 * @return true if successful
+	 */
+	GEOMETRYCORE_API bool CopyVertexNormalsToOverlay(
+		const FDynamicMesh3& Mesh,
+		FDynamicMeshNormalOverlay& NormalOverlayOut,
+		bool bCompactElements = false);
 }
 }
