GeometryCore: Disable CheckValidity before/after visualization

diff --git a/Engine/Source/Runtime/GeometryCore/Private/DynamicMesh/DynamicMesh3_Serialization.cpp b/Engine/Source/Runtime/GeometryCore/Private/DynamicMesh/DynamicMesh3_Serialization.cpp
--- a/Engine/Source/Runtime/GeometryCore/Private/DynamicMesh/DynamicMesh3_Serialization.cpp
+++ b/Engine/Source/Runtime/GeometryCore/Private/DynamicMesh/DynamicMesh3_Serialization.cpp
@@ -248,8 +248,10 @@ void FDynamicMesh3::Serialize(FArchive& Ar)
 		return;
 	}
 
-	// Check validity in debug before saving data.
-	checkSlow(Ar.IsLoading() || CheckValidity(FValidityOptions(), EValidityCheckFailMode::Ensure));
+	constexpr bool bCheckValidityForDebugging = false;
+
+	// Check validity for debugging before saving data.
+	checkSlow(bCheckValidityForDebugging && (Ar.IsLoading() || CheckValidity(FValidityOptions(), EValidityCheckFailMode::Ensure)));
 
 	Ar.UsingCustomVersion(FUE5MainStreamObjectVersion::GUID);
 	if (Ar.IsLoading() && Ar.CustomVer(FUE5MainStreamObjectVersion::GUID) < FUE5MainStreamObjectVersion::DynamicMeshCompactedSerialization)
@@ -273,8 +275,8 @@ void FDynamicMesh3::Serialize(FArchive& Ar)
 		}
 	}
 
-	// Check validity in debug after loading data.
-	checkSlow(!Ar.IsLoading() || CheckValidity(FValidityOptions(), EValidityCheckFailMode::Ensure));
+	// Check validity for debugging after loading data.
+	checkSlow(bCheckValidityForDebugging && (!Ar.IsLoading() || CheckValidity(FValidityOptions(), EValidityCheckFailMode::Ensure)));
 
 	UE_LOG(LogGeometry, Verbose, TEXT("NEW %s of dynamic mesh took %2.3f ms"), ANSI_TO_TCHAR(Ar.IsLoading() ? "deserialization" : "serialization"),
 	       1000.0f * static_cast<float>(FPlatformTime::Seconds() - TimeStart));
