fix crash in sparse octree on growing FDynamicFlagArray by exactly the GrowChunkSize
also fix crash in sparse octree FDynamicFlagArray if Set(Idx, false) where Idx > current Max

#rnx

diff --git a/Engine/Source/Runtime/GeometryCore/Public/Spatial/SparseDynamicOctree3.h b/Engine/Source/Runtime/GeometryCore/Public/Spatial/SparseDynamicOctree3.h
--- a/Engine/Source/Runtime/GeometryCore/Public/Spatial/SparseDynamicOctree3.h
+++ b/Engine/Source/Runtime/GeometryCore/Public/Spatial/SparseDynamicOctree3.h
@@ -41,11 +41,16 @@ public:
 	{
 		if (bValue && BitIndex >= MaxIndex)
 		{
-			int32 ToAdd = (BitIndex - MaxIndex) | GrowChunkSize;
+			int32 ToAdd = (BitIndex - MaxIndex + 1) | GrowChunkSize;
 			BitArray.Add(false, ToAdd);
 			MaxIndex = BitArray.Num();
+			BitArray[BitIndex] = bValue;
 		}
-		BitArray[BitIndex] = bValue;
+		else if (BitIndex < MaxIndex)
+		{
+			BitArray[BitIndex] = bValue;
+		}
+		// don't need to set anything for false bValue with BitIndex beyond bounds; Get() will return false beyond bounds
 	}
 
 	inline bool Get(uint32 BitIndex) const
