GeometryCore:
- add FColliderMeshProjectionTarget to allow a FColliderMesh to be used as a projection target
- remove dead code from TDenseGrid3
- add TBlockedDenseGrid3, variant of TDenseGrid3 that lazy-allocates sub-blocks of a dense grid, to reduce memory usage in sparse/narrow-band situations on large uniform grids
- replace usage of TDenseGrid3 in FMarchingCubes with TBlockedDenseGrid3, to reduce memory footprint on large meshing problems

diff --git a/Engine/Source/Runtime/GeometryCore/Private/DynamicMesh/ColliderMesh.cpp b/Engine/Source/Runtime/GeometryCore/Private/DynamicMesh/ColliderMesh.cpp
--- a/Engine/Source/Runtime/GeometryCore/Private/DynamicMesh/ColliderMesh.cpp
+++ b/Engine/Source/Runtime/GeometryCore/Private/DynamicMesh/ColliderMesh.cpp
@@ -2,6 +2,7 @@
 
 #include "DynamicMesh/ColliderMesh.h"
 #include "DynamicMesh/DynamicMesh3.h"
+#include "Distance/DistPoint3Triangle3.h"
 
 using namespace UE::Geometry;
 
@@ -131,3 +132,58 @@ int32 FColliderMesh::GetSourceTriangleID(int32 TriangleID) const
 		return (TriangleID >= 0 && TriangleID < SourceTriangleIDs.Num()) ? SourceTriangleIDs[TriangleID] : IndexConstants::InvalidID;
 	}
 }
+
+
+
+
+FVector3d FColliderMeshProjectionTarget::Project(const FVector3d& Point, int Identifier)
+{
+	double DistSqr;
+	int NearestTriID = ColliderMesh->GetAABBTree().FindNearestTriangle(Point, DistSqr);
+	if (NearestTriID < 0)
+	{
+		return Point;
+	}
+
+	FTriangle3d Triangle;
+	ColliderMesh->GetTriVertices(NearestTriID, Triangle.V[0], Triangle.V[1], Triangle.V[2]);
+
+	FDistPoint3Triangle3d DistanceQuery(Point, Triangle);
+	DistanceQuery.GetSquared();
+	if (VectorUtil::IsFinite(DistanceQuery.ClosestTrianglePoint))
+	{
+		return DistanceQuery.ClosestTrianglePoint;
+	}
+	else
+	{
+		return Point;
+	}
+}
+
+
+
+FVector3d FColliderMeshProjectionTarget::Project(const FVector3d& Point, FVector3d& ProjectNormalOut, int Identifier)
+{
+	double DistSqr;
+	int NearestTriID = ColliderMesh->GetAABBTree().FindNearestTriangle(Point, DistSqr);
+	if (NearestTriID < 0)
+	{
+		return Point;
+	}
+
+	FTriangle3d Triangle;
+	ColliderMesh->GetTriVertices(NearestTriID, Triangle.V[0], Triangle.V[1], Triangle.V[2]);
+
+	ProjectNormalOut = Triangle.Normal();
+
+	FDistPoint3Triangle3d DistanceQuery(Point, Triangle);
+	DistanceQuery.GetSquared();
+	if (VectorUtil::IsFinite(DistanceQuery.ClosestTrianglePoint))
+	{
+		return DistanceQuery.ClosestTrianglePoint;
+	}
+	else
+	{
+		return Point;
+	}
+}
\ No newline at end of file
diff --git a/Engine/Source/Runtime/GeometryCore/Public/DynamicMesh/ColliderMesh.h b/Engine/Source/Runtime/GeometryCore/Public/DynamicMesh/ColliderMesh.h
--- a/Engine/Source/Runtime/GeometryCore/Public/DynamicMesh/ColliderMesh.h
+++ b/Engine/Source/Runtime/GeometryCore/Public/DynamicMesh/ColliderMesh.h
@@ -3,6 +3,7 @@
 #pragma once
 
 #include "Spatial/MeshAABBTree3.h"
+#include "Spatial/SpatialInterfaces.h"
 
 namespace UE
 {
@@ -109,6 +110,34 @@ protected:
 
 };
 
+
+/**
+ * Projection target API wrapper for an FColliderMesh
+ */
+class GEOMETRYCORE_API FColliderMeshProjectionTarget : public IOrientedProjectionTarget
+{
+public:
+	FColliderMesh* ColliderMesh = nullptr;
+
+	FColliderMeshProjectionTarget()
+	{
+		ColliderMesh = nullptr;
+	}
+
+	FColliderMeshProjectionTarget(FColliderMesh* ColliderMeshIn)
+	{
+		check(ColliderMeshIn);
+		ColliderMesh = ColliderMeshIn;
+	}
+
+	virtual ~FColliderMeshProjectionTarget() {}
+
+	virtual FVector3d Project(const FVector3d& Point, int Identifier = -1) override;
+	virtual FVector3d Project(const FVector3d& Point, FVector3d& ProjectNormalOut, int Identifier = -1) override;
+};
+
+
+
 } // end namespace UE::Geometry
 } // end namespace UE
 
diff --git a/Engine/Source/Runtime/GeometryCore/Public/Generators/MarchingCubes.h b/Engine/Source/Runtime/GeometryCore/Public/Generators/MarchingCubes.h
--- a/Engine/Source/Runtime/GeometryCore/Public/Generators/MarchingCubes.h
+++ b/Engine/Source/Runtime/GeometryCore/Public/Generators/MarchingCubes.h
@@ -12,6 +12,7 @@
 #include "Misc/AssertionMacros.h"
 
 #include "Spatial/DenseGrid3.h"
+#include "Spatial/BlockedDenseGrid3.h"
 
 #include "CompGeom/PolygonTriangulation.h"
 
@@ -129,9 +130,8 @@ public:
 
 		if (bEnableValueCaching)
 		{
-			corner_values_grid = FDenseGrid3f(CellDimensions.X + 1, CellDimensions.Y + 1, CellDimensions.Z + 1, FMathf::MaxReal);
+			BlockedCornerValuesGrid = FBlockedDenseGrid3f(CellDimensions.X + 1, CellDimensions.Y + 1, CellDimensions.Z + 1, FMathf::MaxReal);
 		}
-		corner_values.Reset();
 		InitHashTables();
 		ResetMesh();
 
@@ -165,28 +165,27 @@ public:
 
 		InitHashTables();
 		ResetMesh();
-		corner_values.Reset();
 
 		if (LastGridBounds != GridBounds)
 		{
 			if (bEnableValueCaching)
 			{
-				corner_values_grid = FDenseGrid3f(CellDimensions.X + 1, CellDimensions.Y + 1, CellDimensions.Z + 1, FMathf::MaxReal);
+				BlockedCornerValuesGrid = FBlockedDenseGrid3f(CellDimensions.X + 1, CellDimensions.Y + 1, CellDimensions.Z + 1, FMathf::MaxReal);
 			}
 			if (bParallelCompute)
 			{
-				done_cells = FDenseGrid3i(CellDimensions.X, CellDimensions.Y, CellDimensions.Z, 0);
+				BlockedDoneCells = FBlockedDenseGrid3i(CellDimensions.X, CellDimensions.Y, CellDimensions.Z, 0);
 			}
 		}
 		else
 		{
 			if (bEnableValueCaching)
 			{
-				corner_values_grid.Assign(FMathf::MaxReal);
+				BlockedCornerValuesGrid.Resize(CellDimensions.X + 1, CellDimensions.Y + 1, CellDimensions.Z + 1);
 			}
 			if (bParallelCompute)
 			{
-				done_cells.Assign(0);
+				BlockedDoneCells.Resize(CellDimensions.X, CellDimensions.Y, CellDimensions.Z);
 			}
 		}
 
@@ -363,51 +362,6 @@ protected:
 
 
 
-	//
-	// Store corner values in hash table. This doesn't make
-	// sense if we are evaluating entire grid, way too slow.
-	//
-
-	TMap<int64, double> corner_values;
-	FCriticalSection corner_values_lock;
-
-	double corner_value(const FVector3i& Idx)
-	{
-		int64 hash = corner_hash(Idx);
-		double value = 0;
-		double* findval = corner_values.Find(hash);
-		if (findval == nullptr)
-		{
-			TVector<double> V = corner_pos(Idx);
-			value = Implicit(V);
-			corner_values.Add(hash, value);
-		} 
-		else
-		{
-			value = *findval;
-		}
-		return value;
-	}
-	void initialize_cell_values(FGridCell& Cell, bool Shift)
-	{
-		FScopeLock Lock(&corner_values_lock);
-
-		if ( Shift )
-		{
-			Cell.f[1] = corner_value(Cell.i[1]);
-			Cell.f[2] = corner_value(Cell.i[2]);
-			Cell.f[5] = corner_value(Cell.i[5]);
-			Cell.f[6] = corner_value(Cell.i[6]);
-		}
-		else
-		{
-			for (int i = 0; i < 8; ++i)
-			{
-				Cell.f[i] = corner_value(Cell.i[i]);
-			}
-		}
-	}
-
 
 
 	//
@@ -416,37 +370,30 @@ protected:
 	// (note this is float grid, not double...)
 	//
 
-	FDenseGrid3f corner_values_grid;
-
-	int64 NumCornerValueSections = 64;
-	TArray<FCriticalSection> CornerValueSectionLocks;
-	int32 GetCornerValueSectionIndex(int64 hash)
-	{
-		return (int32)(hash % (NumCornerValueSections - 1));
-	}
+	FBlockedDenseGrid3f BlockedCornerValuesGrid;
 
 	double corner_value_grid_parallel(const FVector3i& Idx)
 	{
-		int32 SectionIndex = GetCornerValueSectionIndex(corner_hash(Idx));
-		float val = 0.0;
+		// note: it's possible to have a race here, where multiple threads might both
+		// GetValue, see that the value is invalid, and compute and set it. Since Implicit(V)
+		// is (intended to be) determinstic, they will compute the same value, so this doesn't cause an error, 
+		// it just wastes a bit of computation time. Since it is common for multiple corners to be
+		// in the same grid-block, and locking is on the block level, it is (or was in some testing)
+		// better to not lock the entire block while Implicit(V) computed, at the cost of
+		// some wasted evals in some cases.
+
+		float CurrentValue = BlockedCornerValuesGrid.GetValueThreadSafe(Idx.X, Idx.Y, Idx.Z);
+		if (CurrentValue != FMathf::MaxReal)
 		{
-			FScopeLock Lock(&CornerValueSectionLocks[SectionIndex]);
-			val = corner_values_grid[Idx];
-		}
-		if (val != FMathf::MaxReal)
-		{
-			return (double)val;
+			return (double)CurrentValue;
 		}
 
 		TVector<double> V = corner_pos(Idx);
-		val = (float)Implicit(V);
+		CurrentValue = (float)Implicit(V);
 
-		{
-			FScopeLock Lock(&CornerValueSectionLocks[SectionIndex]);
-			corner_values_grid[Idx] = val;
-		}
+		BlockedCornerValuesGrid.SetValueThreadSafe(Idx.X, Idx.Y, Idx.Z, CurrentValue);
 
-		return (double)val;
+		return (double)CurrentValue;
 	}
 	double corner_value_grid(const FVector3i& Idx)
 	{
@@ -455,16 +402,18 @@ protected:
 			return corner_value_grid_parallel(Idx);
 		}
 
-		float val = corner_values_grid[Idx];
-		if (val != FMathf::MaxReal)
+		float CurrentValue = BlockedCornerValuesGrid.GetValue(Idx.X, Idx.Y, Idx.Z);
+		if (CurrentValue != FMathf::MaxReal)
 		{
-			return (double)val;
+			return (double)CurrentValue;
 		}
 
 		TVector<double> V = corner_pos(Idx);
-		val = (float)Implicit(V);
-		corner_values_grid[Idx] = val;
-		return (double)val;
+		CurrentValue = (float)Implicit(V);
+
+		BlockedCornerValuesGrid.SetValue(Idx.X, Idx.Y, Idx.Z, CurrentValue);
+
+		return (double)CurrentValue;
 	}
 
 	void initialize_cell_values_grid(FGridCell& Cell, bool Shift)
@@ -531,7 +480,6 @@ protected:
 		Cell.i[6] = FVector3i(Idx.X + 1, Idx.Y + 1, Idx.Z + 1);
 		Cell.i[7] = FVector3i(Idx.X + 0, Idx.Y + 1, Idx.Z + 1);
 
-		//initialize_cell_values(Cell, false);
 		if (bEnableValueCaching)
 		{
 			initialize_cell_values_grid(Cell, false);
@@ -556,7 +504,6 @@ protected:
 		Cell.i[0].X = XIdx; Cell.i[1].X = XIdx+1; Cell.i[2].X = XIdx+1; Cell.i[3].X = XIdx;
 		Cell.i[4].X = XIdx; Cell.i[5].X = XIdx+1; Cell.i[6].X = XIdx+1; Cell.i[7].X = XIdx;
 
-		//initialize_cell_values(Cell, true);
 		if (bEnableValueCaching)
 		{
 			initialize_cell_values_grid(Cell, true);
@@ -574,12 +521,6 @@ protected:
 		EdgeVertexSections.SetNum(NumEdgeVertexSections);
 		EdgeVertexSectionLocks.Reset();
 		EdgeVertexSectionLocks.SetNum(NumEdgeVertexSections);
-
-		CornerValueSectionLocks.Reset();
-		CornerValueSectionLocks.SetNum(NumCornerValueSections);
-
-		DoneCellSectionLocks.Reset();
-		DoneCellSectionLocks.SetNum(NumDoneCellSections);
 	}
 
 
@@ -664,19 +605,19 @@ protected:
 		FGridCell Cell;
 		int vertTArray[12];
 
-		done_cells = FDenseGrid3i(CellDimensions.X, CellDimensions.Y, CellDimensions.Z, 0);
+		BlockedDoneCells = FBlockedDenseGrid3i(CellDimensions.X, CellDimensions.Y, CellDimensions.Z, 0);
 
 		TArray<FVector3i> stack;
 
 		for (FVector3d seed : Seeds)
 		{
 			FVector3i seed_idx = cell_index(seed);
-			if (!done_cells.IsValidIndex(seed_idx) || done_cells[seed_idx] == 1)
+			if (!BlockedDoneCells.IsValidIndex(seed_idx) || BlockedDoneCells.GetValue(seed_idx.X, seed_idx.Y, seed_idx.Z) == 1)
 			{
 				continue;
 			}
 			stack.Add(seed_idx);
-			done_cells[seed_idx] = 1;
+			BlockedDoneCells.SetValue(seed_idx.X, seed_idx.Y, seed_idx.Z, 1);
 
 			while ( stack.Num() > 0 )
 			{
@@ -693,10 +634,10 @@ protected:
 					for ( FVector3i o : IndexUtil::GridOffsets6 )
 					{
 						FVector3i nbr_idx = Idx + o;
-						if (GridBounds.Contains(nbr_idx) && done_cells[nbr_idx] == 0)
+						if (GridBounds.Contains(nbr_idx) && BlockedDoneCells.GetValue(nbr_idx.X, nbr_idx.Y, nbr_idx.Z) == 0)
 						{
 							stack.Add(nbr_idx);
-							done_cells[nbr_idx] = 1;
+							BlockedDoneCells.SetValue(nbr_idx.X, nbr_idx.Y, nbr_idx.Z, 1);
 						}
 					}
 				}
@@ -718,7 +659,7 @@ protected:
 		{
 			FVector3d Seed = Seeds[Index];
 			FVector3i seed_idx = cell_index(Seed);
-			if (!done_cells.IsValidIndex(seed_idx) || set_cell_if_not_done(seed_idx) == false)
+			if (!BlockedDoneCells.IsValidIndex(seed_idx) || set_cell_if_not_done(seed_idx) == false)
 			{
 				return;
 			}
@@ -760,27 +701,20 @@ protected:
 	}
 
 
-
-	FDenseGrid3i done_cells;
-
-	int64 NumDoneCellSections = 64;
-	TArray<FCriticalSection> DoneCellSectionLocks;
-	int GetDoneCellSectionIndex(int64 hash)
-	{
-		return (int32)(hash % (NumDoneCellSections - 1));
-	}
+	FBlockedDenseGrid3i BlockedDoneCells;
 
 	bool set_cell_if_not_done(const FVector3i& Idx)
 	{
-		int32 SectionIndex = GetDoneCellSectionIndex(corner_hash(Idx));
 		bool was_set = false;
 		{
-			FScopeLock Lock(&DoneCellSectionLocks[SectionIndex]);
-			if (done_cells[Idx] == 0)
+			BlockedDoneCells.ProcessValueThreadSafe(Idx.X, Idx.Y, Idx.Z, [&](int& CellValue) 
 			{
-				done_cells[Idx] = 1;
-				was_set = true;
-			}
+				if (CellValue == 0)
+				{
+					CellValue = 1;
+					was_set = true;
+				}
+			});
 		}
 		return was_set;
 	}
diff --git a/Engine/Source/Runtime/GeometryCore/Public/Spatial/BlockedDenseGrid3.h b/Engine/Source/Runtime/GeometryCore/Public/Spatial/BlockedDenseGrid3.h
new file mode 100644
--- /dev/null
+++ b/Engine/Source/Runtime/GeometryCore/Public/Spatial/BlockedDenseGrid3.h
@@ -0,0 +1,263 @@
+// Copyright Epic Games, Inc. All Rights Reserved.
+
+// Port of geometry3Sharp BlockedDenseGrid3
+
+#pragma once
+
+#include "CoreMinimal.h"
+#include "BoxTypes.h"
+#include "IntBoxTypes.h"
+#include "Containers/StaticArray.h"
+
+namespace UE
+{
+namespace Geometry
+{
+
+/**
+ * TBlockedDenseGrid represents a dense 3D uniform grid, but the grid is allocated in BlockSize^3 blocks on-demand
+ * (BlockSize is a compile-time constant). This allows very large grids to be used without
+ * having to pre-allocate all the memory, eg for sparse/narrow-band use cases.
+ * 
+ * Thread-safe access functions are provided, these are implemented by internal locking on the grid-block level.
+ */
+template<typename ElemType>
+class TBlockedDenseGrid3
+{
+protected:
+
+	static const int32 BlockSize = 32;
+	using BlockType = TStaticArray<ElemType, BlockSize*BlockSize*BlockSize>;
+
+	TArray<TUniquePtr<BlockType>> Blocks;
+	TArray<FCriticalSection> BlockLocks;
+
+	/** dimensions per axis */
+	FVector3i BlockDimensions;
+
+	/** dimensions per axis */
+	FVector3i Dimensions;
+
+	/** The value to use when no value is known - ie for Get() queries on unallocated cells, or to initialize new blocks */
+	ElemType ConstantValue = (ElemType)0;
+
+protected:
+
+	void Init(BlockType& Block)
+	{
+		for (int32 k = 0; k < BlockSize * BlockSize * BlockSize; ++k)
+		{
+			Block[k] = ConstantValue;
+		}
+	}
+
+	void GetBlockCoords(int32 X, int32 Y, int32 Z, int32& BlockIndexOut, int32& LocalIndexOut) const
+	{
+		int32 BlockX = X / BlockSize;
+		int32 BlockY = Y / BlockSize;
+		int32 BlockZ = Z / BlockSize;
+		BlockIndexOut = BlockX + BlockDimensions.X * (BlockY + BlockDimensions.Y * BlockZ);
+
+		int32 LocalX = X % BlockSize;
+		int32 LocalY = Y % BlockSize;
+		int32 LocalZ = Z % BlockSize;
+		LocalIndexOut = LocalX + BlockSize * (LocalY + BlockSize * LocalZ);
+	}
+
+#if UE_BUILD_DEBUG
+	TUniquePtr<BlockType>& GetBlock(int32 Index) { return Blocks[Index]; }
+	const TUniquePtr<BlockType>& GetBlock(int32 Index) const { return Blocks[Index]; }
+	FCriticalSection* GetBlockLock(int32 Index) { return &BlockLocks[Index]; }
+#else
+	// skip range checks in non-debug builds
+	TUniquePtr<BlockType>& GetBlock(int32 Index) { return Blocks.GetData()[Index]; }
+	const TUniquePtr<BlockType>& GetBlock(int32 Index) const { return Blocks.GetData()[Index]; }
+	FCriticalSection* GetBlockLock(int32 Index) { return &BlockLocks.GetData()[Index]; }
+#endif
+
+	template<typename FuncType>
+	void WriteValueThreadSafe(int32 X, int32 Y, int32 Z, FuncType Func)
+	{
+		int32 BlockIndex, LocalIndex;
+		GetBlockCoords(X, Y, Z, BlockIndex, LocalIndex);
+
+		FScopeLock Lock(GetBlockLock(BlockIndex));
+
+		if (GetBlock(BlockIndex).IsValid() == false)
+		{
+			Blocks[BlockIndex] = MakeUnique<BlockType>();
+			Init(*Blocks[BlockIndex]);
+		}
+		TUniquePtr<BlockType>& Block = GetBlock(BlockIndex);
+		ElemType& Value = (*Block)[LocalIndex];
+		Func(Value);
+	}
+
+
+	template<typename FuncType>
+	void WriteValue(int32 X, int32 Y, int32 Z, FuncType Func)
+	{
+		int32 BlockIndex, LocalIndex;
+		GetBlockCoords(X, Y, Z, BlockIndex, LocalIndex);
+
+		if (GetBlock(BlockIndex).IsValid() == false)
+		{
+			Blocks[BlockIndex] = MakeUnique<BlockType>();
+			Init(*Blocks[BlockIndex]);
+		}
+		TUniquePtr<BlockType>& Block = GetBlock(BlockIndex);
+		ElemType& Value = (*Block)[LocalIndex];
+		Func(Value);
+	}
+
+
+	ElemType ReadValueThreadSafe(int32 X, int32 Y, int32 Z)
+	{
+		int32 BlockIndex, LocalIndex;
+		GetBlockCoords(X, Y, Z, BlockIndex, LocalIndex);
+
+		FScopeLock Lock(GetBlockLock(BlockIndex));
+		const TUniquePtr<BlockType>& Block = GetBlock(BlockIndex);
+		return (Block.IsValid()) ? (*Block)[LocalIndex] : ConstantValue;
+	}
+
+	ElemType ReadValue(int32 X, int32 Y, int32 Z) const
+	{
+		int32 BlockIndex, LocalIndex;
+		GetBlockCoords(X, Y, Z, BlockIndex, LocalIndex);
+
+		const TUniquePtr<BlockType>& Block = GetBlock(BlockIndex);
+		return (Block.IsValid()) ? (*Block)[LocalIndex] : ConstantValue;
+	}
+
+
+public:
+	/**
+	 * Create empty grid
+	 */
+	TBlockedDenseGrid3() : Dimensions(0,0,0)
+	{
+	}
+
+	TBlockedDenseGrid3(int32 DimX, int32 DimY, int32 DimZ, ElemType InitialValue)
+	{
+		Resize(DimX, DimY, DimZ);
+		ConstantValue = InitialValue;
+	}
+
+	int64 Size() const
+	{
+		return (int64)Dimensions.X * (int64)Dimensions.Y * (int64)Dimensions.Z;
+	}
+
+	bool IsValidIndex(const FVector3i& Index) const
+	{
+		return Index.X >= 0 && Index.Y >= 0 && Index.Z >= 0
+			&& Index.X < Dimensions.X && Index.Y < Dimensions.Y && Index.Z < Dimensions.Z;
+	}
+
+	const FVector3i& GetDimensions() const
+	{
+		return Dimensions;
+	}
+
+	/**
+	 * Reconfigure the grid to have the target dimensions. This clears all the 
+	 * existing grid memory.
+	 */
+	void Resize(int32 DimX, int32 DimY, int32 DimZ)
+	{
+		check((int64)DimX * (int64)DimY * (int64)DimZ < INT_MAX);
+		int32 BlocksX = (DimX / BlockSize) + 1;
+		int32 BlocksY = (DimY / BlockSize) + 1;
+		int32 BlocksZ = (DimZ / BlockSize) + 1;
+
+		Blocks.Reset();
+		int32 NumBlocks = BlocksX * BlocksY * BlocksZ;
+		Blocks.SetNum(NumBlocks);
+		BlockLocks.SetNum(NumBlocks);
+		for (int32 k = 0; k < NumBlocks; ++k)
+		{
+			Blocks[k] = TUniquePtr<BlockType>();
+		}
+
+		BlockDimensions = FVector3i(BlocksX, BlocksY, BlocksZ);
+		Dimensions = FVector3i(DimX, DimY, DimZ);
+	}
+
+	/**
+	 * @return the grid value at (X,Y,Z)
+	 */
+	ElemType GetValue(int32 X, int32 Y, int32 Z) const
+	{
+		return ReadValue(X, Y, Z);
+	}
+
+	/**
+	 * @return the grid value at (X,Y,Z), with internal locking, so it is safe to call this from multiple read & write threads
+	 */
+	ElemType GetValueThreadSafe(int32 X0, int32 Y, int32 Z)
+	{
+		return ReadValueThreadSafe(X0, Y, Z);
+	}
+
+	/**
+	 * Set the grid value at (X,Y,Z)
+	 */
+	void SetValue(int32 X, int32 Y, int32 Z, ElemType NewValue)
+	{
+		WriteValue(X, Y, Z, [NewValue](ElemType& GridValueInOut) { GridValueInOut = NewValue; });
+	}
+
+	/**
+	 * Set the grid value at (X,Y,Z), with internal locking, so it is safe to call this from multiple read & write threads
+	 */
+	void SetValueThreadSafe(int32 X, int32 Y, int32 Z, ElemType NewValue)
+	{
+		WriteValueThreadSafe(X, Y, Z, [NewValue](ElemType& GridValueInOut) { GridValueInOut = NewValue; });
+	}
+
+	/**
+	 * Call an external lambda with a reference to the grid value at (X,Y,Z).
+	 * Called as Func(ElemType&), so the caller can both read and write the grid cell
+	 */
+	template<typename ProcessFunc>
+	void ProcessValueThreadSafe(int32 X, int32 Y, int32 Z, ProcessFunc Func)
+	{
+		WriteValueThreadSafe(X, Y, Z, Func);
+	}
+
+	FAxisAlignedBox3i Bounds() const
+	{
+		return FAxisAlignedBox3i({0, 0, 0},{Dimensions.X, Dimensions.Y, Dimensions.Z});
+	}
+	FAxisAlignedBox3i BoundsInclusive() const
+	{
+		return FAxisAlignedBox3i({0, 0, 0},{Dimensions.X-1, Dimensions.Y-1, Dimensions.Z-1});
+	}
+
+	FVector3i ToIndex(int64 LinearIndex) const
+	{
+		int32 x = (int32)(LinearIndex % (int64)Dimensions.X);
+		int32 y = (int32)((LinearIndex / (int64)Dimensions.X) % (int64)Dimensions.Y);
+		int32 z = (int32)(LinearIndex / ((int64)Dimensions.X * (int64)Dimensions.Y));
+		return FVector3i(x, y, z);
+	}
+	int64 ToLinear(int32 X, int32 Y, int32 Z) const
+	{
+		return (int64)X + (int64)Dimensions.X * ((int64)Y + (int64)Dimensions.Y * (int64)Z);
+	}
+	int64 ToLinear(const FVector3i& IJK) const
+	{
+		return (int64)IJK.X + (int64)Dimensions.X * ((int64)IJK.Y + (int64)Dimensions.Y * (int64)IJK.Z);
+	}
+};
+
+
+typedef TBlockedDenseGrid3<float> FBlockedDenseGrid3f;
+typedef TBlockedDenseGrid3<double> FBlockedDenseGrid3d;
+typedef TBlockedDenseGrid3<int> FBlockedDenseGrid3i;
+
+
+} // end namespace UE::Geometry
+} // end namespace UE
diff --git a/Engine/Source/Runtime/GeometryCore/Public/Spatial/DenseGrid3.h b/Engine/Source/Runtime/GeometryCore/Public/Spatial/DenseGrid3.h
--- a/Engine/Source/Runtime/GeometryCore/Public/Spatial/DenseGrid3.h
+++ b/Engine/Source/Runtime/GeometryCore/Public/Spatial/DenseGrid3.h
@@ -132,46 +132,6 @@ public:
 		}
 	}
 
-	// TODO: implement TDenseGrid2<ElemType> and then implement this
-	//TDenseGrid2<ElemType> get_slice(int slice_i, int dimension)
-	//{
-	//	TDenseGrid2<ElemType> slice;
-	//	if (dimension == 0) {
-	//		slice = TDenseGrid2<ElemType>(Dimensions.Y, Dimensions.Z, 0);
-	//		for (int k = 0; k < Dimensions.Z; ++k)
-	//			for (int j = 0; j < Dimensions.Y; ++j)
-	//				slice[j, k] = Buffer[slice_i + Dimensions.X * (j + Dimensions.Y * k)];
-	//	} else if (dimension == 1) {
-	//		slice = TDenseGrid2<ElemType>(Dimensions.X, Dimensions.Z, 0);
-	//		for (int k = 0; k < Dimensions.Z; ++k)
-	//			for (int i = 0; i < Dimensions.X; ++i)
-	//				slice[i, k] = Buffer[i + Dimensions.X * (slice_i + Dimensions.Y * k)];
-	//	} else {
-	//		slice = TDenseGrid2<ElemType>(Dimensions.X, Dimensions.Y, 0);
-	//		for (int j = 0; j < Dimensions.Y; ++j)
-	//			for (int i = 0; i < Dimensions.X; ++i)
-	//				slice[i, j] = Buffer[i + Dimensions.X * (j + Dimensions.Y * slice_i)];
-	//	}
-	//	return slice;
-	//}
-	//void set_slice(TDenseGrid2<ElemType> slice, int slice_i, int dimension)
-	//{
-	//	if (dimension == 0) {
-	//		for (int k = 0; k < Dimensions.Z; ++k)
-	//			for (int j = 0; j < Dimensions.Y; ++j)
-	//				Buffer[slice_i + Dimensions.X * (j + Dimensions.Y * k)] = slice[j, k];
-	//	} else if (dimension == 1) {
-	//		for (int k = 0; k < Dimensions.Z; ++k)
-	//			for (int i = 0; i < Dimensions.X; ++i)
-	//				Buffer[i + Dimensions.X * (slice_i + Dimensions.Y * k)] = slice[i, k];
-	//	} else {
-	//		for (int j = 0; j < Dimensions.Y; ++j)
-	//			for (int i = 0; i < Dimensions.X; ++i)
-	//				Buffer[i + Dimensions.X * (j + Dimensions.Y * slice_i)] = slice[i, j];
-	//	}
-	//}
-
-
 	FAxisAlignedBox3i Bounds() const
 	{
 		return FAxisAlignedBox3i({0, 0, 0},{Dimensions.X, Dimensions.Y, Dimensions.Z});
