Add missing include
#rnx

diff --git a/Engine/Source/Runtime/GeometryCore/Public/DynamicMesh/MeshTransforms.h b/Engine/Source/Runtime/GeometryCore/Public/DynamicMesh/MeshTransforms.h
--- a/Engine/Source/Runtime/GeometryCore/Public/DynamicMesh/MeshTransforms.h
+++ b/Engine/Source/Runtime/GeometryCore/Public/DynamicMesh/MeshTransforms.h
@@ -5,6 +5,7 @@
 #pragma once
 
 #include "DynamicMesh/DynamicMesh3.h"
+#include "TransformTypes.h"
 
 
 /**
