Change inequality direction in a check. VertexBoneWeights should have size equal to the number of vertices in the parent mesh *after* compacting, so this will be less than or equal to the size of the compact map, not geq. Includes a unit test for verification.

#rnx

diff --git a/Engine/Source/Runtime/GeometryCore/Public/DynamicMesh/DynamicVertexSkinWeightsAttribute.h b/Engine/Source/Runtime/GeometryCore/Public/DynamicMesh/DynamicVertexSkinWeightsAttribute.h
--- a/Engine/Source/Runtime/GeometryCore/Public/DynamicMesh/DynamicVertexSkinWeightsAttribute.h
+++ b/Engine/Source/Runtime/GeometryCore/Public/DynamicMesh/DynamicVertexSkinWeightsAttribute.h
@@ -131,7 +131,7 @@ public:
 	void CompactCopy(const FCompactMaps& CompactMaps, const TDynamicVertexSkinWeightsAttribute<ParentType>& ToCopy)
 	{
 		TDynamicAttributeBase<ParentType>::CopyParentClassData(ToCopy);
-		check(CompactMaps.NumVertexMappings() <= VertexBoneWeights.Num());
+		check(CompactMaps.NumVertexMappings() >= VertexBoneWeights.Num());
 		FBoneWeights Data;
 		for (int32 VID = 0, NumVID = CompactMaps.NumVertexMappings(); VID < NumVID; VID++)
 		{
