GeometryCore: Fix legacy serialization of TDynamicVector<> for LWC data types using double and bulk serialization

diff --git a/Engine/Source/Runtime/GeometryCore/Public/Util/DynamicVector.h b/Engine/Source/Runtime/GeometryCore/Public/Util/DynamicVector.h
--- a/Engine/Source/Runtime/GeometryCore/Public/Util/DynamicVector.h
+++ b/Engine/Source/Runtime/GeometryCore/Public/Util/DynamicVector.h
@@ -157,7 +157,7 @@ public:
 			Ar.CountBytes(CountBytes, CountBytes);
 			Ar << CurBlock;
 			Ar << CurBlockUsed;
-			Blocks.Serialize_Legacy(Ar);
+			Blocks.Serialize_LegacyLoad(Ar);
 
 			// Account for a previously existing issue where one additional empty block was allocated and the resulting values for the current block can cause
 			// crashes, for example when using Back().
@@ -393,19 +393,31 @@ private:
 		}
 
 		/** Serialize TBlockVector to archive. */
-		void Serialize_Legacy(FArchive& Ar)
-		{
-			const auto SerializeElement = [](FArchive& Ar, ArrayType* Element)
-			{
-				if (TCanBulkSerialize<Type>::Value)
-				{
-					Ar.Serialize(Element->GetData(), Element->Num() * sizeof(Type));
-				}
-				else
-				{
-					Ar << *Element;
-				}
-			};
+		void Serialize_LegacyLoad(FArchive& Ar)
+		{
+			// Only loading of legacy archives is allowed.
+			checkSlow(Ar.IsLoading());
+
+			// Bulk serialization for a number of double types was enabled as part of the transition to Large World Coordinates.
+			// If the currently stored type is one of these types, and the archive is from before bulk serialization for these types was enabled,
+			// we need to still use per element serialization for legacy data.
+			constexpr bool bIsLWCBulkSerializedDoubleType =
+				TIsSame<Type, FVector2d>::Value ||
+				TIsSame<Type, FVector3d>::Value ||
+				TIsSame<Type, FVector4d>::Value ||
+				TIsSame<Type, FQuat4d>::Value ||
+				TIsSame<Type, FTransform3d>::Value;
+
+			const auto SerializeElement =
+				TCanBulkSerialize<Type>::Value && !(bIsLWCBulkSerializedDoubleType && Ar.UEVer() < EUnrealEngineObjectUE5Version::LARGE_WORLD_COORDINATES)
+					? [](FArchive& Archive, ArrayType* Element)
+					{
+						Archive.Serialize(Element->GetData(), Element->Num() * sizeof(Type));
+					}
+					: [](FArchive& Archive, ArrayType* Element)
+					{
+						Archive << *Element;
+					};
 
 			int32 BlockNum = Num();
 			Ar << BlockNum;
@@ -420,14 +432,6 @@ private:
 					Add(NewElement);
 				}
 			}
-			else
-			{
-				// Save array.
-				for (int32 Index = 0; Index < BlockNum; Index++)
-				{
-					SerializeElement(Ar, &(*this)[Index]);
-				}
-			}
 		}
 
 		template <bool bForceBulkSerialization, bool bUseCompression>
