GeometryCore: Fix compilation error in TDynamicVector

diff --git a/Engine/Source/Runtime/GeometryCore/Public/Util/DynamicVector.h b/Engine/Source/Runtime/GeometryCore/Public/Util/DynamicVector.h
--- a/Engine/Source/Runtime/GeometryCore/Public/Util/DynamicVector.h
+++ b/Engine/Source/Runtime/GeometryCore/Public/Util/DynamicVector.h
@@ -397,34 +397,36 @@ private:
 		{
 			// Only loading of legacy archives is allowed.
 			checkSlow(Ar.IsLoading());
-
-			// Bulk serialization for a number of double types was enabled as part of the transition to Large World Coordinates.
-			// If the currently stored type is one of these types, and the archive is from before bulk serialization for these types was enabled,
-			// we need to still use per element serialization for legacy data.
-			constexpr bool bIsLWCBulkSerializedDoubleType =
-				TIsSame<Type, FVector2d>::Value ||
-				TIsSame<Type, FVector3d>::Value ||
-				TIsSame<Type, FVector4d>::Value ||
-				TIsSame<Type, FQuat4d>::Value ||
-				TIsSame<Type, FTransform3d>::Value;
-
-			const auto SerializeElement =
-				TCanBulkSerialize<Type>::Value && !(bIsLWCBulkSerializedDoubleType && Ar.UEVer() < EUnrealEngineObjectUE5Version::LARGE_WORLD_COORDINATES)
-					? [](FArchive& Archive, ArrayType* Element)
-					{
-						Archive.Serialize(Element->GetData(), Element->Num() * sizeof(Type));
-					}
-					: [](FArchive& Archive, ArrayType* Element)
-					{
-						Archive << *Element;
-					};
-
-			int32 BlockNum = Num();
-			Ar << BlockNum;
 			if (Ar.IsLoading())
 			{
-				// Load array.
+				// Bulk serialization for a number of double types was enabled as part of the transition to Large World Coordinates.
+				// If the currently stored type is one of these types, and the archive is from before bulk serialization for these types was enabled,
+				// we need to still use per element serialization for legacy data.
+				constexpr bool bIsLWCBulkSerializedDoubleType =
+					TIsSame<Type, FVector2d>::Value ||
+					TIsSame<Type, FVector3d>::Value ||
+					TIsSame<Type, FVector4d>::Value ||
+					TIsSame<Type, FQuat4d>::Value ||
+					TIsSame<Type, FTransform3d>::Value;
+				const bool bUseBulkSerialization = TCanBulkSerialize<Type>::Value && !(bIsLWCBulkSerializedDoubleType && Ar.UEVer() <
+					EUnrealEngineObjectUE5Version::LARGE_WORLD_COORDINATES);
+
+				// Lambda for serializing a block element either via bulk serializing the contained data or via serializing the element container itself.
+				// Note that the static_cast<> was necessary to resolve compiler errors when using MSVC.
+				const auto SerializeElement = bUseBulkSerialization
+					                              ? static_cast<void(*)(FArchive&, ArrayType*)>([](FArchive& Archive, ArrayType* Element)
+					                              {
+						                              Archive.Serialize(Element->GetData(), Element->Num() * sizeof(Type));
+					                              })
+					                              : static_cast<void(*)(FArchive&, ArrayType*)>([](FArchive& Archive, ArrayType* Element)
+					                              {
+						                              Archive << *Element;
+					                              });
+
+				int32 BlockNum = Num();
+				Ar << BlockNum;
 				Empty(BlockNum);
+
 				for (int32 Index = 0; Index < BlockNum; Index++)
 				{
 					ArrayType* NewElement = new ArrayType;
