GeometryCore: do not hold onto intermediate grid in two-pass computations in TImplicitMorphology. The first-pass grid is not used in the second pass, and having both grids allocated simultaneously results in significantly more peak memory allocated during large-grid computations (possibly resulting in OOMs)
#rnx

diff --git a/Engine/Source/Runtime/GeometryCore/Public/Implicit/Morphology.h b/Engine/Source/Runtime/GeometryCore/Public/Implicit/Morphology.h
--- a/Engine/Source/Runtime/GeometryCore/Public/Implicit/Morphology.h
+++ b/Engine/Source/Runtime/GeometryCore/Public/Implicit/Morphology.h
@@ -90,8 +90,8 @@ protected:
 	// Stores result (returned as a const FMeshShapeGenerator)
 	FMarchingCubes MarchingCubes;
 
-	/// Intermediate -- result of first pass, could potentially be re-used if recomputed w/ equal or smaller Distance
-	TSweepingMeshSDF<TriangleMeshType> ComputedSDF;
+	// computed in first pass, re-used in second
+	double NarrowBandMaxDistance;
 
 public:
 	bool Validate()
@@ -135,6 +135,7 @@ protected:
 	{
 		TRACE_CPUPROFILER_EVENT_SCOPE(Geometry_Morphology_FirstPass);
 
+		TSweepingMeshSDF<TriangleMeshType> ComputedSDF;
 		ComputedSDF.Mesh = Source;
 
 		ComputedSDF.Spatial = SourceSpatial;
@@ -142,7 +143,8 @@ protected:
 
 
 		ComputedSDF.CellSize = GridCellSize;
-		ComputedSDF.NarrowBandMaxDistance = UnsignedOffset + ComputedSDF.CellSize;
+		NarrowBandMaxDistance = UnsignedOffset + ComputedSDF.CellSize;
+		ComputedSDF.NarrowBandMaxDistance = NarrowBandMaxDistance;
 		ComputedSDF.ExactBandWidth = FMath::CeilToInt(ComputedSDF.NarrowBandMaxDistance / ComputedSDF.CellSize);
 
 		// for meshes with long triangles relative to the width of the narrow band, don't use the AABB tree
@@ -247,7 +249,7 @@ protected:
 		MarchingCubes.Bounds.Expand(GridCellSize);
 		if (MarchingCubes.IsoValue < 0)
 		{
-			MarchingCubes.Bounds.Expand(ComputedSDF.NarrowBandMaxDistance);
+			MarchingCubes.Bounds.Expand(NarrowBandMaxDistance);
 		}
 
 
