Fix for mismatched deprecation macros.

diff --git a/Engine/Source/Runtime/GeometryCore/Public/IntVectorTypes.h b/Engine/Source/Runtime/GeometryCore/Public/IntVectorTypes.h
--- a/Engine/Source/Runtime/GeometryCore/Public/IntVectorTypes.h
+++ b/Engine/Source/Runtime/GeometryCore/Public/IntVectorTypes.h
@@ -46,7 +46,7 @@ struct FVector2i
 PRAGMA_DISABLE_DEPRECATION_WARNINGS
 	constexpr FVector2i(const FVector2i& Vec) = default;
 	constexpr FVector2i& operator=(const FVector2i& V2) = default;
-//PRAGMA_ENABLE_DEPRECATION_WARNINGS
+PRAGMA_ENABLE_DEPRECATION_WARNINGS
 
 	explicit constexpr operator const int32*() const
 	{
