GeometryCore: improve thread utililization in parallel continuation Marching Cubes
#rnx

diff --git a/Engine/Source/Runtime/GeometryCore/Public/Generators/MarchingCubes.h b/Engine/Source/Runtime/GeometryCore/Public/Generators/MarchingCubes.h
--- a/Engine/Source/Runtime/GeometryCore/Public/Generators/MarchingCubes.h
+++ b/Engine/Source/Runtime/GeometryCore/Public/Generators/MarchingCubes.h
@@ -653,49 +653,101 @@ protected:
 	*/
 	void generate_continuation_parallel(TArrayView<const FVector3d> Seeds)
 	{
+		// Parallel marching cubes based on continuation (ie surface-following / front propagation) 
+		// can have quite poor multithreaded performance depending on the ordering of the region-growing.
+		// For example processing each seed point in parallel can result in one thread that
+		// takes significantly longer than others, if the seed point distribution is such
+		// that a large part of the surface is only reachable from one seed (or gets
+		// "cut off" by a thin area, etc). So we want to basically do front-marching in
+		// parallel passes. However this can result in a large number of very short passes
+		// if the front ends up with many small regions, etc. So, in the implementation below,
+		// each "seed cell" is allowed to process up to N neighbour cells before terminating, at 
+		// which point any cells remaining on the active front are added as seed cells for the next pass.
+		// This seems to provide good utilization, however more profiling may be needed.
+		// (In particular, if the active cell list is large, some blocks of the ParallelFor
+		//  may still end up doing much more work than others)
+
+
+		// set this flag so that append vertex/triangle operations will lock the mesh
 		parallel_mesh_access = true;
 
-		ParallelFor(Seeds.Num(), [this, &Seeds](int32 Index)
+		// maximum number of cells to process in each ParallelFor iteration
+		static constexpr int MaxNeighboursPerActiveCell = 100;
+
+		// list of active cells on the MC front to process in the next pass
+		TArray<FVector3i> ActiveCells;
+
+		// initially push list of seed-point cells onto the ActiveCells list
+		for (FVector3d Seed : Seeds)
 		{
-			FVector3d Seed = Seeds[Index];
 			FVector3i seed_idx = cell_index(Seed);
-			if (!BlockedDoneCells.IsValidIndex(seed_idx) || set_cell_if_not_done(seed_idx) == false)
+			if ( BlockedDoneCells.IsValidIndex(seed_idx) && set_cell_if_not_done(seed_idx) )
 			{
-				return;
+				ActiveCells.Add(seed_idx);
 			}
+		}
 
-			FGridCell Cell;
-			int vertTArray[12];
-
-			TArray<FVector3i> stack;
-			stack.Add(seed_idx);
+		// new active cells will be accumulated in each parallel-pass
+		TArray<FVector3i> NewActiveCells;
+		FCriticalSection NewActiveCellsLock;
 
-			while (stack.Num() > 0)
+		while (ActiveCells.Num() > 0)
+		{
+			// process all active cells
+			ParallelFor(ActiveCells.Num(), [&](int32 Idx)
 			{
-				FVector3i Idx = stack[stack.Num() - 1];
-				stack.RemoveAt(stack.Num() - 1);
+				FVector3i InitialCellIndex = ActiveCells[Idx];
 				if (CancelF())
 				{
 					return;
 				}
 
-				initialize_cell(Cell, Idx);
-				if (polygonize_cell(Cell, vertTArray))
-				{     // found crossing
-					for (FVector3i o : IndexUtil::GridOffsets6)
+				FGridCell TempCell;
+				int TempArray[12];
+
+				// we will process up to MaxNeighboursPerActiveCell new cells in each ParallelFor iteration
+				TArray<FVector3i, TInlineAllocator<64>> LocalStack;
+				LocalStack.Add(InitialCellIndex);
+				int32 CellsProcessed = 0;
+
+				while (LocalStack.Num() > 0 && CellsProcessed++ < MaxNeighboursPerActiveCell)
+				{
+					FVector3i CellIndex = LocalStack.Pop(false);
+
+					initialize_cell(TempCell, CellIndex);
+					if (polygonize_cell(TempCell, TempArray))
 					{
-						FVector3i nbr_idx = Idx + o;
-						if (GridBounds.Contains(nbr_idx))
+						// found crossing
+						for (FVector3i GridOffset : IndexUtil::GridOffsets6)
 						{
-							if (set_cell_if_not_done(nbr_idx) == true)
-							{ 
-								stack.Add(nbr_idx);
+							FVector3i NbrCellIndex = CellIndex + GridOffset;
+							if (GridBounds.Contains(NbrCellIndex))
+							{
+								if (set_cell_if_not_done(NbrCellIndex) == true)
+								{ 
+									LocalStack.Add(NbrCellIndex);
+								}
 							}
 						}
 					}
 				}
+
+				// if stack is not empty, ie hit MaxNeighboursPerActiveCell, add remaining cells to next-pass Active list
+				if (LocalStack.Num() > 0)
+				{
+					NewActiveCellsLock.Lock();
+					NewActiveCells.Append(LocalStack);
+					NewActiveCellsLock.Unlock();
+				}
+
+			});
+
+			ActiveCells.Reset();
+			if (NewActiveCells.Num() > 0)
+			{
+				Swap(ActiveCells, NewActiveCells);
 			}
-		});
+		}
 
 		parallel_mesh_access = false;
 	}
@@ -723,15 +775,10 @@ protected:
 
 
 
-
-
-
-
-
 	/**
 	*  find edge crossings and generate triangles for this cell
 	*/
-	bool polygonize_cell(FGridCell& Cell, int VertIndexArray[8])
+	bool polygonize_cell(FGridCell& Cell, int VertIndexArray[])
 	{
 		// construct bits of index into edge table, where bit for each
 		// corner is 1 if that value is < isovalue.
diff --git a/Engine/Source/Runtime/GeometryCore/Public/Implicit/Morphology.h b/Engine/Source/Runtime/GeometryCore/Public/Implicit/Morphology.h
--- a/Engine/Source/Runtime/GeometryCore/Public/Implicit/Morphology.h
+++ b/Engine/Source/Runtime/GeometryCore/Public/Implicit/Morphology.h
@@ -221,6 +221,7 @@ protected:
 
 		if (SecondSDF.ExactBandWidth > 1) // for larger band width, prefer using the AABB tree to do one distance per cell.  TODO: tune?
 		{
+			TRACE_CPUPROFILER_EVENT_SCOPE(Geometry_Morphology_SecondPass_BuildSpatial);
 			SecondSpatial.Build();
 			SecondSDF.Spatial = &SecondSpatial;
 			SecondSDF.ComputeMode = TSweepingMeshSDF<TIndexVectorMeshArrayAdapter<FIndex3i, double, FVector3d>>::EComputeModes::NarrowBand_SpatialFloodFill;
