MeshModelingTools: Added comments clarifying things related to FDynamicMesh3, particularly role of per-vertex attributes.

diff --git a/Engine/Plugins/Runtime/GeometryProcessing/Source/DynamicMesh/Public/GroupTopology.h b/Engine/Plugins/Runtime/GeometryProcessing/Source/DynamicMesh/Public/GroupTopology.h
--- a/Engine/Plugins/Runtime/GeometryProcessing/Source/DynamicMesh/Public/GroupTopology.h
+++ b/Engine/Plugins/Runtime/GeometryProcessing/Source/DynamicMesh/Public/GroupTopology.h
@@ -142,8 +142,9 @@ private:
  * You can override GetGroupID to provide your own grouping.
  * 
  * Various query functions are provided to allow group topology to be interrogated.
- * Note that these functions refer to "GroupID", "CornerID", and "GroupEdgeID", 
- * these are simply indices into the internal .Groups, .Corners, and .Edges arrays
+ * Note that when functions refer to "CornerID", and "GroupEdgeID", these are simply
+ * indices into the internal Corners and Edges arrays. However references to "GroupID"
+ * refer to the actual group IDs in FDynamicMesh (not an index into Groups).
  */
 class DYNAMICMESH_API FGroupTopology
 {
diff --git a/Engine/Source/Runtime/GeometryCore/Public/DynamicMesh/DynamicMesh3.h b/Engine/Source/Runtime/GeometryCore/Public/DynamicMesh/DynamicMesh3.h
--- a/Engine/Source/Runtime/GeometryCore/Public/DynamicMesh/DynamicMesh3.h
+++ b/Engine/Source/Runtime/GeometryCore/Public/DynamicMesh/DynamicMesh3.h
@@ -56,12 +56,21 @@ enum class EMeshComponents : uint8
 * if necessary, as the internal data structure is not exposed
 *
 * Per-vertex Vertex Normals, Colors, and UVs are optional and stored as floats.
+* Note that in practice, these are generally only used as scratch space, in limited 
+* circumstances, usually when needed for performance reasons. Most of our geometry
+* code instead prefers to read attributes from the per-triangle AttributeSet accessed 
+* via Attributes() (see TDynamicMeshOverlay for a description of the structure). For
+* instance, an empty (but existing) attribute set will take precedence over non-empty
+* vertex normals in much of our processing code.
 *
 * For each vertex, VertexEdgeLists[i] is the unordered list of connected edges. The
 * elements of the list are indices into the edges list.
 * This list is unsorted but can be traversed in-order (ie cw/ccw) at some additional cost.
 *
-* Triangles are stored as 3 ints, with optionally a per-triangle integer group id.
+* Triangles are stored as 3 ints, with optionally a per-triangle integer group id. 
+* The group IDs stored here DO get widely used and preserved in our geometry code 
+* (unlike the per-vertex attributes described earlier), even though the AttributeSet
+* can store group IDs as well (potentially in multiple layers).
 *
 * The edges of a triangle are similarly stored as 3 ints, in triangle_edes. If the
 * triangle is [v1,v2,v3], then the triangle edges [e1,e2,e3] are
@@ -167,7 +176,10 @@ protected:
 	/** Extended Attributes for the Mesh (UV layers, Hard Normals, additional Polygroup Layers, etc) */
 	TUniquePtr<FDynamicMeshAttributeSet> AttributeSet{};
 
-	/** Enable/Disable updating of ShapeChangeStamp. This can be problematic in multi-threaded contexts so it is disabled by default. */
+	/** 
+	 * Enable/Disable updating of ShapeChangeStamp. This can be problematic in multi-threaded contexts so it is disabled by default. 
+	 * In fact, it is not suggested that these be used at all (see commment for SetShapeChangeStampEnabled).
+	 */
 	bool bEnableShapeChangeStamp = false;
 	/** If bEnableShapeChangeStamp=true, The shape change stamp is incremented any time a function that modifies the mesh shape or topology is called */
 	std::atomic<uint32> ShapeChangeStamp = 1;
@@ -343,11 +355,18 @@ public:
 	//
 	// Change Tracking support
 	// 
-	// 
+	// Note: May someday be removed (see comment for SetShapeChangeStampEnabled) 
 	// 
 public:
 
-	/** Enable/Disable incrementing of the ShapeChangeStamp */
+	/** 
+	 * Enable/Disable incrementing of the ShapeChangeStamp.
+	 * 
+	 * NOTE: Both change stamps are unreliable in many contexts, and may someday get removed.
+	 * Specifically, they get reset on Clear(), copied to the destination rather than incremented
+	 * in copies and moves, and do not track any changes via access to the buffers directly (through
+	 * GetVerticesBuffer(), etc).
+	 */
 	void SetShapeChangeStampEnabled(bool bEnabled)
 	{
 		bEnableShapeChangeStamp = bEnabled;
diff --git a/Engine/Source/Runtime/GeometryCore/Public/DynamicMesh/InfoTypes.h b/Engine/Source/Runtime/GeometryCore/Public/DynamicMesh/InfoTypes.h
--- a/Engine/Source/Runtime/GeometryCore/Public/DynamicMesh/InfoTypes.h
+++ b/Engine/Source/Runtime/GeometryCore/Public/DynamicMesh/InfoTypes.h
@@ -52,7 +52,7 @@ struct FEdgeSplitInfo
 	bool bIsBoundary;					// was the split edge a boundary edge?  (redundant)
 
 	int NewVertex;						// new vertex f that was created
-	FIndex2i NewTriangles;				// new triangles [t2,t3], oriented as explained below
+	FIndex2i NewTriangles;				// new triangles [t2,t3], oriented as explained in SplitEdge() header comment
 	FIndex3i NewEdges;					// new edges are [f,b], [f,c] and [f,d] if this is not a boundary edge
 
 	double SplitT;						// parameter value for NewVertex along original edge
