make convex hull code work directly with floats if the input is in floats, rather than converting everything to doubles

diff --git a/Engine/Source/Runtime/GeometryCore/Private/CompGeom/ConvexHull2.cpp b/Engine/Source/Runtime/GeometryCore/Private/CompGeom/ConvexHull2.cpp
--- a/Engine/Source/Runtime/GeometryCore/Private/CompGeom/ConvexHull2.cpp
+++ b/Engine/Source/Runtime/GeometryCore/Private/CompGeom/ConvexHull2.cpp
@@ -50,12 +50,12 @@ bool TConvexHull2<RealType>::Solve(int32 NumPoints, TFunctionRef<TVector2<RealTy
 	}
 	else
 	{
-		FVector2d FirstTwoPts[2]{ (FVector2d)GetPointFunc(Hull[0]), (FVector2d)GetPointFunc(Hull[1]) };
+		TVector2<RealType> FirstTwoPts[2]{ GetPointFunc(Hull[0]), GetPointFunc(Hull[1]) };
 		bool bFoundSecondDim = false;
 		for (int32 Idx = 2; Idx < Hull.Num(); Idx++)
 		{
-			FVector2d Pt = (FVector2d)GetPointFunc(Hull[Idx]);
-			if (ExactPredicates::Orient2D(FirstTwoPts[0], FirstTwoPts[1], Pt) != 0)
+			TVector2<RealType> Pt = GetPointFunc(Hull[Idx]);
+			if (ExactPredicates::Orient2<RealType>(FirstTwoPts[0], FirstTwoPts[1], Pt) != 0)
 			{
 				bFoundSecondDim = true;
 				break;
@@ -113,14 +113,14 @@ void TConvexHull2<RealType>::Merge(TFunctionRef<TVector2<RealType>(int32)> GetPo
 	int32 size1 = j3 - j2 + 1;
 
 	int32 i;
-	FVector2d p;
+	TVector2<RealType> p;
 
 	// Find the right-most point of the left subhull.
-	FVector2d pmax0 = (FVector2d)GetPointFunc(Hull[j0]);
+	TVector2<RealType> pmax0 = GetPointFunc(Hull[j0]);
 	int32 imax0 = j0;
 	for (i = j0 + 1; i <= j1; ++i)
 	{
-		p = (FVector2d)GetPointFunc(Hull[i]);
+		p = GetPointFunc(Hull[i]);
 		if (pmax0.X < p.X || (pmax0.X == p.X && pmax0.Y < p.Y)) // lexicographic pmax0 < p
 		{
 			pmax0 = p;
@@ -129,11 +129,11 @@ void TConvexHull2<RealType>::Merge(TFunctionRef<TVector2<RealType>(int32)> GetPo
 	}
 
 	// Find the left-most point of the right subhull.
-	FVector2d pmin1 = (FVector2d)GetPointFunc(Hull[j2]);
+	TVector2<RealType> pmin1 = GetPointFunc(Hull[j2]);
 	int32 imin1 = j2;
 	for (i = j2 + 1; i <= j3; ++i)
 	{
-		p = (FVector2d)GetPointFunc(Hull[i]);
+		p = GetPointFunc(Hull[i]);
 		if (p.X < pmin1.X || (p.X == pmin1.X && p.Y < pmin1.Y)) // lexicographic p < pmin1
 		{
 			pmin1 = p;
@@ -198,10 +198,10 @@ enum class EPointOrdering
 	CollinearContain
 };
 
-
-EPointOrdering PointOnLine(FVector2d P, FVector2d L0, FVector2d L1)
+template<typename RealType>
+EPointOrdering PointOnLine(const TVector2<RealType>& P, const TVector2<RealType>& L0, const TVector2<RealType>& L1)
 {
-	double OnLine = ExactPredicates::Orient2D(P, L0, L1);
+	RealType OnLine = ExactPredicates::Orient2(P, L0, L1);
 	if (OnLine > 0)
 	{
 		return EPointOrdering::Positive;
@@ -247,20 +247,20 @@ void TConvexHull2<RealType>::GetTangent(TFunctionRef<TVector2<RealType>(int32)>
 	int32 size1 = j3 - j2 + 1;
 	int32 const imax = size0 + size1;
 	int32 i, iLm1, iRp1;
-	FVector2d L0, L1, R0, R1;
+	TVector2<RealType> L0, L1, R0, R1;
 
 	for (i = 0; i < imax; i++)
 	{
 		// Get the endpoints of the potential tangent.
-		L1 = (FVector2d)GetPointFunc(Hull[i0]);
-		R0 = (FVector2d)GetPointFunc(Hull[i1]);
+		L1 = GetPointFunc(Hull[i0]);
+		R0 = GetPointFunc(Hull[i1]);
 
 		// Walk along the left hull to find the point of tangency.
 		if (size0 > 1)
 		{
 			iLm1 = (i0 > j0 ? i0 - 1 : j1);
-			L0 = (FVector2d)GetPointFunc(Hull[iLm1]);
-			EPointOrdering Order = PointOnLine(R0, L0, L1);
+			L0 = GetPointFunc(Hull[iLm1]);
+			EPointOrdering Order = PointOnLine<RealType>(R0, L0, L1);
 			if (Order == EPointOrdering::Negative
 				|| Order == EPointOrdering::CollinearRight)
 			{
@@ -273,8 +273,8 @@ void TConvexHull2<RealType>::GetTangent(TFunctionRef<TVector2<RealType>(int32)>
 		if (size1 > 1)
 		{
 			iRp1 = (i1 < j3 ? i1 + 1 : j2);
-			R1 = (FVector2d)GetPointFunc(Hull[iRp1]);
-			EPointOrdering Order = PointOnLine(L1, R0, R1);
+			R1 = GetPointFunc(Hull[iRp1]);
+			EPointOrdering Order = PointOnLine<RealType>(L1, R0, R1);
 			if (Order == EPointOrdering::Negative
 				|| Order == EPointOrdering::CollinearLeft)
 			{
diff --git a/Engine/Source/Runtime/GeometryCore/Private/CompGeom/ConvexHull3.cpp b/Engine/Source/Runtime/GeometryCore/Private/CompGeom/ConvexHull3.cpp
--- a/Engine/Source/Runtime/GeometryCore/Private/CompGeom/ConvexHull3.cpp
+++ b/Engine/Source/Runtime/GeometryCore/Private/CompGeom/ConvexHull3.cpp
@@ -233,20 +233,20 @@ struct FHullConnectivity
 	 * @param ValueOut The volume of the tetrahedron formed by TriPts and Pt
 	 * @return true if Pt is on the 'positive' side of the triangle
 	 */
-	bool IsVisible(const FVector3d TriPts[3], const FVector3d& Pt, double& ValueOut)
+	bool IsVisible(const TVector<RealType> TriPts[3], const TVector<RealType>& Pt, double& ValueOut)
 	{
-		ValueOut = ExactPredicates::Orient3D(TriPts[0], TriPts[1], TriPts[2], Pt);
+		ValueOut = ExactPredicates::Orient3<RealType>(TriPts[0], TriPts[1], TriPts[2], Pt);
 		return ValueOut > 0;
 	}
 
 	/**
 	 * Set triangle positions to prep for IsVisible calls
 	 */
-	void SetTriPts(const FIndex3i& Tri, TFunctionRef<TVector<RealType>(int32)> GetPointFunc,  FVector3d TriPtsOut[3])
+	void SetTriPts(const FIndex3i& Tri, TFunctionRef<TVector<RealType>(int32)> GetPointFunc,  TVector<RealType> TriPtsOut[3])
 	{
-		TriPtsOut[0] = (FVector3d)GetPointFunc(Tri[0]);
-		TriPtsOut[1] = (FVector3d)GetPointFunc(Tri[1]);
-		TriPtsOut[2] = (FVector3d)GetPointFunc(Tri[2]);
+		TriPtsOut[0] = GetPointFunc(Tri[0]);
+		TriPtsOut[1] = GetPointFunc(Tri[1]);
+		TriPtsOut[2] = GetPointFunc(Tri[2]);
 	}
 
 	/**
@@ -256,15 +256,15 @@ struct FHullConnectivity
 	{
 		VisiblePoints.SetNum(Triangles.Num());
 
-		FVector3d TriPts[3];
-		FVector3d Pt;
+		TVector<RealType> TriPts[3];
+		TVector<RealType> Pt;
 		for (int32 PtIdx = 0; PtIdx < NumPoints; PtIdx++)
 		{
 			if (!FilterFunc(PtIdx))
 			{
 				continue;
 			}
-			Pt = (FVector3d)GetPointFunc(PtIdx);
+			Pt = GetPointFunc(PtIdx);
 			for (int32 TriIdx = 0; TriIdx < Triangles.Num(); TriIdx++)
 			{
 				SetTriPts(Triangles[TriIdx], GetPointFunc, TriPts);
@@ -474,12 +474,12 @@ struct FHullConnectivity
 	 * @param CrossedEdgeFirstVertex The first vertex of the edge that was 'crossed over' to traverse to this TriIdx, or -1 for the initial call
 	 * @return false if triangle is beyond horizon (so we don't need to search through it / remove it), true otherwise
 	 */
-	bool HorizonHelper(const TArray<FIndex3i>& Triangles, TFunctionRef<TVector<RealType>(int32)> GetPointFunc, const FVector3d& Pt, TArray<int32>& NewlyUnclaimed, TSet<int32>& ToDelete, TArray<FNewTriangle>& ToAdd, int32 TriIdx, int32 CrossedEdgeFirstVertex)
+	bool HorizonHelper(const TArray<FIndex3i>& Triangles, TFunctionRef<TVector<RealType>(int32)> GetPointFunc, const TVector<RealType>& Pt, TArray<int32>& NewlyUnclaimed, TSet<int32>& ToDelete, TArray<FNewTriangle>& ToAdd, int32 TriIdx, int32 CrossedEdgeFirstVertex)
 	{
 		// if it's not the first triangle, crossed edge should be set and we should check if the triangle is visible / actually needs to be replaced
 		if (CrossedEdgeFirstVertex != -1)
 		{
-			FVector3d TriPts[3];
+			TVector<RealType> TriPts[3];
 			SetTriPts(Triangles[TriIdx], GetPointFunc, TriPts);
 			double UnusedValue;
 			if (!IsVisible(TriPts, Pt, UnusedValue))
@@ -530,7 +530,7 @@ struct FHullConnectivity
 	{
 		//ValidateConnectivity(Triangles);
 
-		FVector3d Pt = (FVector3d)GetPointFunc(PtIdx);
+		TVector<RealType> Pt = GetPointFunc(PtIdx);
 
 		// Call the recursive helper to find all the tris we need to delete + the 'horizon' of edges that will form new triangles
 		TArray<int32> NewlyUnclaimed;
@@ -541,7 +541,7 @@ struct FHullConnectivity
 		// Connect up all the horizon triangles
 		int32 NewTriStart = Triangles.Num();
 		int32 NumAdd = ToAdd.Num();
-		FVector3d TriPts[3];
+		TVector<RealType> TriPts[3];
 		for (int32 AddIdx = 0; AddIdx < NumAdd; AddIdx++)
 		{
 			const FNewTriangle& TriData = ToAdd[AddIdx];
@@ -559,7 +559,7 @@ struct FHullConnectivity
 			{
 				int32 UnPtIdx = NewlyUnclaimed[UnclaimedIdx];
 				double Value;
-				if (IsVisible(TriPts, (FVector3d)GetPointFunc(UnPtIdx), Value))
+				if (IsVisible(TriPts, GetPointFunc(UnPtIdx), Value))
 				{
 					Visible.AddPt(UnPtIdx, Value);
 					NewlyUnclaimed.RemoveAtSwap(UnclaimedIdx, 1, false);
@@ -608,7 +608,7 @@ bool TConvexHull3<RealType>::Solve(int32 NumPoints, TFunctionRef<TVector<RealTyp
 	}
 
 	// safety check; seems possible the InitialTet chosen points were actually coplanar, because it was constructed w/ inexact math
-	if (ExactPredicates::Orient3D(GetPointFunc(InitialTet.Extreme[0]), GetPointFunc(InitialTet.Extreme[1]), GetPointFunc(InitialTet.Extreme[2]), GetPointFunc(InitialTet.Extreme[3])) == 0)
+	if (ExactPredicates::Orient3<RealType>(GetPointFunc(InitialTet.Extreme[0]), GetPointFunc(InitialTet.Extreme[1]), GetPointFunc(InitialTet.Extreme[2]), GetPointFunc(InitialTet.Extreme[3])) == 0)
 	{
 		Plane = TPlane3<RealType>(InitialTet.Basis[1], InitialTet.Origin);
 		Dimension = 2;
