Add a 'contain(other interval)' function to TInterval1

diff --git a/Engine/Source/Runtime/GeometryCore/Public/BoxTypes.h b/Engine/Source/Runtime/GeometryCore/Public/BoxTypes.h
--- a/Engine/Source/Runtime/GeometryCore/Public/BoxTypes.h
+++ b/Engine/Source/Runtime/GeometryCore/Public/BoxTypes.h
@@ -67,6 +67,18 @@ struct TInterval1
 		}
 	}
 
+	void Contain(const TInterval1<RealType>& O)
+	{
+		if (O.Min < Min)
+		{
+			Min = O.Min;
+		}
+		if (O.Max > Max)
+		{
+			Max = O.Max;
+		}
+	}
+
 	bool Contains(RealType D) const
 	{
 		return D >= Min && D <= Max;
