new Delaunay triangulation code using robust predicates
Much faster and usable from non-plugin code

diff --git a/Engine/Source/Runtime/GeometryCore/Private/CompGeom/Delaunay2.cpp b/Engine/Source/Runtime/GeometryCore/Private/CompGeom/Delaunay2.cpp
new file mode 100644
--- /dev/null
+++ b/Engine/Source/Runtime/GeometryCore/Private/CompGeom/Delaunay2.cpp
@@ -0,0 +1,919 @@
+// Copyright Epic Games, Inc. All Rights Reserved.
+
+#include "CompGeom/Delaunay2.h"
+
+#include "CompGeom/ExactPredicates.h"
+#include "Spatial/ZOrderCurvePoints.h"
+
+namespace UE
+{
+namespace Geometry
+{
+
+// Simple triangle connectivity structure designed for Delaunay triangulation specifically; may not support e.g. non-manifold meshes
+// To support Delaunay triangulation algorithms, the structure supports having a single 'ghost vertex' connected to the boundary of the triangulation
+// Currently this is a very simple edge TMap + an optional vertex->edge cache; it may be faster if switched to something that is not so TMap-based
+struct FDelaunay2Connectivity
+{
+	static constexpr int32 GhostIndex = -1;
+	static constexpr int32 InvalidIndex = -2;
+
+	void Empty(int32 ExpectedMaxVertices = 0)
+	{
+		EdgeToVert.Empty(ExpectedMaxVertices * 2 * 3);
+		DisableVertexAdjacency();
+	}
+
+	// Add the triangles from another mesh directly to this one
+	void Append(const FDelaunay2Connectivity& ToAdd)
+	{
+		for (const TPair<FIndex2i, int32>& EdgeV : ToAdd.EdgeToVert)
+		{
+			EdgeToVert.Add(EdgeV.Key, EdgeV.Value);
+		}
+	}
+
+	// just turns on faster vertex->edge lookups, at cost of increased storage and bookkeeping
+	void EnableVertexAdjacency(int32 NumVertices)
+	{
+		// if cache was off or the wrong size, create it fresh
+		if (!bUseAdjCache || VertexAdjCache.Num() < VertexIDToAdjIndex(NumVertices))
+		{
+			bUseAdjCache = true;
+			VertexAdjCache.Init(-1, VertexIDToAdjIndex(NumVertices));
+			for (const TPair<FIndex2i, int32>& EV : EdgeToVert)
+			{
+				UpdateAdjCache(EV.Key);
+			}
+		}
+	}
+
+	void DisableVertexAdjacency()
+	{
+		bUseAdjCache = false;
+		VertexAdjCache.Empty();
+	}
+
+	bool HasEdge(const FIndex2i& Edge) const
+	{
+		return EdgeToVert.Contains(Edge);
+	}
+
+	int32 NumTriangles() const
+	{
+		return EdgeToVert.Num() / 3;
+	}
+
+	int32 NumHalfEdges() const
+	{
+		return EdgeToVert.Num();
+	}
+
+	TArray<FIndex3i> GetTriangles() const
+	{
+		TArray<FIndex3i> Triangles;
+		Triangles.Reserve(EdgeToVert.Num() / 3);
+		for (const TPair<FIndex2i, int32>& EV : EdgeToVert)
+		{
+			if (!IsGhost(EV.Key, EV.Value) && EV.Value < EV.Key.A && EV.Value < EV.Key.B)
+			{
+				Triangles.Emplace(EV.Key.A, EV.Key.B, EV.Value);
+			}
+		}
+		return Triangles;
+	}
+
+	void GetTrianglesAndAdjacency(TArray<FIndex3i>& Triangles, TArray<FIndex3i>& Adjacency) const
+	{
+		Triangles.Reset(EdgeToVert.Num() / 3);
+		for (const TPair<FIndex2i, int32>& EV : EdgeToVert)
+		{
+			if (!IsGhost(EV.Key, EV.Value) && EV.Value < EV.Key.A && EV.Value < EV.Key.B)
+			{
+				if (!IsGhost(EV.Key, EV.Value) && EV.Value < EV.Key.A && EV.Value < EV.Key.B)
+				{
+					Triangles.Emplace(EV.Key.A, EV.Key.B, EV.Value);
+				}
+			}
+		}
+
+		// Because the EdgeToVert map doesn't know anything about our triangle indices,
+		// it's easiest to build the Adjacency data from scratch with a new Map
+		Adjacency.Init(FIndex3i::Invalid(), Triangles.Num());
+		TMap<FIndex2i, int32> EdgeToTri;
+		for (int32 TriIdx = 0; TriIdx < Triangles.Num(); TriIdx++)
+		{
+			const FIndex3i& Tri = Triangles[TriIdx];
+			for (int32 Sub0 = 2, Sub1 = 0; Sub1 < 3; Sub0 = Sub1++)
+			{
+				FIndex2i RevEdge(Tri[Sub1], Tri[Sub0]);
+				int32* NbrTriIdx = EdgeToTri.Find(RevEdge);
+				if (NbrTriIdx)
+				{
+					EdgeToTri.Remove(RevEdge); // only need the match once
+
+					const FIndex3i& AdjTri = Triangles[*NbrTriIdx];
+					int32 AdjSub = AdjTri.IndexOf(RevEdge.A);
+					Adjacency[*NbrTriIdx][AdjSub] = TriIdx;
+					Adjacency[TriIdx][Sub0] = *NbrTriIdx;
+				}
+				else
+				{
+					EdgeToTri.Add(FIndex2i(RevEdge.B, RevEdge.A), TriIdx);
+				}
+			}
+		}
+	}
+
+	static bool IsGhost(const FIndex2i& Edge, int32 Vertex)
+	{
+		return Edge.A == GhostIndex || Edge.B == GhostIndex || Vertex == GhostIndex;
+	}
+
+	static bool IsGhost(const FIndex2i& Edge)
+	{
+		return Edge.A == GhostIndex || Edge.B == GhostIndex;
+	}
+
+	void AddTriangle(const FIndex3i& Tri)
+	{
+		EdgeToVert.Add(FIndex2i(Tri.A, Tri.B), Tri.C);
+		EdgeToVert.Add(FIndex2i(Tri.B, Tri.C), Tri.A);
+		EdgeToVert.Add(FIndex2i(Tri.C, Tri.A), Tri.B);
+		if (bUseAdjCache)
+		{
+			FIndex3i TriAdjI(VertexIDToAdjIndex(Tri.A), VertexIDToAdjIndex(Tri.B), VertexIDToAdjIndex(Tri.C));
+			VertexAdjCache[TriAdjI.A] = TriAdjI.B;
+			VertexAdjCache[TriAdjI.B] = TriAdjI.C;
+			VertexAdjCache[TriAdjI.C] = TriAdjI.A;
+		}
+	}
+
+	// Create a first initial triangle that is surround by ghost triangles
+	void InitWithGhosts(const FIndex3i& Tri)
+	{
+		AddTriangle(Tri);
+		AddTriangle(FIndex3i(Tri.B, Tri.A, GhostIndex));
+		AddTriangle(FIndex3i(Tri.C, Tri.B, GhostIndex));
+		AddTriangle(FIndex3i(Tri.A, Tri.C, GhostIndex));
+	}
+
+	void DeleteTriangle(const FIndex3i& Tri)
+	{
+		EdgeToVert.Remove(FIndex2i(Tri.A, Tri.B));
+		EdgeToVert.Remove(FIndex2i(Tri.B, Tri.C));
+		EdgeToVert.Remove(FIndex2i(Tri.C, Tri.A));
+		if (bUseAdjCache)
+		{
+			// clear any adj info if it was set in the cache
+			FIndex3i TriAdjI(VertexIDToAdjIndex(Tri.A), VertexIDToAdjIndex(Tri.B), VertexIDToAdjIndex(Tri.C));
+			const int32 InvalidAdj = VertexIDToAdjIndex(InvalidIndex);
+			if (VertexAdjCache[TriAdjI.A] == TriAdjI.B)
+			{
+				VertexAdjCache[TriAdjI.A] = InvalidAdj;
+			}
+			if (VertexAdjCache[TriAdjI.B] == TriAdjI.C)
+			{
+				VertexAdjCache[TriAdjI.B] = InvalidAdj;
+			}
+			if (VertexAdjCache[TriAdjI.C] == TriAdjI.A)
+			{
+				VertexAdjCache[TriAdjI.C] = InvalidAdj;
+			}
+		}
+	}
+
+	int32 GetVertex(const FIndex2i& Edge) const
+	{
+		const int32* V = EdgeToVert.Find(Edge);
+		if (!V)
+		{
+			return InvalidIndex;
+		}
+		return *V;
+	}
+
+	// Get any edge BC opposite vertex A, such that triangle ABC is in the mesh (or return InvalidIndex edge if no such edge is present)
+	// Before calling this frequently, consider calling EnableVertexAdjacency()
+	FIndex2i GetEdge(int32 Vertex) const
+	{
+		if (bUseAdjCache)
+		{
+			int32 AdjVertex = GetCachedAdjVertex(Vertex);
+			if (AdjVertex != InvalidIndex)
+			{
+				int32 LastVertex = GetVertex(FIndex2i(Vertex, AdjVertex));
+				return FIndex2i(AdjVertex, LastVertex);
+			}
+		}
+		for (const TPair<FIndex2i, int32>& EV : EdgeToVert)
+		{
+			if (bUseAdjCache)
+			{
+				UpdateAdjCache(EV.Key);
+			}
+			if (Vertex == EV.Key.A)
+			{
+				return FIndex2i(EV.Key.B, EV.Value);
+			}
+		}
+		return FIndex2i(InvalidIndex, InvalidIndex);
+	}
+	
+	// Call a function on every oriented edge (+ next vertex) on the mesh
+	// (note the number of edges visited will be 3x the number of triangles)
+	// VisitFunctionType is expected to take (FIndex2i Edge, int32 Vertex) and return bool
+	// Returning false from VisitFn will end the enumeration early
+	template<typename VisitFunctionType>
+	void EnumerateOrientedEdges(VisitFunctionType VisitFn)
+	{
+		for (const TPair<FIndex2i, int32>& EdgeVert : EdgeToVert)
+		{
+			if (!VisitFn(EdgeVert.Key, EdgeVert.Value))
+			{
+				break;
+			}
+		}
+	}
+
+	// Similar to EnumerateOrientedEdges but only visits each triangle once, instead of 3x
+	// and optionally skips ghost triangles (triangles connected to the ghost vertex)
+	template<typename VisitFunctionType>
+	void EnumerateTriangles(VisitFunctionType VisitFn, bool bSkipGhosts = false)
+	{
+		for (const TPair<FIndex2i, int32>& EdgeVert : EdgeToVert)
+		{
+			// to visit triangles only once, only visit when the vertex ID is smaller than the edge IDs
+			// since the vertex ID is the smallest ID, it is also the only one we need to check vs the GhostIndex (if we're skipping ghosts)
+			if (EdgeVert.Key.A < EdgeVert.Value || EdgeVert.Key.B < EdgeVert.Value || (bSkipGhosts && EdgeVert.Value == GhostIndex))
+			{
+				continue;
+			}
+			if (!VisitFn(EdgeVert.Key, EdgeVert.Value))
+			{
+				break;
+			}
+		}
+	}
+
+protected:
+	TMap<FIndex2i, int32> EdgeToVert;
+	
+	// Optional cache of a single vertex in the 1-ring of each vertex
+	// Makes GetEdge() constant time (as long as the cache hits) instead of O(#Edges), at the cost of additional storage and bookkeeping.
+	mutable TArray<int32> VertexAdjCache;
+	bool bUseAdjCache = false;
+
+	static inline int32 VertexIDToAdjIndex(int32 VertexID)
+	{
+		// Offset by 1 so that GhostIndex enters slot 0
+		return VertexID + 1;
+	}
+	static inline int32 AdjIndexToVertexID(int32 AdjIndex)
+	{
+		return AdjIndex - 1;
+	}
+	inline int32 GetCachedAdjVertex(int32 VertexID) const
+	{
+		int32 AdjIndex = VertexIDToAdjIndex(VertexID);
+		return AdjIndexToVertexID(VertexAdjCache[AdjIndex]);
+	}
+	inline void UpdateAdjCache(FIndex2i Edge) const
+	{
+		FIndex2i AdjEdge(VertexIDToAdjIndex(Edge.A), VertexIDToAdjIndex(Edge.B));
+		VertexAdjCache[AdjEdge.A] = AdjEdge.B;
+	}
+};
+
+namespace DelaunayInternal
+{
+	// Helper to create a permutation array
+	TArray<int32> GetShuffledOrder(FRandomStream& Random, int32 Num, int32 StartIn = -1, int32 EndIn = -1)
+	{
+		TArray<int32> Order;
+		Order.Reserve(Num);
+		for (int32 OrderIdx = 0; OrderIdx < Num; OrderIdx++)
+		{
+			Order.Add(OrderIdx);
+		}
+		int32 Start = StartIn >= 0 ? StartIn : 0;
+		int32 End = EndIn >= 0 ? EndIn : Num - 1;
+		for (int32 OrderIdx = EndIn; OrderIdx > Start; OrderIdx--)
+		{
+			int32 SwapIdx = Start + Random.RandHelper(OrderIdx - 1 - Start);
+			Swap(Order[SwapIdx], Order[OrderIdx]);
+		}
+		return Order;
+	}
+
+	// @return true if Vertex is inside the circumcircle of Tri
+	// For ghost triangles, this is defined as being on the one solid edge of the triangle or inside that edge's half-space
+	template<typename RealType>
+	bool InTriCircle(TArrayView<const TVector2<RealType>> Vertices, FIndex3i Tri, int32 Vertex)
+	{
+		int32 GhostSub = Tri.IndexOf(FDelaunay2Connectivity::GhostIndex);
+		if (GhostSub == -1)
+		{
+			return ExactPredicates::InCircle2<RealType>(Vertices[Tri.A], Vertices[Tri.B], Vertices[Tri.C], Vertices[Vertex]) > 0;
+		}
+		FIndex3i GhostFirst = Tri.GetCycled(FDelaunay2Connectivity::GhostIndex);
+		RealType Pred = ExactPredicates::Orient2<RealType>(Vertices[GhostFirst.B], Vertices[GhostFirst.C], Vertices[Vertex]);
+		if (Pred > 0)
+		{
+			return true;
+		}
+		if (Pred < 0)
+		{
+			return false;
+		}
+		// Pred == 0 case: need to check if Vertex is *on* the edge
+		if (Vertices[GhostFirst.B].X != Vertices[GhostFirst.C].X)
+		{
+			TInterval1<RealType> XRange(Vertices[GhostFirst.B].X, Vertices[GhostFirst.C].X);
+			return XRange.Contains(Vertices[Vertex].X);
+		}
+		else
+		{
+			TInterval1<RealType> YRange(Vertices[GhostFirst.B].Y, Vertices[GhostFirst.C].Y);
+			return YRange.Contains(Vertices[Vertex].Y);
+		}
+	}
+
+	// @return triangle containing Vertex
+	template<typename RealType>
+	FIndex3i WalkToContainingTri(FRandomStream& Random, FDelaunay2Connectivity& Connectivity, TArrayView<const TVector2<RealType>> Vertices, FIndex3i StartTri, int32 Vertex, bool bAssumeDelaunay, bool& bIsDuplicateOut)
+	{
+		bIsDuplicateOut = false;
+
+		FIndex2i FirstEdge(FDelaunay2Connectivity::InvalidIndex, FDelaunay2Connectivity::InvalidIndex);
+		constexpr int32 GhostV = FDelaunay2Connectivity::GhostIndex; // shorter name
+
+		auto ChooseCross = [&Random, bAssumeDelaunay, &Vertices, Vertex, GhostV](const FIndex3i& Tri, bool bSkipFirst) -> int32
+		{
+			auto CrossesEdge = [&Vertices, Vertex, GhostV](int32 A, int32 B, bool bOnGhostTri) -> bool
+			{
+				if (!bOnGhostTri || (A != GhostV && B != GhostV))
+				{
+					RealType Orient = ExactPredicates::Orient2<RealType>(Vertices[A], Vertices[B], Vertices[Vertex]);
+					// Note: could refine to quickly say we're on the triangle in the ghost + orient==0 case, if we're exactly on the edge, but this only saves walking one edge
+					if (Orient < 0 || (bOnGhostTri && Orient == 0))
+					{
+						return true;
+					}
+				}
+				return false;
+			};
+
+			int32 Choose[2]{ -1, -1 };
+			int32 Chosen = 0;
+			bool bIsGhost = Tri.Contains(GhostV);
+			int32 NextSub[3]{ 1, 2, 0 };
+			for (int32 EdgeSub = (int32)bSkipFirst; EdgeSub < 3; EdgeSub++)
+			{
+				if (CrossesEdge(Tri[EdgeSub], Tri[NextSub[EdgeSub]], bIsGhost))
+				{
+					// On a Delaunay mesh we can always walk across the first edge that has the target vertex on the other side of it
+					if (bAssumeDelaunay)
+					{
+						return EdgeSub;
+					}
+					// If the mesh is not Delaunay, randomly choose between edges that have the target vertex on the other side; this avoids a possible infinite cycle
+					else
+					{
+						Choose[Chosen++] = EdgeSub;
+					}
+				}
+			}
+			if (Chosen == 0)
+			{
+				return -1; // we're on this tri
+			}
+			else if (Chosen == 1)
+			{
+				return Choose[0];
+			}
+			else // (Chosen == 2)
+			{
+				return Choose[Random.RandHelper(2)];
+			}
+		};
+		
+		FIndex3i WalkTri = StartTri;
+		int32 Cross = ChooseCross(WalkTri, false);
+		int32 NumSteps = 0;
+		int32 NextIdx[3]{ 1,2,0 };
+		while (Cross != -1)
+		{
+			// if !bAssumeDelaunay, the random edge walk could choose poorly enough for any amount of steps to occur, but it should not happen in practice ...
+			// if this ensure() triggers it is more likely that some other problem has caused an infinite loop
+			if (!ensure(NumSteps++ < Connectivity.NumTriangles() * 100))
+			{
+				return FIndex3i(FDelaunay2Connectivity::InvalidIndex, FDelaunay2Connectivity::InvalidIndex, FDelaunay2Connectivity::InvalidIndex);
+			}
+
+			FIndex2i OppEdge = FIndex2i(WalkTri[NextIdx[Cross]], WalkTri[Cross]);
+			int32 OppVert = Connectivity.GetVertex(OppEdge);
+			checkSlow(OppVert != FDelaunay2Connectivity::InvalidIndex);
+			WalkTri = FIndex3i(OppEdge.A, OppEdge.B, OppVert);
+			Cross = ChooseCross(WalkTri, true);
+		}
+
+		bIsDuplicateOut =
+			(WalkTri.A >= 0 && Vertices[Vertex] == Vertices[WalkTri.A]) ||
+			(WalkTri.B >= 0 && Vertices[Vertex] == Vertices[WalkTri.B]) ||
+			(WalkTri.C >= 0 && Vertices[Vertex] == Vertices[WalkTri.C]);
+
+		return WalkTri;
+	}
+
+	// Insert Vertex into the triangulation; it must already be on the OnTri triangle
+	// Uses Bowyer-Watson algorithm:
+	//  1. Delete all the connected triangles whose circumcircles contain the vertex
+	//  2. Make a fan of triangles from the new vertex out to the border of the deletions
+	// @return one of the inserted triangles containing Vertex
+	template<typename RealType>
+	FIndex3i Insert(FDelaunay2Connectivity& Connectivity, TArrayView<const TVector2<RealType>> Vertices, FIndex3i OnTri, int32 ToInsertV)
+	{
+		// depth first search + deletion of triangles whose circles contain vertex
+		TArray<FIndex2i> ToConsider;
+		auto AddIfValid = [&ToConsider, &Connectivity](FIndex2i Edge)
+		{
+			if (Connectivity.HasEdge(Edge))
+			{
+				ToConsider.Add(Edge);
+			}
+		};
+
+		auto DeleteTri = [&AddIfValid, &Connectivity](FIndex3i Tri)
+		{
+			Connectivity.DeleteTriangle(Tri);
+
+			AddIfValid(FIndex2i(Tri.B, Tri.A));
+			AddIfValid(FIndex2i(Tri.C, Tri.B));
+			AddIfValid(FIndex2i(Tri.A, Tri.C));
+		};
+		DeleteTri(OnTri);
+
+		TArray<FIndex2i> Border;
+		while (!ToConsider.IsEmpty())
+		{
+			FIndex2i Edge = ToConsider.Pop(false);
+			int32 TriV = Connectivity.GetVertex(Edge);
+			if (TriV != FDelaunay2Connectivity::InvalidIndex) // tri still existed (wasn't deleted by earlier traversal)
+			{
+				FIndex3i ConsiderTri = FIndex3i(Edge.A, Edge.B, TriV);
+				if (InTriCircle(Vertices, ConsiderTri, ToInsertV))
+				{
+					DeleteTri(ConsiderTri);
+				}
+				else
+				{
+					Border.Add(Edge);
+				}
+			}
+		}
+		
+		for (FIndex2i BorderEdge : Border)
+		{
+			Connectivity.AddTriangle(FIndex3i(BorderEdge.B, BorderEdge.A, ToInsertV));
+		}
+
+		return Border.Num() > 0 ? FIndex3i(Border[0].B, Border[0].A, ToInsertV) : FIndex3i::Invalid();
+	}
+
+	template<typename RealType>
+	bool IsDelaunay(FDelaunay2Connectivity& Connectivity, TArrayView<const TVector2<RealType>> Vertices)
+	{
+		bool bFoundNonDelaunay = false;
+		Connectivity.EnumerateOrientedEdges([&Connectivity, &Vertices, &bFoundNonDelaunay](const FIndex2i& Edge, int32 Vert) -> bool
+		{
+			if (Connectivity.IsGhost(Edge, Vert))
+			{
+				return true;
+			}
+			const FIndex2i Pair(Edge.B, Edge.A);
+			if (Pair.Contains(Connectivity.GhostIndex))
+			{
+				return true;
+			}
+			const int32 PairV = Connectivity.GetVertex(Pair);
+			if (PairV < 0) // skip if ghost or missing
+			{
+				return true;
+			}
+			const RealType InCircleRes = ExactPredicates::InCircle2<RealType>(Vertices[Edge.A], Vertices[Edge.B], Vertices[Vert], Vertices[PairV]);
+			if (InCircleRes > 0)
+			{
+				bFoundNonDelaunay = true;
+				return false;
+			}
+			return true;
+		});
+		return !bFoundNonDelaunay;
+	}
+
+	template<typename RealType>
+	bool GetFirstCrossingEdge(FDelaunay2Connectivity& Connectivity, TArrayView<const TVector2<RealType>> Vertices, const FIndex2i& EdgeToConnect, FIndex2i& CrossingEdgeOut)
+	{
+		if (Connectivity.HasEdge(EdgeToConnect))
+		{
+			return false; // nothing to dig
+		}
+
+		FIndex2i StartWalk = Connectivity.GetEdge(EdgeToConnect.A);
+		if (StartWalk.A == FDelaunay2Connectivity::InvalidIndex)
+		{
+			return false; // edge starts at a vertex that is not in the triangulation (e.g., could have been a duplicate that was rejected)
+		}
+
+		TVector2<RealType> VA = Vertices[EdgeToConnect.A];
+		TVector2<RealType> VB = Vertices[EdgeToConnect.B];
+
+		auto IsCrossingEdgeOnA = [&VA, &VB, &Vertices](const FIndex2i& Edge, RealType OrientA, RealType& OrientBOut) -> bool
+		{
+			if (Edge.B >= 0)
+			{
+				OrientBOut = ExactPredicates::Orient2<RealType>(VA, VB, Vertices[Edge.B]);
+			}
+
+			if (FDelaunay2Connectivity::IsGhost(Edge))
+			{
+				return false;
+			}
+			int32 SignA = FMath::Sign(OrientA);
+			if (SignA >= 0)
+			{
+				return false;
+			}
+			int32 SignB = FMath::Sign(OrientBOut);
+			// A properly oriented edge crossing the AB segment, on a tri that includes A, must go from the negative side to the positive side of AB
+			// (positive to negative would be behind the AB edge, and a zero would either be behind or would prevent the edge from being inserted)
+			return SignB == 1;
+		};
+
+		FIndex2i WalkEdge = StartWalk;
+		RealType OrientA = 0;
+		if (WalkEdge.A != FDelaunay2Connectivity::GhostIndex)
+		{
+			OrientA = ExactPredicates::Orient2<RealType>(VA, VB, Vertices[WalkEdge.A]);
+		}
+		RealType OrientB; // computed by the crossing edge test
+		bool bIsCrossing = IsCrossingEdgeOnA(WalkEdge, OrientA, OrientB);
+		int32 EdgesWalked = 0;
+		while (!bIsCrossing)
+		{
+			checkSlow(EdgesWalked++ < Connectivity.NumHalfEdges());
+			int32 NextVertex = Connectivity.GetVertex(FIndex2i(EdgeToConnect.A, WalkEdge.B));
+			check(NextVertex != Connectivity.InvalidIndex); // should not be a hole in the mesh at this stage!  if there were, need to stop this loop and then walk the opposite direction
+			WalkEdge = FIndex2i(WalkEdge.B, NextVertex);
+			if (WalkEdge == StartWalk) // full cycle with no crossing found, cannot insert the edge (this can happen if the edge is blocked by an exactly-on-edge vertex)
+			{
+				return false;
+			}
+			OrientA = OrientB;
+			bIsCrossing = IsCrossingEdgeOnA(WalkEdge, OrientA, OrientB);
+		}
+		CrossingEdgeOut = WalkEdge;
+		return true;
+	}
+
+	// @return vertex we need to fill to.  If EdgeToConnect.A, no fill needed; if EdgeToConnect.B, a normal re-triangulation needed; if other, digging failed in the middle and we need to fill partially
+	template<typename RealType>
+	int32 DigCavity(FDelaunay2Connectivity& Connectivity, TArrayView<const TVector2<RealType>> Vertices, const FIndex2i& EdgeToConnect, TArray<int32>& CavityLOut, TArray<int32>& CavityROut)
+	{
+		CavityLOut.Reset();
+		CavityROut.Reset();
+
+		FIndex2i FirstCross;
+		bool bFoundCross = GetFirstCrossingEdge<RealType>(Connectivity, Vertices, EdgeToConnect, FirstCross);
+		if (!bFoundCross)
+		{
+			return EdgeToConnect.A;
+		}
+
+		// Delete the first triangle in the cavity
+		Connectivity.DeleteTriangle(FIndex3i(EdgeToConnect.A, FirstCross.A, FirstCross.B));
+
+		CavityROut.Add(EdgeToConnect.A);
+		CavityLOut.Add(EdgeToConnect.A);
+		CavityROut.Add(FirstCross.A);
+		CavityLOut.Add(FirstCross.B);
+
+		TVector2<RealType> VA = Vertices[EdgeToConnect.A];
+		TVector2<RealType> VB = Vertices[EdgeToConnect.B];
+		
+		// by convention WalkCross is always crossing from left to right
+		int32 PrevVertex = EdgeToConnect.A;
+		FIndex2i WalkCross(FirstCross.B, FirstCross.A);
+		while (true) // Note: Can't loop infinitely because it is deleting triangles as it walks
+		{
+			int32 NextV = Connectivity.GetVertex(WalkCross);
+			if (NextV == FDelaunay2Connectivity::InvalidIndex)
+			{
+				ensure(false); // walking off the triangulation would mean the triangulation is unrecoverably broken
+				return EdgeToConnect.A;
+			}
+			Connectivity.DeleteTriangle(FIndex3i(WalkCross.A, WalkCross.B, NextV)); // immediately delete where we walk
+
+			if (NextV == EdgeToConnect.B)
+			{
+				CavityROut.Add(EdgeToConnect.B);
+				CavityLOut.Add(EdgeToConnect.B);
+				return EdgeToConnect.B;
+			}
+
+			FIndex2i NextCross;
+
+			RealType OrientNextV = ExactPredicates::Orient2<RealType>(VA, VB, Vertices[NextV]);
+			if (OrientNextV == 0)
+			{
+				// can't reach target edge due to intersecting this vertex; just stop here
+				CavityROut.Add(NextV);
+				CavityLOut.Add(NextV);
+				return NextV;
+			}
+			else if (OrientNextV < 0)
+			{
+				PrevVertex = WalkCross.B;
+				NextCross = FIndex2i(WalkCross.A, NextV); // facing the next triangle
+				CavityROut.Add(NextV);
+			}
+			else
+			{
+				PrevVertex = WalkCross.A;
+				NextCross = FIndex2i(NextV, WalkCross.B); // facing the next triangle
+				CavityLOut.Add(NextV);
+			}
+			WalkCross = NextCross;
+		}
+
+		check(false); // can't reach here
+		return EdgeToConnect.A;
+	}
+
+	// Helper for FillCavity.  Adds new vertex U to the cavity, trying to attach it initially via triangle UVW
+	template<typename RealType>
+	void CavityInsertVertex(FDelaunay2Connectivity& CavityCDT, TArrayView<const TVector2<RealType>> Vertices, int32 U, FIndex2i VW)
+	{
+		int32 X = CavityCDT.GetVertex(FIndex2i(VW.B, VW.A));
+		// If adding the triangle does not conflict with the existing triangle opp edge VW
+		// we can immediately add the triangle
+		if (
+			X == CavityCDT.InvalidIndex ||
+			(0 < ExactPredicates::Orient2<RealType>(Vertices[U], Vertices[VW.A], Vertices[VW.B]) &&
+			0 >= ExactPredicates::InCircle2<RealType>(Vertices[U], Vertices[VW.A], Vertices[VW.B], Vertices[X]))
+		)
+		{
+			CavityCDT.AddTriangle(FIndex3i(U, VW.A, VW.B)); // already Delaunay
+			return;
+		}
+		// New vertex U conflicts with existing triangle VWX across edge VW,
+		// so we need to delete tri VWX and insert flipped triangles UVX and UXW
+		// then recurse to flip any triangles that are made non-Delaunay by that flip
+		CavityCDT.DeleteTriangle(FIndex3i(VW.B, VW.A, X));
+		CavityInsertVertex<RealType>(CavityCDT, Vertices, U, FIndex2i(VW.A, X));
+		CavityInsertVertex<RealType>(CavityCDT, Vertices, U, FIndex2i(X, VW.B));
+	};
+
+	// Implements the cavity CDT algorithm from "Delaunay Mesh Generation" page 76, 77
+	template<typename RealType>
+	void FillCavity(FRandomStream& Random, FDelaunay2Connectivity& Connectivity, TArrayView<const TVector2<RealType>> Vertices, const FIndex2i& Edge, const TArray<int32>& Cavity)
+	{
+		check(Cavity.Num() > 2 && Edge.B == Cavity[0] && Edge.A == Cavity.Last());
+		int32 CavityNum = Cavity.Num();
+
+		TArray<int32> Permute = GetShuffledOrder(Random, CavityNum, 1, CavityNum-2); // permutation of the inner vertices of Cavity
+		// doubly-linked list tracking the cavity re-ordering
+		TArray<int32> Next, Prev;
+		Next.SetNumUninitialized(CavityNum);
+		Prev.SetNumUninitialized(CavityNum);
+		for (int32 CurIdx = 0, PrevIdx = CavityNum-1; CurIdx < CavityNum; PrevIdx = CurIdx++)
+		{
+			Next[PrevIdx] = CurIdx;
+			Prev[CurIdx] = PrevIdx;
+		}
+
+		// Precompute Orient2D values that are proportional to distance to the AB line
+		TArray<RealType> ABOrient;
+		ABOrient.Init(0, CavityNum);
+		TVector2<RealType> V0 = Vertices[Cavity[0]];
+		TVector2<RealType> VLast = Vertices[Cavity.Last()];
+		for (int32 OrientIdx = 1; OrientIdx + 1 < CavityNum; OrientIdx++)
+		{
+			// TODO: this predicate is only exact vs zero; could use a fast alternative here?
+			ABOrient[OrientIdx] = ExactPredicates::Orient2<RealType>(V0, Vertices[Cavity[OrientIdx]], VLast);
+			checkSlow(ABOrient[OrientIdx] > 0);
+		}
+
+		FDelaunay2Connectivity CavityCDT;
+		CavityCDT.Empty(CavityNum - 2); 
+
+		// Create an insertion ordering that tries to avoid concavities between adjacent pairs, and track adjacencies via Next/Prev
+		for (int Idx = CavityNum - 2; Idx >= 2; Idx--)
+		{
+			// Make sure Permute[Idx] is not closer to the AB line than both its neighbors
+			while (
+				ABOrient[Permute[Idx]] < ABOrient[Prev[Permute[Idx]]] &&
+				ABOrient[Permute[Idx]] < ABOrient[Next[Permute[Idx]]])
+			{
+				int32 SwapIdx = Random.RandRange(1, Idx - 1);
+				Swap(Permute[Idx], Permute[SwapIdx]);
+			}
+			// make Next/Prev skip over Permute[Idx]
+			Next[Prev[Permute[Idx]]] = Next[Permute[Idx]];
+			Prev[Next[Permute[Idx]]] = Prev[Permute[Idx]];
+		}
+
+		// Add the first triangle of the cavity
+		CavityCDT.AddTriangle(FIndex3i(Cavity[0], Cavity[Permute[1]], Cavity[CavityNum-1]));
+
+		// Progressively add remaining triangles in the permuted ordering (via a recursive function that keeps the cavity triangulation Delaunay)
+		for (int32 Idx = 2; Idx < CavityNum - 1; Idx++)
+		{
+			CavityInsertVertex<RealType>(CavityCDT, Vertices, Cavity[Permute[Idx]],
+				FIndex2i(Cavity[Next[Permute[Idx]]], Cavity[Prev[Permute[Idx]]]));
+		}
+
+		// Insert the cavity triangulation into the overall triangulation
+		Connectivity.Append(CavityCDT);
+	}
+
+	template<typename RealType>
+	void ConstrainEdges(FRandomStream& Random, FDelaunay2Connectivity& Connectivity,
+		TArrayView<const TVector2<RealType>> Vertices, TArrayView<const FIndex2i> Edges, bool bKeepFastEdgeAdjacencyData)
+	{
+		constexpr int32 NeedFasterEdgeLookupThreshold = 4; // TODO: do some profiling to determine what this threshold should be
+		if (bKeepFastEdgeAdjacencyData || Edges.Num() > NeedFasterEdgeLookupThreshold)
+		{
+			Connectivity.EnableVertexAdjacency(Vertices.Num());
+		}
+
+		TArray<int32> CavityVerts[2];
+
+		// Random insertion order to improve expected performance
+		TArray<int32> EdgeOrder = GetShuffledOrder(Random, Edges.Num());
+		for (int32 OrderIdx = 0; OrderIdx < Edges.Num(); OrderIdx++)
+		{
+			int32 EdgeIdx = EdgeOrder[OrderIdx];
+			FIndex2i Edge = Edges[EdgeIdx];
+			int32 DigTo = DigCavity<RealType>(Connectivity, Vertices, Edge, CavityVerts[0], CavityVerts[1]);
+			if (DigTo != Edge.A)
+			{
+				FIndex2i DugEdge(Edge.A, DigTo); // fill the cavity we dug out (which may end at a different vertex than target, if there was a colinear vertex first)
+				FIndex2i RevEdge(DugEdge.B, DugEdge.A);
+				Algo::Reverse(CavityVerts[0]);
+				FillCavity<RealType>(Random, Connectivity, Vertices, DugEdge, CavityVerts[0]);
+				FillCavity<RealType>(Random, Connectivity, Vertices, RevEdge, CavityVerts[1]);
+			}
+		}
+
+		if (!bKeepFastEdgeAdjacencyData)
+		{
+			Connectivity.DisableVertexAdjacency();
+		}
+	}
+
+	template<typename RealType>
+	bool Triangulate(FRandomStream& Random, FDelaunay2Connectivity& Connectivity,
+		TArrayView<const TVector2<RealType>> Vertices, TArrayView<const FIndex2i> Edges, bool bKeepFastEdgeAdjacencyData)
+	{
+		Connectivity.Empty(Vertices.Num());
+
+		if (Vertices.Num() < 3)
+		{
+			return false;
+		}
+
+		// TODO: Combine the Z Order Curve ordering with a BRIO, to add enough randomization to break up pathological bad orderings
+		FZOrderCurvePoints InsertOrder;
+		InsertOrder.Compute(Vertices);
+		TArray<int32>& Order = InsertOrder.Order;
+
+		int32 BootstrapIndices[3]{ Order.Num() - 1, -1, -1 };
+		TVector2<RealType> Pts[3];
+		Pts[0] = Vertices[Order[BootstrapIndices[0]]];
+		RealType BootstrapOrient = 0;
+		for (int32 SecondIdx = Order.Num() - 2; SecondIdx >= 0; SecondIdx--)
+		{
+			if (Pts[0] != Vertices[Order[SecondIdx]])
+			{
+				Pts[1] = Vertices[Order[SecondIdx]];
+				BootstrapIndices[1] = SecondIdx;
+				break;
+			}
+		}
+		if (BootstrapIndices[1] == -1) // all points were identical; nothing to triagulate
+		{
+			return false;
+		}
+
+		for (int32 ThirdIdx = BootstrapIndices[1] - 1; ThirdIdx >= 0; ThirdIdx--)
+		{
+			RealType Orient = ExactPredicates::Orient2<RealType>(Pts[0], Pts[1], Vertices[Order[ThirdIdx]]);
+			if (Orient != 0)
+			{
+				Pts[2] = Vertices[Order[ThirdIdx]];
+				BootstrapIndices[2] = ThirdIdx;
+				BootstrapOrient = Orient;
+				break;
+			}
+		}
+		if (BootstrapIndices[2] == -1) // all points were colinear; nothing to triangulate
+		{
+			return false;
+		}
+
+		// Make the first triangle from the bootstrap points and remove the bootstrap points from the insertion ordering
+		FIndex3i FirstTri(Order[BootstrapIndices[0]], Order[BootstrapIndices[1]], Order[BootstrapIndices[2]]);
+		if (BootstrapOrient < 0)
+		{
+			Swap(FirstTri.B, FirstTri.C);
+		}
+		Connectivity.InitWithGhosts(FirstTri);
+		Order.RemoveAt(BootstrapIndices[0]);
+		Order.RemoveAt(BootstrapIndices[1]);
+		Order.RemoveAt(BootstrapIndices[2]);
+
+		FIndex3i SearchTri = FirstTri;
+		for (int32 OrderIdx = 0; OrderIdx < Order.Num(); OrderIdx++)
+		{
+			int32 Vertex = Order[OrderIdx];
+			bool bIsDuplicate = false;
+			constexpr bool bAssumeDelaunay = true; // initial construction, before constraint edges, so safe to assume Delaunay
+			FIndex3i ContainingTri = WalkToContainingTri<RealType>(Random, Connectivity, Vertices, SearchTri, Vertex, bAssumeDelaunay, bIsDuplicate);
+			if (bIsDuplicate || ContainingTri[0] == FDelaunay2Connectivity::InvalidIndex)
+			{
+				continue;
+			}
+			SearchTri = Insert<RealType>(Connectivity, Vertices, ContainingTri, Vertex);
+			checkSlow(SearchTri.A != FDelaunay2Connectivity::InvalidIndex);
+		}
+
+		// TODO: detect edge insertion failures and report them back here as well
+		ConstrainEdges(Random, Connectivity, Vertices, Edges, bKeepFastEdgeAdjacencyData);
+
+		return true;
+	}
+}
+
+bool FDelaunay2::Triangulate(TArrayView<const FVector2d> Vertices, TArrayView<const FIndex2i> Edges)
+{
+	Connectivity = MakePimpl<FDelaunay2Connectivity>();
+
+	bIsConstrained = Edges.Num() > 0;
+
+	return DelaunayInternal::Triangulate<double>(RandomStream, *Connectivity, Vertices, Edges, bKeepFastEdgeAdjacencyData);
+}
+
+bool FDelaunay2::Triangulate(TArrayView<const FVector2f> Vertices, TArrayView<const FIndex2i> Edges)
+{
+	Connectivity = MakePimpl<FDelaunay2Connectivity>();
+
+	bIsConstrained = Edges.Num() > 0;
+
+	return DelaunayInternal::Triangulate<float>(RandomStream, *Connectivity, Vertices, Edges, bKeepFastEdgeAdjacencyData);
+}
+
+void FDelaunay2::ConstrainEdges(TArrayView<const FVector2d> Vertices, TArrayView<const FIndex2i> Edges)
+{
+	bIsConstrained = bIsConstrained || Edges.Num() > 0;
+
+	DelaunayInternal::ConstrainEdges<double>(RandomStream, *Connectivity, Vertices, Edges, bKeepFastEdgeAdjacencyData);
+}
+
+void FDelaunay2::ConstrainEdges(TArrayView<const FVector2f> Vertices, TArrayView<const FIndex2i> Edges)
+{
+	bIsConstrained = bIsConstrained || Edges.Num() > 0;
+
+	DelaunayInternal::ConstrainEdges<float>(RandomStream, *Connectivity, Vertices, Edges, bKeepFastEdgeAdjacencyData);
+}
+
+TArray<FIndex3i> FDelaunay2::GetTriangles() const
+{
+	if (Connectivity.IsValid())
+	{
+		return Connectivity->GetTriangles();
+	}
+	return TArray<FIndex3i>();
+}
+
+void FDelaunay2::GetTrianglesAndAdjacency(TArray<FIndex3i>& Triangles, TArray<FIndex3i>& Adjacency) const
+{
+	if (Connectivity.IsValid())
+	{
+		Connectivity->GetTrianglesAndAdjacency(Triangles, Adjacency);
+	}
+}
+
+bool FDelaunay2::IsDelaunay(TArrayView<const FVector2f> Vertices) const
+{
+	return DelaunayInternal::IsDelaunay<float>(*Connectivity, Vertices);
+}
+
+bool FDelaunay2::IsDelaunay(TArrayView<const FVector2d> Vertices) const
+{
+	return DelaunayInternal::IsDelaunay<double>(*Connectivity, Vertices);
+}
+
+} // end namespace UE::Geometry
+} // end namespace UE
\ No newline at end of file
diff --git a/Engine/Source/Runtime/GeometryCore/Public/CompGeom/Delaunay2.h b/Engine/Source/Runtime/GeometryCore/Public/CompGeom/Delaunay2.h
new file mode 100644
--- /dev/null
+++ b/Engine/Source/Runtime/GeometryCore/Public/CompGeom/Delaunay2.h
@@ -0,0 +1,90 @@
+// Copyright Epic Games, Inc. All Rights Reserved.
+
+#pragma once
+
+#include "VectorTypes.h"
+#include "IndexTypes.h"
+
+#include "CoreMinimal.h"
+#include "Templates/PimplPtr.h"
+
+#include "Math/RandomStream.h"
+
+#include "MathUtil.h"
+#include "IndexTypes.h"
+#include "PlaneTypes.h"
+#include "LineTypes.h"
+#include "Polygon2.h"
+
+
+namespace UE {
+namespace Geometry {
+
+using namespace UE::Math;
+
+// Internal representation of mesh connectivity; not exposed to interface
+struct FDelaunay2Connectivity;
+
+class GEOMETRYCORE_API FDelaunay2
+{
+public:
+	//
+	// Inputs
+	//
+
+	// Source for random permutations, used internally in the triangulation algorithm
+	FRandomStream RandomStream;
+
+	// Option to keep extra vertex->edge adjacency data; useful if you will call ConstrainEdges many times on the same triangulation
+	bool bKeepFastEdgeAdjacencyData = false;
+
+	// TODO: it would often be useful to pass in sparse vertex data
+	//// Optional function to allow Triangulate to skip vertices
+	//TFunction<bool(int32)> SkipVertexFn = nullptr;
+
+	/**
+	 * Compute an (optionally constrained) Delaunay triangulation
+	 *
+	 * @return false if triangulation failed
+	 */
+	bool Triangulate(TArrayView<const TVector2<double>> Vertices, TArrayView<const FIndex2i> Edges = TArrayView<const FIndex2i>());
+	bool Triangulate(TArrayView<const TVector2<float>> Vertices, TArrayView<const FIndex2i> Edges = TArrayView<const FIndex2i>());
+
+	/**
+	 * Update an already-computed triangulation so the given edges are in the triangulation.
+	 * Note: Assumes the edges do not intersect other constrained edges OR existing vertices in the triangulation
+	 * TODO: track at least whether any easy-to-detect failures occur (at least an edge intersecting a vertex and never being inserted)
+	 */
+	void ConstrainEdges(TArrayView<const TVector2<double>> Vertices, TArrayView<const FIndex2i> Edges);
+	void ConstrainEdges(TArrayView<const TVector2<float>> Vertices, TArrayView<const FIndex2i> Edges);
+
+	// TODO: Support incremental vertex insertion
+	// Update the triangulation incrementally, assuming Vertices are unchanged before FirstNewIndex, and nothing after FirstNewIndex has been inserted yet
+	// Note that updating with new vertices *after* constraining edges may remove previously-constrained edges, unless we also add a way to tag constrained edges
+	// bool Update(TArrayView<const TVector2<double>> Vertices, int32 FirstNewIdx);
+
+	// Get the triangulation as an array of triangles
+	// Note: This creates a new array each call, because the internal data structure does not have a triangle array
+	TArray<FIndex3i> GetTriangles() const;
+
+	// Get the triangulation as an array with a corresponding adjacency array, indicating the adjacent triangle on each triangle edge (-1 if no adjacent triangle)
+	void GetTrianglesAndAdjacency(TArray<FIndex3i>& Triangles, TArray<FIndex3i>& Adjacency) const;
+
+	// @return true if this is a constrained Delaunay triangulation
+	bool IsConstrained() const
+	{
+		return bIsConstrained;
+	}
+
+	// @return true if triangulation is Delaunay, useful for validating results (note: likely to be false if edges are constrained)
+	bool IsDelaunay(TArrayView<const FVector2f> Vertices) const;
+	bool IsDelaunay(TArrayView<const FVector2d> Vertices) const;
+
+protected:
+	TPimplPtr<FDelaunay2Connectivity> Connectivity;
+
+	bool bIsConstrained = false;
+};
+
+} // end namespace UE::Geometry
+} // end namespace UE
