Fix convex hull crashes (due to zero points and due to exact predicate failure)
add geometry processing low level test coverage for both cases

diff --git a/Engine/Source/Runtime/GeometryCore/Private/CompGeom/ConvexHull3.cpp b/Engine/Source/Runtime/GeometryCore/Private/CompGeom/ConvexHull3.cpp
--- a/Engine/Source/Runtime/GeometryCore/Private/CompGeom/ConvexHull3.cpp
+++ b/Engine/Source/Runtime/GeometryCore/Private/CompGeom/ConvexHull3.cpp
@@ -21,11 +21,12 @@ void TExtremePoints3<RealType>::Init(int32 NumPoints, TFunctionRef<TVector<RealT
 {
 	TVector<RealType> FirstPoint;
 	int FirstPtIdx = -1;
-	for (FirstPtIdx = 0; FirstPtIdx < NumPoints; FirstPtIdx++)
+	for (int FirstPtIdxTest = 0; FirstPtIdxTest < NumPoints; FirstPtIdxTest++)
 	{
-		if (FilterFunc(FirstPtIdx))
+		if (FilterFunc(FirstPtIdxTest))
 		{
-			FirstPoint = GetPointFunc(FirstPtIdx);
+			FirstPoint = GetPointFunc(FirstPtIdxTest);
+			FirstPtIdx = FirstPtIdxTest;
 			break;
 		}
 	}
