update new Delaunay code to be more usable:
 - Report failures to constrain edges, optionally detecting the case where the edge was flipped away by a subsequent constraint (happens if the constraint edges intersect)
 - Add options to directly triangulation TPolygon2 and TGeneralPolygon2 (+ the supporting fill algorithms to be run on any array of edges)
 - Add validation of Delaunay property for CDTs on all non-constrained edges
 - Update edge cache (if enabled) on sub-triangulation append, improving perf on CDTs with many constrained edges

add corresponding test coverage, and change the test runner to always automatically run exact predicates globalinit

diff --git a/Engine/Source/Runtime/GeometryCore/Private/CompGeom/Delaunay2.cpp b/Engine/Source/Runtime/GeometryCore/Private/CompGeom/Delaunay2.cpp
--- a/Engine/Source/Runtime/GeometryCore/Private/CompGeom/Delaunay2.cpp
+++ b/Engine/Source/Runtime/GeometryCore/Private/CompGeom/Delaunay2.cpp
@@ -27,9 +27,14 @@ struct FDelaunay2Connectivity
 	// Add the triangles from another mesh directly to this one
 	void Append(const FDelaunay2Connectivity& ToAdd)
 	{
+		EdgeToVert.Reserve(EdgeToVert.Num() + ToAdd.EdgeToVert.Num());
 		for (const TPair<FIndex2i, int32>& EdgeV : ToAdd.EdgeToVert)
 		{
 			EdgeToVert.Add(EdgeV.Key, EdgeV.Value);
+			if (bUseAdjCache)
+			{
+				UpdateAdjCache(EdgeV.Key);
+			}
 		}
 	}
 
@@ -125,6 +130,170 @@ struct FDelaunay2Connectivity
 		}
 	}
 
+	// return the triangle indices rotated s.t. the smallest vertex is first (and the winding is unchanged)
+	inline FIndex3i AsUniqueTriangle(const FIndex2i& Edge, int32 Vertex)
+	{
+		if (Edge.A < Edge.B)
+		{
+			if (Vertex < Edge.A)
+			{
+				return FIndex3i(Vertex, Edge.A, Edge.B);
+			}
+			else
+			{
+				return FIndex3i(Edge.A, Edge.B, Vertex);
+			}
+		}
+		else
+		{
+			if (Vertex < Edge.B)
+			{
+				return FIndex3i(Vertex, Edge.A, Edge.B);
+			}
+			else
+			{
+				return FIndex3i(Edge.B, Vertex, Edge.A);
+			}
+		}
+	}
+
+	TArray<FIndex3i> GetFilledTriangles(TArrayView<const FIndex2i> Edges, FDelaunay2::EFillMode FillMode)
+	{
+		TMap<FIndex2i, int> EdgeDelta;
+		EdgeDelta.Reserve(Edges.Num());
+		for (const FIndex2i& Edge : Edges)
+		{
+			FIndex2i Unoriented = Edge;
+			int32 Sign = 1;
+			if (Unoriented.A > Unoriented.B)
+			{
+				Swap(Unoriented.A, Unoriented.B);
+				if (FillMode != FDelaunay2::EFillMode::Solid)
+				{
+					Sign = -1;
+				}
+			}
+			int32& Delta = EdgeDelta.FindOrAdd(Unoriented, 0);
+			Delta += Sign;
+		}
+
+
+		TMap<FIndex3i, int32> Windings;
+		struct FWalk
+		{
+			FIndex2i Edge;
+			int32 Winding;
+		};
+		TArray<FWalk> ToWalk;
+
+		auto GetWindingChange = [&EdgeDelta](FIndex2i Edge)
+		{
+			int32 Winding = 0;
+			int32 Sign = -1;
+			if (Edge.A > Edge.B)
+			{
+				Sign = 1;
+				Swap(Edge.A, Edge.B);
+			}
+			int32* Change = EdgeDelta.Find(Edge);
+			if (!Change)
+			{
+				return 0;
+			}
+			else
+			{
+				return Sign * (*Change);
+			}
+		};
+		auto WalkEdge = [FillMode, &EdgeDelta, &ToWalk, &GetWindingChange](const FIndex2i& Edge, int32 BeforeWinding)
+		{
+			if (FillMode == FDelaunay2::EFillMode::Solid)
+			{
+				FIndex2i Unoriented = Edge;
+				if (Unoriented.A > Unoriented.B)
+				{
+					Swap(Unoriented.A, Unoriented.B);
+				}
+				if (!EdgeDelta.Contains(Unoriented)) // in solid fill, only check if the edge is present
+				{
+					ToWalk.Add({ FIndex2i(Edge.B, Edge.A), 1 });
+				}
+			}
+			else
+			{
+				int32 Winding = BeforeWinding + GetWindingChange(Edge);
+				ToWalk.Add({ FIndex2i(Edge.B, Edge.A), Winding });
+			}
+		};
+
+		// First walk all the ghost triangles, through their non-ghost edge
+		int32 GhostTriCount = 0;
+		for (const TPair<FIndex2i, int32>& EdgeV : EdgeToVert)
+		{
+			if (EdgeV.Value == GhostIndex)
+			{
+				WalkEdge(EdgeV.Key, 0);
+				GhostTriCount++;
+			}
+		}
+
+		int32 PickIdx = 0;
+		while (ToWalk.Num())
+		{
+			// Traverse the walk list in something similar to a breadth-first traversal
+			// Non-closed shapes in the input edges will give different results depending on visit order
+			//  and depth-first traversal would be more sensitive to starting triangle in that case
+			PickIdx = (PickIdx + 1) % ToWalk.Num();
+			const FWalk Walk = ToWalk[PickIdx];
+			ToWalk.RemoveAtSwap(PickIdx);
+			int32 Vert = EdgeToVert[Walk.Edge];
+			FIndex3i UniqueTri = AsUniqueTriangle(Walk.Edge, Vert);
+			if (UniqueTri.A < 0) // it's a ghost
+			{
+				continue;
+			}
+			int32* Winding = Windings.Find(UniqueTri);
+			if (Winding)
+			{
+				continue;
+			}
+			Windings.Add(UniqueTri, Walk.Winding);
+			WalkEdge(FIndex2i(Walk.Edge.B, Vert), Walk.Winding);
+			WalkEdge(FIndex2i(Vert, Walk.Edge.A), Walk.Winding);
+		}
+
+		TArray<FIndex3i> Triangles;
+		if (FillMode == FDelaunay2::EFillMode::Solid)
+		{
+			Triangles.Reserve(EdgeToVert.Num() / 3 - GhostTriCount - Windings.Num());
+			// everything we *didn't* walk over is included
+			EnumerateTriangles([&Windings, &Triangles](const FIndex2i& Edge, int32 Vertex) -> bool
+				{
+					FIndex3i Triangle(Vertex, Edge.A, Edge.B);
+					if (!Windings.Contains(Triangle))
+					{
+						Triangles.Add(Triangle);
+					}
+					return true;
+				}, true /*bSkipGhosts*/);
+		}
+		else
+		{
+			Triangles.Reserve(Windings.Num());
+			for (const TPair<FIndex3i, int32>& TriWinding : Windings)
+			{
+				if ((FillMode == FDelaunay2::EFillMode::NegativeWinding	&& TriWinding.Value < 0) ||
+					(FillMode == FDelaunay2::EFillMode::NonZeroWinding	&& TriWinding.Value != 0) ||
+					(FillMode == FDelaunay2::EFillMode::PositiveWinding && TriWinding.Value > 0) ||
+					(FillMode == FDelaunay2::EFillMode::OddWinding		&& (TriWinding.Value % 2) != 0))
+				{
+					Triangles.Add(TriWinding.Key);
+				}
+			}
+		}
+		return Triangles;
+	}
+
 	static bool IsGhost(const FIndex2i& Edge, int32 Vertex)
 	{
 		return Edge.A == GhostIndex || Edge.B == GhostIndex || Vertex == GhostIndex;
@@ -237,7 +406,9 @@ struct FDelaunay2Connectivity
 	}
 
 	// Similar to EnumerateOrientedEdges but only visits each triangle once, instead of 3x
-	// and optionally skips ghost triangles (triangles connected to the ghost vertex)
+	// Calls VisitFn(FIndex2i Edge, int32 Vertex) with Vertex always smaller than the Edge.A and Edge.B
+	// Returning false from VisitFn will end the enumeration early
+	// @param bSkipGhosts If true, skip ghost triangles (triangles connected to the ghost vertex)
 	template<typename VisitFunctionType>
 	void EnumerateTriangles(VisitFunctionType VisitFn, bool bSkipGhosts = false)
 	{
@@ -482,11 +653,34 @@ namespace DelaunayInternal
 	}
 
 	template<typename RealType>
-	bool IsDelaunay(FDelaunay2Connectivity& Connectivity, TArrayView<const TVector2<RealType>> Vertices)
+	bool IsDelaunay(FDelaunay2Connectivity& Connectivity, TArrayView<const TVector2<RealType>> Vertices, TArrayView<const FIndex2i> SkipEdgesIn)
 	{
+		TSet<FIndex2i> SkipSet;
+		SkipSet.Reserve(SkipEdgesIn.Num());
+		for (FIndex2i Edge : SkipEdgesIn)
+		{
+			if (Edge.A > Edge.B)
+			{
+				Swap(Edge.A, Edge.B);
+			}
+			SkipSet.Add(Edge);
+		}
+
 		bool bFoundNonDelaunay = false;
-		Connectivity.EnumerateOrientedEdges([&Connectivity, &Vertices, &bFoundNonDelaunay](const FIndex2i& Edge, int32 Vert) -> bool
+		Connectivity.EnumerateOrientedEdges([&Connectivity, &Vertices, &SkipSet, &bFoundNonDelaunay](const FIndex2i& Edge, int32 Vert) -> bool
 		{
+			if (SkipSet.Num())
+			{
+				FIndex2i Unoriented = Edge;
+				if (Unoriented.A > Unoriented.B)
+				{
+					Swap(Unoriented.A, Unoriented.B);
+				}
+				if (SkipSet.Contains(Unoriented))
+				{
+					return true;
+				}
+			}
 			if (Connectivity.IsGhost(Edge, Vert))
 			{
 				return true;
@@ -579,8 +773,10 @@ namespace DelaunayInternal
 
 	// @return vertex we need to fill to.  If EdgeToConnect.A, no fill needed; if EdgeToConnect.B, a normal re-triangulation needed; if other, digging failed in the middle and we need to fill partially
 	template<typename RealType>
-	int32 DigCavity(FDelaunay2Connectivity& Connectivity, TArrayView<const TVector2<RealType>> Vertices, const FIndex2i& EdgeToConnect, TArray<int32>& CavityLOut, TArray<int32>& CavityROut)
+	int32 DigCavity(FDelaunay2Connectivity& Connectivity, TArrayView<const TVector2<RealType>> Vertices, const FIndex2i& EdgeToConnect, TArray<int32>& CavityLOut, TArray<int32>& CavityROut, bool& bDigSuccess)
 	{
+		bDigSuccess = false;
+
 		CavityLOut.Reset();
 		CavityROut.Reset();
 
@@ -588,6 +784,7 @@ namespace DelaunayInternal
 		bool bFoundCross = GetFirstCrossingEdge<RealType>(Connectivity, Vertices, EdgeToConnect, FirstCross);
 		if (!bFoundCross)
 		{
+			bDigSuccess = true; // no-dig-needed case counts as a success
 			return EdgeToConnect.A;
 		}
 
@@ -619,6 +816,7 @@ namespace DelaunayInternal
 			{
 				CavityROut.Add(EdgeToConnect.B);
 				CavityLOut.Add(EdgeToConnect.B);
+				bDigSuccess = true;
 				return EdgeToConnect.B;
 			}
 
@@ -739,7 +937,7 @@ namespace DelaunayInternal
 	}
 
 	template<typename RealType>
-	void ConstrainEdges(FRandomStream& Random, FDelaunay2Connectivity& Connectivity,
+	bool ConstrainEdges(FRandomStream& Random, FDelaunay2Connectivity& Connectivity,
 		TArrayView<const TVector2<RealType>> Vertices, TArrayView<const FIndex2i> Edges, bool bKeepFastEdgeAdjacencyData)
 	{
 		constexpr int32 NeedFasterEdgeLookupThreshold = 4; // TODO: do some profiling to determine what this threshold should be
@@ -748,6 +946,8 @@ namespace DelaunayInternal
 			Connectivity.EnableVertexAdjacency(Vertices.Num());
 		}
 
+		bool bSuccess = true;
+
 		TArray<int32> CavityVerts[2];
 
 		// Random insertion order to improve expected performance
@@ -756,7 +956,8 @@ namespace DelaunayInternal
 		{
 			int32 EdgeIdx = EdgeOrder[OrderIdx];
 			FIndex2i Edge = Edges[EdgeIdx];
-			int32 DigTo = DigCavity<RealType>(Connectivity, Vertices, Edge, CavityVerts[0], CavityVerts[1]);
+			bool bDigSucceeded = false;
+			int32 DigTo = DigCavity<RealType>(Connectivity, Vertices, Edge, CavityVerts[0], CavityVerts[1], bDigSucceeded);
 			if (DigTo != Edge.A)
 			{
 				FIndex2i DugEdge(Edge.A, DigTo); // fill the cavity we dug out (which may end at a different vertex than target, if there was a colinear vertex first)
@@ -771,6 +972,8 @@ namespace DelaunayInternal
 		{
 			Connectivity.DisableVertexAdjacency();
 		}
+
+		return bSuccess;
 	}
 
 	template<typename RealType>
@@ -849,10 +1052,7 @@ namespace DelaunayInternal
 			checkSlow(SearchTri.A != FDelaunay2Connectivity::InvalidIndex);
 		}
 
-		// TODO: detect edge insertion failures and report them back here as well
-		ConstrainEdges(Random, Connectivity, Vertices, Edges, bKeepFastEdgeAdjacencyData);
-
-		return true;
+		return ConstrainEdges(Random, Connectivity, Vertices, Edges, bKeepFastEdgeAdjacencyData);
 	}
 }
 
@@ -862,7 +1062,8 @@ bool FDelaunay2::Triangulate(TArrayView<const FVector2d> Vertices, TArrayView<co
 
 	bIsConstrained = Edges.Num() > 0;
 
-	return DelaunayInternal::Triangulate<double>(RandomStream, *Connectivity, Vertices, Edges, bKeepFastEdgeAdjacencyData);
+	bool bSuccess = DelaunayInternal::Triangulate<double>(RandomStream, *Connectivity, Vertices, Edges, bKeepFastEdgeAdjacencyData);
+	return bSuccess && ValidateResult(Edges);
 }
 
 bool FDelaunay2::Triangulate(TArrayView<const FVector2f> Vertices, TArrayView<const FIndex2i> Edges)
@@ -871,21 +1072,24 @@ bool FDelaunay2::Triangulate(TArrayView<const FVector2f> Vertices, TArrayView<co
 
 	bIsConstrained = Edges.Num() > 0;
 
-	return DelaunayInternal::Triangulate<float>(RandomStream, *Connectivity, Vertices, Edges, bKeepFastEdgeAdjacencyData);
+	bool bSuccess = DelaunayInternal::Triangulate<float>(RandomStream, *Connectivity, Vertices, Edges, bKeepFastEdgeAdjacencyData);
+	return bSuccess && ValidateResult(Edges);
 }
 
-void FDelaunay2::ConstrainEdges(TArrayView<const FVector2d> Vertices, TArrayView<const FIndex2i> Edges)
+bool FDelaunay2::ConstrainEdges(TArrayView<const FVector2d> Vertices, TArrayView<const FIndex2i> Edges)
 {
 	bIsConstrained = bIsConstrained || Edges.Num() > 0;
 
-	DelaunayInternal::ConstrainEdges<double>(RandomStream, *Connectivity, Vertices, Edges, bKeepFastEdgeAdjacencyData);
+	bool bSuccess = DelaunayInternal::ConstrainEdges<double>(RandomStream, *Connectivity, Vertices, Edges, bKeepFastEdgeAdjacencyData);
+	return bSuccess && ValidateResult(Edges);
 }
 
-void FDelaunay2::ConstrainEdges(TArrayView<const FVector2f> Vertices, TArrayView<const FIndex2i> Edges)
+bool FDelaunay2::ConstrainEdges(TArrayView<const FVector2f> Vertices, TArrayView<const FIndex2i> Edges)
 {
 	bIsConstrained = bIsConstrained || Edges.Num() > 0;
 
-	DelaunayInternal::ConstrainEdges<float>(RandomStream, *Connectivity, Vertices, Edges, bKeepFastEdgeAdjacencyData);
+	bool bSuccess = DelaunayInternal::ConstrainEdges<float>(RandomStream, *Connectivity, Vertices, Edges, bKeepFastEdgeAdjacencyData);
+	return bSuccess && ValidateResult(Edges);
 }
 
 TArray<FIndex3i> FDelaunay2::GetTriangles() const
@@ -905,14 +1109,47 @@ void FDelaunay2::GetTrianglesAndAdjacency(TArray<FIndex3i>& Triangles, TArray<FI
 	}
 }
 
-bool FDelaunay2::IsDelaunay(TArrayView<const FVector2f> Vertices) const
+TArray<FIndex3i> FDelaunay2::GetFilledTriangles(TArrayView<const FIndex2i> Edges, EFillMode FillMode)
+{
+	if (!ensure(Connectivity.IsValid()))
+	{
+		return TArray<FIndex3i>();
+	}
+	return Connectivity->GetFilledTriangles(Edges, FillMode);
+}
+
+bool FDelaunay2::IsDelaunay(TArrayView<const FVector2f> Vertices, TArrayView<const FIndex2i> SkipEdges) const
 {
-	return DelaunayInternal::IsDelaunay<float>(*Connectivity, Vertices);
+	if (ensure(Connectivity.IsValid()))
+	{
+		return DelaunayInternal::IsDelaunay<float>(*Connectivity, Vertices, SkipEdges);
+	}
+	return false; // if no triangulation was performed, function should be been called
+}
+
+bool FDelaunay2::IsDelaunay(TArrayView<const FVector2d> Vertices, TArrayView<const FIndex2i> SkipEdges) const
+{
+	if (ensure(Connectivity.IsValid()))
+	{
+		return DelaunayInternal::IsDelaunay<double>(*Connectivity, Vertices, SkipEdges);
+	}
+	return false; // if no triangulation was performed, function should not be called
 }
 
-bool FDelaunay2::IsDelaunay(TArrayView<const FVector2d> Vertices) const
+bool FDelaunay2::HasEdges(TArrayView<const FIndex2i> Edges) const
 {
-	return DelaunayInternal::IsDelaunay<double>(*Connectivity, Vertices);
+	if (ensure(Connectivity.IsValid()))
+	{
+		for (const FIndex2i& Edge : Edges)
+		{
+			if (!Connectivity->HasEdge(Edge))
+			{
+				return false;
+			}
+		}
+		return true;
+	}
+	return false; // if no triangulation was performed, function should not be called
 }
 
 } // end namespace UE::Geometry
diff --git a/Engine/Source/Runtime/GeometryCore/Public/CompGeom/Delaunay2.h b/Engine/Source/Runtime/GeometryCore/Public/CompGeom/Delaunay2.h
--- a/Engine/Source/Runtime/GeometryCore/Public/CompGeom/Delaunay2.h
+++ b/Engine/Source/Runtime/GeometryCore/Public/CompGeom/Delaunay2.h
@@ -15,6 +15,7 @@
 #include "PlaneTypes.h"
 #include "LineTypes.h"
 #include "Polygon2.h"
+#include "Curve/GeneralPolygon2.h"
 
 
 namespace UE {
@@ -28,6 +29,20 @@ struct FDelaunay2Connectivity;
 class GEOMETRYCORE_API FDelaunay2
 {
 public:
+
+	// Options for selecting what triangles to include in the output, for constrained Delaunay triangulation of polygons
+	enum class EFillMode
+	{
+		// Fastest/simplest option: Keep all triangles if you'd have to cross a constrained edge to reach it from outside, ignoring edge orientation
+		Solid,
+		// Winding-number based fill options: Keep triangles based on the winding number
+		PositiveWinding,
+		NonZeroWinding,
+		NegativeWinding,
+		OddWinding
+	};
+
+
 	//
 	// Inputs
 	//
@@ -38,25 +53,99 @@ public:
 	// Option to keep extra vertex->edge adjacency data; useful if you will call ConstrainEdges many times on the same triangulation
 	bool bKeepFastEdgeAdjacencyData = false;
 
+	// Option to validate that edges remain in the triangulation after multiple constraint edges passed in
+	// Only validates for the edges in the current call; if separate calls constrain additional edges, set this to 'false' and call HasEdges() on the full edge set
+	// TODO: Consider having the internal mesh remember what edges have been constrained, so we can set an error flag when constrained edges cross previous constrained edges
+	bool bValidateEdges = true;
+
 	// TODO: it would often be useful to pass in sparse vertex data
 	//// Optional function to allow Triangulate to skip vertices
 	//TFunction<bool(int32)> SkipVertexFn = nullptr;
 
 	/**
-	 * Compute an (optionally constrained) Delaunay triangulation
+	 * Compute an (optionally constrained) Delaunay triangulation.
+	 * Note this clears any previously-held triangulation data, and triangulates the passed-in vertices (and optional edges) from scratch
 	 *
 	 * @return false if triangulation failed
 	 */
 	bool Triangulate(TArrayView<const TVector2<double>> Vertices, TArrayView<const FIndex2i> Edges = TArrayView<const FIndex2i>());
 	bool Triangulate(TArrayView<const TVector2<float>> Vertices, TArrayView<const FIndex2i> Edges = TArrayView<const FIndex2i>());
 
+	// Triangulate a polygon, and optionally pass back the resulting triangles
+	// Uses the 'solid' fill mode, and fills the polygon regardless of the edge orientation
+	// Note: TrianglesOut may be empty or incomplete if the input is self-intersecting.
+	// Note this clears any previously-held triangulation data, and triangulates the passed-in polygon from scratch
+	template<typename RealType>
+	bool Triangulate(const TPolygon2<RealType>& Polygon, TArray<FIndex3i>* TrianglesOut = nullptr)
+	{
+		int32 NumVertices = Polygon.VertexCount();
+		TArray<FIndex2i> Edges;
+		Edges.Reserve(NumVertices - 1);
+		for (int32 Last = NumVertices - 1, Idx = 0; Idx < NumVertices; Last = Idx++)
+		{
+			Edges.Add(FIndex2i(Last, Idx));
+		}
+		bool bSuccess = Triangulate(Polygon.GetVertices(), Edges);
+		if (TrianglesOut)
+		{
+			*TrianglesOut = GetFilledTriangles(Edges, EFillMode::Solid);
+		}
+		return bSuccess;
+	}
+
+	// Triangulate a polygon, and optionally pass back the resulting triangles
+	// Uses a winding-number-based fill mode, and relies on the general polygon having correct orientations
+	// (Uses TGeneralPolygon2's OuterIsClockwise() to automatically choose between positive or negative winding)
+	// Note: TrianglesOut may be empty or incomplete if the input is self-intersecting.
+	// Note this clears any previously-held triangulation data, and triangulates the passed-in general polygon from scratch
+	template<typename RealType>
+	bool Triangulate(const TGeneralPolygon2<RealType>& GeneralPolygon, TArray<FIndex3i>* TrianglesOut = nullptr, TArray<TVector2<RealType>>* VerticesOut = nullptr)
+	{
+		TArray<TVector2<RealType>> AllVertices;
+		TArray<FIndex2i> AllEdges;
+
+		auto AppendVertices = [&AllVertices, &AllEdges](const TArray<TVector2<RealType>>& Vertices)
+		{
+			if (Vertices.IsEmpty())
+			{
+				return;
+			}
+
+			int32 StartIdx = AllVertices.Num();
+			AllVertices.Append(Vertices);
+			AllEdges.Reserve(AllEdges.Num() + Vertices.Num() - 1);
+			int32 NumVertices = Vertices.Num();
+			for (int32 Last = NumVertices - 1, Idx = 0; Idx < NumVertices; Last = Idx++)
+			{
+				AllEdges.Add(FIndex2i(StartIdx + Last, StartIdx + Idx));
+			}
+		};
+		AppendVertices(GeneralPolygon.GetOuter().GetVertices());
+		for (int32 HoleIdx = 0; HoleIdx < GeneralPolygon.GetHoles().Num(); HoleIdx++)
+		{
+			AppendVertices(GeneralPolygon.GetHoles()[HoleIdx].GetVertices());
+		}
+		
+		bool bSuccess = Triangulate(AllVertices, AllEdges);
+		if (TrianglesOut)
+		{
+			*TrianglesOut = GetFilledTriangles(AllEdges, GeneralPolygon.OuterIsClockwise() ? EFillMode::NegativeWinding : EFillMode::PositiveWinding);
+		}
+		if (VerticesOut)
+		{
+			*VerticesOut = MoveTemp(AllVertices);
+		}
+		return bSuccess;
+	}
+
 	/**
 	 * Update an already-computed triangulation so the given edges are in the triangulation.
 	 * Note: Assumes the edges do not intersect other constrained edges OR existing vertices in the triangulation
-	 * TODO: track at least whether any easy-to-detect failures occur (at least an edge intersecting a vertex and never being inserted)
+	 * @return true if edges were successfully inserted.
+	 * Note: Will not detect failure due to intersection of constrained edges unless member bValidateEdges == true.
 	 */
-	void ConstrainEdges(TArrayView<const TVector2<double>> Vertices, TArrayView<const FIndex2i> Edges);
-	void ConstrainEdges(TArrayView<const TVector2<float>> Vertices, TArrayView<const FIndex2i> Edges);
+	bool ConstrainEdges(TArrayView<const TVector2<double>> Vertices, TArrayView<const FIndex2i> Edges);
+	bool ConstrainEdges(TArrayView<const TVector2<float>> Vertices, TArrayView<const FIndex2i> Edges);
 
 	// TODO: Support incremental vertex insertion
 	// Update the triangulation incrementally, assuming Vertices are unchanged before FirstNewIndex, and nothing after FirstNewIndex has been inserted yet
@@ -70,6 +159,12 @@ public:
 	// Get the triangulation as an array with a corresponding adjacency array, indicating the adjacent triangle on each triangle edge (-1 if no adjacent triangle)
 	void GetTrianglesAndAdjacency(TArray<FIndex3i>& Triangles, TArray<FIndex3i>& Adjacency) const;
 
+	/**
+	 * Return the triangles that are inside the given edges, removing the outer boundary triangles
+	 * If a Winding-Number-based fill mode is used, assumes edges are oriented and tracks the winding number across edges
+	 */
+	TArray<FIndex3i> GetFilledTriangles(TArrayView<const FIndex2i> Edges, EFillMode FillMode = EFillMode::PositiveWinding);
+
 	// @return true if this is a constrained Delaunay triangulation
 	bool IsConstrained() const
 	{
@@ -77,13 +172,22 @@ public:
 	}
 
 	// @return true if triangulation is Delaunay, useful for validating results (note: likely to be false if edges are constrained)
-	bool IsDelaunay(TArrayView<const FVector2f> Vertices) const;
-	bool IsDelaunay(TArrayView<const FVector2d> Vertices) const;
+	bool IsDelaunay(TArrayView<const FVector2f> Vertices, TArrayView<const FIndex2i> SkipEdges = TArrayView<const FIndex2i>()) const;
+	bool IsDelaunay(TArrayView<const FVector2d> Vertices, TArrayView<const FIndex2i> SkipEdges = TArrayView<const FIndex2i>()) const;
+
+	// @return true if triangulation has the given edges, useful for validating results
+	bool HasEdges(TArrayView<const FIndex2i> Edges) const;
 
 protected:
 	TPimplPtr<FDelaunay2Connectivity> Connectivity;
 
 	bool bIsConstrained = false;
+
+	// helper to perform standard validation on results after Triangulate or ConstrainEdges calls
+	bool ValidateResult(TArrayView<const FIndex2i> Edges) const
+	{
+		return !bValidateEdges || HasEdges(Edges);
+	}
 };
 
 } // end namespace UE::Geometry
diff --git a/Engine/Source/Runtime/GeometryCore/Public/Curve/GeneralPolygon2.h b/Engine/Source/Runtime/GeometryCore/Public/Curve/GeneralPolygon2.h
--- a/Engine/Source/Runtime/GeometryCore/Public/Curve/GeneralPolygon2.h
+++ b/Engine/Source/Runtime/GeometryCore/Public/Curve/GeneralPolygon2.h
@@ -206,8 +206,12 @@ public:
 		}
     }
 
+	bool OuterIsClockwise() const
+	{
+		return bOuterIsCW;
+	}
 
-	bool Contains(TVector2<T> vTest)
+	bool Contains(TVector2<T> vTest) const
 	{
 		if (Outer.Contains(vTest) == false)
 		{
@@ -223,7 +227,7 @@ public:
 		return true;
 	}
 
-    bool Contains(TPolygon2<T> Poly) {
+    bool Contains(TPolygon2<T> Poly) const {
 		if (Outer.Contains(Poly) == false)
 		{
 			return false;
@@ -239,7 +243,7 @@ public:
     }
 
 
-    bool Intersects(TPolygon2<T> Poly)
+    bool Intersects(TPolygon2<T> Poly) const
     {
 		if (Outer.Intersects(Poly))
 		{
@@ -256,7 +260,7 @@ public:
     }
 
 
-    TVector2<T> GetSegmentPoint(int iSegment, double fSegT, int iHoleIndex = -1)
+    TVector2<T> GetSegmentPoint(int iSegment, double fSegT, int iHoleIndex = -1) const
 	{
 		if (iHoleIndex == -1)
 		{
@@ -265,7 +269,7 @@ public:
 		return Holes[iHoleIndex].GetSegmentPoint(iSegment, fSegT);
 	}
 
-	TSegment2<T> Segment(int iSegment, int iHoleIndex = -1)
+	TSegment2<T> Segment(int iSegment, int iHoleIndex = -1) const
 	{
 		if (iHoleIndex == -1)
 		{
@@ -274,7 +278,7 @@ public:
 		return Holes[iHoleIndex].Segment(iSegment);			
 	}
 
-	TVector2<T> GetNormal(int iSegment, double segT, int iHoleIndex = -1)
+	TVector2<T> GetNormal(int iSegment, double segT, int iHoleIndex = -1) const
 	{
 		if (iHoleIndex == -1)
 		{
@@ -284,7 +288,7 @@ public:
 	}
 
 	// this should be more efficient when there are Holes...
-	double DistanceSquared(TVector2<T> p, int &iHoleIndex, int &iNearSeg, double &fNearSegT)
+	double DistanceSquared(TVector2<T> p, int &iHoleIndex, int &iNearSeg, double &fNearSegT) const
 	{
 		iNearSeg = iHoleIndex = -1;
 		fNearSegT = TMathUtil<T>::MaxReal;
