Expose optional second importance weight to the FPriorityOrderPoints class

diff --git a/Engine/Source/Runtime/GeometryCore/Private/Spatial/PriorityOrderPoints.cpp b/Engine/Source/Runtime/GeometryCore/Private/Spatial/PriorityOrderPoints.cpp
--- a/Engine/Source/Runtime/GeometryCore/Private/Spatial/PriorityOrderPoints.cpp
+++ b/Engine/Source/Runtime/GeometryCore/Private/Spatial/PriorityOrderPoints.cpp
@@ -32,10 +32,9 @@ private:
 };
 
 // Spatial ordering algorithm
-// Note we support two sets of importance weights because it is sometimes helpful to also sort by distance from center
-// However this feature is not yet exposed in the outer interface
+// Note we support two sets of importance weights because it is sometimes helpful to also sort by an additional metric, such as distance from center
 template<typename RealType>
-void OrderPoints(TArray<int32>& PointOrder, TArrayView<const TVector<RealType>> Points, TArrayView<TArrayView<const float>> ImportanceWeights, int32 EarlyStop, int32 SpatialLevels)
+void OrderPoints(TArray<int32>& PointOrder, TArrayView<const TVector<RealType>> Points, TArrayView<const TArrayView<const float>> ImportanceWeights, int32 EarlyStop, int32 SpatialLevels)
 {
 	const int32 NumPoints = Points.Num();
 	PointOrder.SetNumUninitialized(NumPoints);
@@ -327,6 +326,22 @@ void FPriorityOrderPoints::ComputeUniformSpaced(TArrayView<const FVector3f> Poin
 	PriorityOrderPointsLocal::OrderPoints<float>(Order, Points, ImportanceWeightsWeights, EarlyStop, SpatialLevels);
 }
 
+void FPriorityOrderPoints::ComputeUniformSpaced(TArrayView<const FVector3d> Points, TArrayView<const float> ImportanceWeights, TArrayView<const float> SecondImportanceWeights, int32 EarlyStop)
+{
+	TArray<TArrayView<const float>, TFixedAllocator<2>> MultiImportanceWeights;
+	MultiImportanceWeights.Add(ImportanceWeights);
+	MultiImportanceWeights.Add(SecondImportanceWeights);
+	PriorityOrderPointsLocal::OrderPoints<double>(Order, Points, MultiImportanceWeights, EarlyStop, SpatialLevels);
+}
+
+void FPriorityOrderPoints::ComputeUniformSpaced(TArrayView<const FVector3f> Points, TArrayView<const float> ImportanceWeights, TArrayView<const float> SecondImportanceWeights, int32 EarlyStop)
+{
+	TArray<TArrayView<const float>, TFixedAllocator<2>> MultiImportanceWeights;
+	MultiImportanceWeights.Add(ImportanceWeights);
+	MultiImportanceWeights.Add(SecondImportanceWeights);
+	PriorityOrderPointsLocal::OrderPoints<float>(Order, Points, MultiImportanceWeights, EarlyStop, SpatialLevels);
+}
+
 void FPriorityOrderPoints::ComputeDescendingImportance(TArrayView<const float> ImportanceWeights)
 {
 	int32 NumPoints = ImportanceWeights.Num();
diff --git a/Engine/Source/Runtime/GeometryCore/Public/Spatial/PriorityOrderPoints.h b/Engine/Source/Runtime/GeometryCore/Public/Spatial/PriorityOrderPoints.h
--- a/Engine/Source/Runtime/GeometryCore/Public/Spatial/PriorityOrderPoints.h
+++ b/Engine/Source/Runtime/GeometryCore/Public/Spatial/PriorityOrderPoints.h
@@ -44,11 +44,22 @@ public:
 	 * 
 	 * @param Points				Points to re-order
 	 * @param ImportanceWeights		Points with a higher ImportanceWeight in their local region of space will appear earlier in the ordering
-	 * @param EarlyStop				If >= 0, stop re-ordering one we've sorted this many points; beyond this the ordering will be arbitrary
+	 * @param EarlyStop				If >= 0, stop re-ordering once we've sorted this many points; beyond this the ordering will be arbitrary
 	 */
 	void ComputeUniformSpaced(TArrayView<const FVector3d> Points, TArrayView<const float> ImportanceWeights, int32 EarlyStop = -1);
 	void ComputeUniformSpaced(TArrayView<const FVector3f> Points, TArrayView<const float> ImportanceWeights, int32 EarlyStop = -1);
 
+	/**
+	 * Compute an ordering for the input points that attempts to keep points 'well spaced' / un-clumped
+	 *
+	 * @param Points					Points to re-order
+	 * @param ImportanceWeights			Points with a higher ImportanceWeight in their local region of space will appear earlier in the ordering
+	 * @param SecondImportanceWeights	A second set of importance weights; we may pick an additional point from each local region to also sample points by this metric
+	 * @param EarlyStop					If >= 0, stop re-ordering once we've sorted this many points; beyond this the ordering will be arbitrary
+	 */
+	void ComputeUniformSpaced(TArrayView<const FVector3d> Points, TArrayView<const float> ImportanceWeights, TArrayView<const float> SecondImportanceWeights, int32 EarlyStop = -1);
+	void ComputeUniformSpaced(TArrayView<const FVector3f> Points, TArrayView<const float> ImportanceWeights, TArrayView<const float> SecondImportanceWeights, int32 EarlyStop = -1);
+
 	/**
 	 * Compute an ordering that only sorts by descending ImportanceWeights, and does not attempt to spatially distribute the points
 	 * 
