Delaunay triangulation algorithm improvements:
 - fix bug with detection of vertices on the solid segment of a ghost triangle
 - detect when 'fill' of a triangulation is ambiguous under winding rules (e.g. due to open edges)
 - optionally automatically track duplicate vertices and remap duplicate-vertex references to the vertex that was actually inserted

diff --git a/Engine/Source/Runtime/GeometryCore/Private/CompGeom/Delaunay2.cpp b/Engine/Source/Runtime/GeometryCore/Private/CompGeom/Delaunay2.cpp
--- a/Engine/Source/Runtime/GeometryCore/Private/CompGeom/Delaunay2.cpp
+++ b/Engine/Source/Runtime/GeometryCore/Private/CompGeom/Delaunay2.cpp
@@ -18,6 +18,9 @@ struct FDelaunay2Connectivity
 	static constexpr int32 GhostIndex = -1;
 	static constexpr int32 InvalidIndex = -2;
 
+	FDelaunay2Connectivity(bool bTrackDuplicateVertices) : bTrackDuplicateVertices(bTrackDuplicateVertices)
+	{}
+
 	void Empty(int32 ExpectedMaxVertices = 0)
 	{
 		EdgeToVert.Empty(ExpectedMaxVertices * 2 * 3);
@@ -157,13 +160,19 @@ struct FDelaunay2Connectivity
 		}
 	}
 
-	TArray<FIndex3i> GetFilledTriangles(TArrayView<const FIndex2i> Edges, FDelaunay2::EFillMode FillMode)
+	bool GetFilledTriangles(TArray<FIndex3i>& TrianglesOut, TArrayView<const FIndex2i> Edges, FDelaunay2::EFillMode FillMode)
 	{
+		bool bWellDefinedResult = true;
+
 		TMap<FIndex2i, int> EdgeDelta;
 		EdgeDelta.Reserve(Edges.Num());
 		for (const FIndex2i& Edge : Edges)
 		{
 			FIndex2i Unoriented = Edge;
+			if (bTrackDuplicateVertices)
+			{
+				FixDuplicatesOnEdge(Unoriented);
+			}
 			int32 Sign = 1;
 			if (Unoriented.A > Unoriented.B)
 			{
@@ -255,6 +264,7 @@ struct FDelaunay2Connectivity
 			int32* Winding = Windings.Find(UniqueTri);
 			if (Winding)
 			{
+				bWellDefinedResult = bWellDefinedResult && (*Winding == Walk.Winding);
 				continue;
 			}
 			Windings.Add(UniqueTri, Walk.Winding);
@@ -262,24 +272,24 @@ struct FDelaunay2Connectivity
 			WalkEdge(FIndex2i(Vert, Walk.Edge.A), Walk.Winding);
 		}
 
-		TArray<FIndex3i> Triangles;
+		TrianglesOut.Reset();
 		if (FillMode == FDelaunay2::EFillMode::Solid)
 		{
-			Triangles.Reserve(EdgeToVert.Num() / 3 - GhostTriCount - Windings.Num());
+			TrianglesOut.Reserve(EdgeToVert.Num() / 3 - GhostTriCount - Windings.Num());
 			// everything we *didn't* walk over is included
-			EnumerateTriangles([&Windings, &Triangles](const FIndex2i& Edge, int32 Vertex) -> bool
+			EnumerateTriangles([&Windings, &TrianglesOut](const FIndex2i& Edge, int32 Vertex) -> bool
 				{
 					FIndex3i Triangle(Vertex, Edge.A, Edge.B);
 					if (!Windings.Contains(Triangle))
 					{
-						Triangles.Add(Triangle);
+						TrianglesOut.Add(Triangle);
 					}
 					return true;
 				}, true /*bSkipGhosts*/);
 		}
 		else
 		{
-			Triangles.Reserve(Windings.Num());
+			TrianglesOut.Reserve(Windings.Num());
 			for (const TPair<FIndex3i, int32>& TriWinding : Windings)
 			{
 				if ((FillMode == FDelaunay2::EFillMode::NegativeWinding	&& TriWinding.Value < 0) ||
@@ -287,11 +297,11 @@ struct FDelaunay2Connectivity
 					(FillMode == FDelaunay2::EFillMode::PositiveWinding && TriWinding.Value > 0) ||
 					(FillMode == FDelaunay2::EFillMode::OddWinding		&& (TriWinding.Value % 2) != 0))
 				{
-					Triangles.Add(TriWinding.Key);
+					TrianglesOut.Add(TriWinding.Key);
 				}
 			}
 		}
-		return Triangles;
+		return bWellDefinedResult;
 	}
 
 	static bool IsGhost(const FIndex2i& Edge, int32 Vertex)
@@ -362,6 +372,31 @@ struct FDelaunay2Connectivity
 		return *V;
 	}
 
+	void MarkDuplicateVertex(int32 Orig, int32 Duplicate)
+	{
+		checkSlow(bTrackDuplicateVertices);
+		DuplicateVertices.Add(Duplicate, Orig);
+	}
+
+	bool HasDuplicateTracking()
+	{
+		return bTrackDuplicateVertices;
+	}
+
+	void FixDuplicatesOnEdge(FIndex2i& Edge)
+	{
+		int32* OrigA = DuplicateVertices.Find(Edge.A);
+		if (OrigA)
+		{
+			Edge.A = *OrigA;
+		}
+		int32* OrigB = DuplicateVertices.Find(Edge.B);
+		if (OrigB)
+		{
+			Edge.B = *OrigB;
+		}
+	}
+
 	// Get any edge BC opposite vertex A, such that triangle ABC is in the mesh (or return InvalidIndex edge if no such edge is present)
 	// Before calling this frequently, consider calling EnableVertexAdjacency()
 	FIndex2i GetEdge(int32 Vertex) const
@@ -435,6 +470,9 @@ protected:
 	mutable TArray<int32> VertexAdjCache;
 	bool bUseAdjCache = false;
 
+	bool bTrackDuplicateVertices = false;
+	TMap<int32, int32> DuplicateVertices;
+
 	static inline int32 VertexIDToAdjIndex(int32 VertexID)
 	{
 		// Offset by 1 so that GhostIndex enters slot 0
@@ -500,21 +538,21 @@ namespace DelaunayInternal
 		// Pred == 0 case: need to check if Vertex is *on* the edge
 		if (Vertices[GhostFirst.B].X != Vertices[GhostFirst.C].X)
 		{
-			TInterval1<RealType> XRange(Vertices[GhostFirst.B].X, Vertices[GhostFirst.C].X);
+			TInterval1<RealType> XRange = TInterval1<RealType>::MakeFromUnordered(Vertices[GhostFirst.B].X, Vertices[GhostFirst.C].X);
 			return XRange.Contains(Vertices[Vertex].X);
 		}
 		else
 		{
-			TInterval1<RealType> YRange(Vertices[GhostFirst.B].Y, Vertices[GhostFirst.C].Y);
+			TInterval1<RealType> YRange = TInterval1<RealType>::MakeFromUnordered(Vertices[GhostFirst.B].Y, Vertices[GhostFirst.C].Y);
 			return YRange.Contains(Vertices[Vertex].Y);
 		}
 	}
 
 	// @return triangle containing Vertex
 	template<typename RealType>
-	FIndex3i WalkToContainingTri(FRandomStream& Random, FDelaunay2Connectivity& Connectivity, TArrayView<const TVector2<RealType>> Vertices, FIndex3i StartTri, int32 Vertex, bool bAssumeDelaunay, bool& bIsDuplicateOut)
+	FIndex3i WalkToContainingTri(FRandomStream& Random, FDelaunay2Connectivity& Connectivity, TArrayView<const TVector2<RealType>> Vertices, FIndex3i StartTri, int32 Vertex, bool bAssumeDelaunay, int32& IsDuplicateOfOut)
 	{
-		bIsDuplicateOut = false;
+		IsDuplicateOfOut = FDelaunay2Connectivity::InvalidIndex;
 
 		FIndex2i FirstEdge(FDelaunay2Connectivity::InvalidIndex, FDelaunay2Connectivity::InvalidIndex);
 		constexpr int32 GhostV = FDelaunay2Connectivity::GhostIndex; // shorter name
@@ -589,10 +627,18 @@ namespace DelaunayInternal
 			Cross = ChooseCross(WalkTri, true);
 		}
 
-		bIsDuplicateOut =
-			(WalkTri.A >= 0 && Vertices[Vertex] == Vertices[WalkTri.A]) ||
-			(WalkTri.B >= 0 && Vertices[Vertex] == Vertices[WalkTri.B]) ||
-			(WalkTri.C >= 0 && Vertices[Vertex] == Vertices[WalkTri.C]);
+		if (WalkTri.A >= 0 && Vertices[Vertex] == Vertices[WalkTri.A])
+		{
+			IsDuplicateOfOut = WalkTri.A;
+		}
+		else if (WalkTri.B >= 0 && Vertices[Vertex] == Vertices[WalkTri.B])
+		{
+			IsDuplicateOfOut = WalkTri.B;
+		}
+		else if (WalkTri.C >= 0 && Vertices[Vertex] == Vertices[WalkTri.C])
+		{
+			IsDuplicateOfOut = WalkTri.C;
+		}
 
 		return WalkTri;
 	}
@@ -659,6 +705,10 @@ namespace DelaunayInternal
 		SkipSet.Reserve(SkipEdgesIn.Num());
 		for (FIndex2i Edge : SkipEdgesIn)
 		{
+			if (Connectivity.HasDuplicateTracking())
+			{
+				Connectivity.FixDuplicatesOnEdge(Edge);
+			}
 			if (Edge.A > Edge.B)
 			{
 				Swap(Edge.A, Edge.B);
@@ -903,7 +953,7 @@ namespace DelaunayInternal
 			checkSlow(ABOrient[OrientIdx] > 0);
 		}
 
-		FDelaunay2Connectivity CavityCDT;
+		FDelaunay2Connectivity CavityCDT(false);
 		CavityCDT.Empty(CavityNum - 2); 
 
 		// Create an insertion ordering that tries to avoid concavities between adjacent pairs, and track adjacencies via Next/Prev
@@ -956,6 +1006,10 @@ namespace DelaunayInternal
 		{
 			int32 EdgeIdx = EdgeOrder[OrderIdx];
 			FIndex2i Edge = Edges[EdgeIdx];
+			if (Connectivity.HasDuplicateTracking())
+			{
+				Connectivity.FixDuplicatesOnEdge(Edge);
+			}
 			bool bDigSucceeded = false;
 			int32 DigTo = DigCavity<RealType>(Connectivity, Vertices, Edge, CavityVerts[0], CavityVerts[1], bDigSucceeded);
 			if (DigTo != Edge.A)
@@ -1041,11 +1095,15 @@ namespace DelaunayInternal
 		for (int32 OrderIdx = 0; OrderIdx < Order.Num(); OrderIdx++)
 		{
 			int32 Vertex = Order[OrderIdx];
-			bool bIsDuplicate = false;
+			int32 DuplicateOf = -1;
 			constexpr bool bAssumeDelaunay = true; // initial construction, before constraint edges, so safe to assume Delaunay
-			FIndex3i ContainingTri = WalkToContainingTri<RealType>(Random, Connectivity, Vertices, SearchTri, Vertex, bAssumeDelaunay, bIsDuplicate);
-			if (bIsDuplicate || ContainingTri[0] == FDelaunay2Connectivity::InvalidIndex)
+			FIndex3i ContainingTri = WalkToContainingTri<RealType>(Random, Connectivity, Vertices, SearchTri, Vertex, bAssumeDelaunay, DuplicateOf);
+			if (DuplicateOf >= 0 || ContainingTri[0] == FDelaunay2Connectivity::InvalidIndex)
 			{
+				if (DuplicateOf >= 0 && Connectivity.HasDuplicateTracking())
+				{
+					Connectivity.MarkDuplicateVertex(DuplicateOf, Vertex);
+				}
 				continue;
 			}
 			SearchTri = Insert<RealType>(Connectivity, Vertices, ContainingTri, Vertex);
@@ -1058,7 +1116,7 @@ namespace DelaunayInternal
 
 bool FDelaunay2::Triangulate(TArrayView<const FVector2d> Vertices, TArrayView<const FIndex2i> Edges)
 {
-	Connectivity = MakePimpl<FDelaunay2Connectivity>();
+	Connectivity = MakePimpl<FDelaunay2Connectivity>(bAutomaticallyFixEdgesToDuplicateVertices);
 
 	bIsConstrained = Edges.Num() > 0;
 
@@ -1068,7 +1126,7 @@ bool FDelaunay2::Triangulate(TArrayView<const FVector2d> Vertices, TArrayView<co
 
 bool FDelaunay2::Triangulate(TArrayView<const FVector2f> Vertices, TArrayView<const FIndex2i> Edges)
 {
-	Connectivity = MakePimpl<FDelaunay2Connectivity>();
+	Connectivity = MakePimpl<FDelaunay2Connectivity>(bAutomaticallyFixEdgesToDuplicateVertices);
 
 	bIsConstrained = Edges.Num() > 0;
 
@@ -1079,6 +1137,8 @@ bool FDelaunay2::Triangulate(TArrayView<const FVector2f> Vertices, TArrayView<co
 bool FDelaunay2::ConstrainEdges(TArrayView<const FVector2d> Vertices, TArrayView<const FIndex2i> Edges)
 {
 	bIsConstrained = bIsConstrained || Edges.Num() > 0;
+	// Note: automatic edge-to-duplicate fixing will not work if duplicate tracking was not enabled
+	checkSlow(!bAutomaticallyFixEdgesToDuplicateVertices || Connectivity->HasDuplicateTracking());
 
 	bool bSuccess = DelaunayInternal::ConstrainEdges<double>(RandomStream, *Connectivity, Vertices, Edges, bKeepFastEdgeAdjacencyData);
 	return bSuccess && ValidateResult(Edges);
@@ -1087,6 +1147,8 @@ bool FDelaunay2::ConstrainEdges(TArrayView<const FVector2d> Vertices, TArrayView
 bool FDelaunay2::ConstrainEdges(TArrayView<const FVector2f> Vertices, TArrayView<const FIndex2i> Edges)
 {
 	bIsConstrained = bIsConstrained || Edges.Num() > 0;
+	// Note: automatic edge-to-duplicate fixing will not work if duplicate tracking was not enabled
+	checkSlow(!bAutomaticallyFixEdgesToDuplicateVertices || Connectivity->HasDuplicateTracking());
 
 	bool bSuccess = DelaunayInternal::ConstrainEdges<float>(RandomStream, *Connectivity, Vertices, Edges, bKeepFastEdgeAdjacencyData);
 	return bSuccess && ValidateResult(Edges);
@@ -1109,13 +1171,13 @@ void FDelaunay2::GetTrianglesAndAdjacency(TArray<FIndex3i>& Triangles, TArray<FI
 	}
 }
 
-TArray<FIndex3i> FDelaunay2::GetFilledTriangles(TArrayView<const FIndex2i> Edges, EFillMode FillMode)
+bool FDelaunay2::GetFilledTriangles(TArray<FIndex3i>& TrianglesOut, TArrayView<const FIndex2i> Edges, EFillMode FillMode)
 {
 	if (!ensure(Connectivity.IsValid()))
 	{
-		return TArray<FIndex3i>();
+		return false;
 	}
-	return Connectivity->GetFilledTriangles(Edges, FillMode);
+	return Connectivity->GetFilledTriangles(TrianglesOut, Edges, FillMode);
 }
 
 bool FDelaunay2::IsDelaunay(TArrayView<const FVector2f> Vertices, TArrayView<const FIndex2i> SkipEdges) const
@@ -1140,8 +1202,12 @@ bool FDelaunay2::HasEdges(TArrayView<const FIndex2i> Edges) const
 {
 	if (ensure(Connectivity.IsValid()))
 	{
-		for (const FIndex2i& Edge : Edges)
+		for (FIndex2i Edge : Edges)
 		{
+			if (bAutomaticallyFixEdgesToDuplicateVertices)
+			{
+				Connectivity->FixDuplicatesOnEdge(Edge);
+			}
 			if (!Connectivity->HasEdge(Edge))
 			{
 				return false;
diff --git a/Engine/Source/Runtime/GeometryCore/Public/BoxTypes.h b/Engine/Source/Runtime/GeometryCore/Public/BoxTypes.h
--- a/Engine/Source/Runtime/GeometryCore/Public/BoxTypes.h
+++ b/Engine/Source/Runtime/GeometryCore/Public/BoxTypes.h
@@ -36,6 +36,16 @@ struct TInterval1
 		return TInterval1(TNumericLimits<RealType>::Max(), -TNumericLimits<RealType>::Max());
 	}
 
+	static TInterval1<RealType> MakeFromUnordered(const RealType& A, const RealType& B)
+	{
+		TInterval1<RealType> Result(A, B);
+		if (A > B)
+		{
+			Swap(Result.Min, Result.Max);
+		}
+		return Result;
+	}
+
 	RealType Center() const
 	{
 		return (Min + Max) * (RealType)0.5;
diff --git a/Engine/Source/Runtime/GeometryCore/Public/CompGeom/Delaunay2.h b/Engine/Source/Runtime/GeometryCore/Public/CompGeom/Delaunay2.h
--- a/Engine/Source/Runtime/GeometryCore/Public/CompGeom/Delaunay2.h
+++ b/Engine/Source/Runtime/GeometryCore/Public/CompGeom/Delaunay2.h
@@ -58,6 +58,10 @@ public:
 	// TODO: Consider having the internal mesh remember what edges have been constrained, so we can set an error flag when constrained edges cross previous constrained edges
 	bool bValidateEdges = true;
 
+	// Option to automatically track duplicate vertices and treat edges that reference them as if they referenced the instance that was actually used in the triangulation
+	// Note: Cannot change this to 'true' *after* triangulation and then call ConstrainEdges; duplicate vertices will only be detected on their initial insertion
+	bool bAutomaticallyFixEdgesToDuplicateVertices = false;
+
 	// TODO: it would often be useful to pass in sparse vertex data
 	//// Optional function to allow Triangulate to skip vertices
 	//TFunction<bool(int32)> SkipVertexFn = nullptr;
@@ -162,8 +166,29 @@ public:
 	/**
 	 * Return the triangles that are inside the given edges, removing the outer boundary triangles
 	 * If a Winding-Number-based fill mode is used, assumes edges are oriented and tracks the winding number across edges
+	 * 
+	 * @param Edges						The array of edges that were already constrained in the triangulation (by a previous call to Triangulate or ConstrainEdges)
+	 * @param FillMode					Strategy to use to define which triangles to include in the output
+	 * @return							A subset of the triangulation that is 'inside' the given edges, as defined by the FillMode
+	 */
+	TArray<FIndex3i> GetFilledTriangles(TArrayView<const FIndex2i> Edges, EFillMode FillMode = EFillMode::PositiveWinding)
+	{
+		TArray<FIndex3i> Triangles;
+		GetFilledTriangles(Triangles, Edges, FillMode);
+		return Triangles;
+	}
+
+	/**
+	 * Get (by reference) the triangles that are inside the given edges, removing the outer boundary triangles
+	 * If a Winding-Number-based fill mode is used, assumes edges are oriented and tracks the winding number across edges
+	 *
+	 * @param TrianglesOut				Will be filled with a subset of the triangulation that is 'inside' the given edges, as defined by the FillMode
+	 * @param Edges						The array of edges that were already constrained in the triangulation (by a previous call to Triangulate or ConstrainEdges)
+	 * @param FillMode					Strategy to use to define which triangles to include in the output
+	 * @return							true if the result was well defined and consistent; false otherwise. Solid fill mode always has a well defined result; Winding-number-based fills do not if the edges have open spans.
+	 *									Note the triangulation will still be filled by best-effort even if the function returns false.
 	 */
-	TArray<FIndex3i> GetFilledTriangles(TArrayView<const FIndex2i> Edges, EFillMode FillMode = EFillMode::PositiveWinding);
+	bool GetFilledTriangles(TArray<FIndex3i>& TrianglesOut, TArrayView<const FIndex2i> Edges, EFillMode FillMode = EFillMode::PositiveWinding);
 
 	// @return true if this is a constrained Delaunay triangulation
 	bool IsConstrained() const
