Add a vertex-based connected components class that can be used with any TriangleMeshType and supports connecting based on spatial proximity as well as triangle topology

diff --git a/Engine/Source/Runtime/GeometryCore/Public/VertexConnectedComponents.h b/Engine/Source/Runtime/GeometryCore/Public/VertexConnectedComponents.h
new file mode 100644
--- /dev/null
+++ b/Engine/Source/Runtime/GeometryCore/Public/VertexConnectedComponents.h
@@ -0,0 +1,294 @@
+// Copyright Epic Games, Inc. All Rights Reserved.
+
+#pragma once
+
+#include "BoxTypes.h"
+#include "IndexTypes.h"
+#include "Spatial/PointHashGrid3.h"
+
+namespace UE
+{
+namespace Geometry
+{
+
+/// Disjoint set with additional storage to track the size of each set
+struct FSizedDisjointSet
+{
+	TArray<int32> Parents, Sizes;
+
+	void Init(int32 NumIDs)
+	{
+		Parents.SetNumUninitialized(NumIDs);
+		Sizes.SetNumUninitialized(NumIDs);
+		for (int32 Idx = 0; Idx < NumIDs; Idx++)
+		{
+			Parents[Idx] = Idx;
+			Sizes[Idx] = 1;
+		}
+	}
+
+	// IsElement(Index) returns false if Index is not valid.  Find, Union and GetSize should never be called on an invalid index.
+	template<typename IsElementFunction>
+	void Init(int32 NumIDs, IsElementFunction IsElement)
+	{
+		Parents.SetNumUninitialized(NumIDs);
+		Sizes.SetNumUninitialized(NumIDs);
+		for (int32 Idx = 0; Idx < NumIDs; Idx++)
+		{
+			if (IsElement(Idx))
+			{
+				Parents[Idx] = Idx;
+				Sizes[Idx] = 1;
+			}
+			else
+			{
+				Parents[Idx] = -1;
+				Sizes[Idx] = 0;
+			}
+		}
+	}
+
+	void Union(int32 A, int32 B)
+	{
+		int32 ParentA = Find(A);
+		int32 ParentB = Find(B);
+		if (ParentA == ParentB)
+		{
+			return;
+		}
+		if (Sizes[ParentA] > Sizes[ParentB])
+		{
+			Swap(ParentA, ParentB);
+		}
+		Parents[ParentA] = ParentB;
+		Sizes[ParentB] += Sizes[ParentA];
+	}
+
+	int32 Find(int32 Idx)
+	{
+		int32 Parent = Idx;
+
+		while (Parents[Parent] != Parent)
+		{
+			Parents[Parent] = Parents[Parents[Parent]];
+			Parent = Parents[Parent];
+		}
+
+		return Parent;
+	}
+
+	int32 GetSize(int32 Idx)
+	{
+		int32 Parent = Find(Idx);
+		return Sizes[Parent];
+	}
+};
+
+/// Vertex-based connected components class -- can work with any mesh that has vertex IDs
+/// Also supports linking spatially-close vertices in the same component
+/// 
+/// Functions templated on TemplateMeshType are designed to work with any mesh that 
+/// implements the standard MeshAdapter functions (see MeshAdapter.h)
+/// Functions templated on TriangleType are designed to work with triangles 
+/// with vertex IDs that can be array-accessed (i.e.: Tri[0], Tri[1], Tri[2])
+class FVertexConnectedComponents
+{
+public:
+	FVertexConnectedComponents()
+	{
+		Init(0);
+	}
+
+	FVertexConnectedComponents(int32 MaxVertexID)
+	{
+		Init(MaxVertexID);
+	}
+	
+	void Init(int32 MaxVertexID)
+	{
+		DisjointSet.Init(MaxVertexID);
+	}
+
+	template<typename TriangleMeshType>
+	void Init(const TriangleMeshType& Mesh)
+	{
+		DisjointSet.Init(Mesh.MaxVertexID(), [&Mesh](int32 VID) -> bool { return Mesh.IsVertex(VID); });
+	}
+
+	template<typename TriangleMeshType>
+	void ConnectTriangles(const TriangleMeshType& Mesh)
+	{
+		for (int32 TID = 0; TID < Mesh.MaxTriangleID(); TID++)
+		{
+			if (Mesh.IsTriangle(TID))
+			{
+				FIndex3i Triangle = Mesh.GetTriangle(TID);
+				DisjointSet.Union(Triangle[0], Triangle[1]);
+				DisjointSet.Union(Triangle[1], Triangle[2]);
+			}
+		}
+	}
+
+	template<typename TriangleType>
+	void ConnectTriangles(TArrayView<const TriangleType> Triangles)
+	{
+		for (const TriangleType& Triangle : Triangles)
+		{
+			DisjointSet.Union(Triangle[0], Triangle[1]);
+			DisjointSet.Union(Triangle[1], Triangle[2]);
+		}
+	}
+
+	// Note for meshes that may have floating (unreferenced) vertices, using a KeepSizeThreshold > 1 after calling ConnectTriangles will ensure the unreferenced vertices are skipped
+	template<typename TriangleMeshType>
+	void ConnectCloseVertices(const TriangleMeshType& Mesh, double CloseVertexThreshold, int32 KeepSizeThreshold = 0)
+	{
+		TPointHashGrid3d<int32> VertexHash(CloseVertexThreshold * 3, -1);
+		TArray<int32> Neighbors;
+		for (int32 VID = 0; VID < Mesh.MaxVertexID(); VID++)
+		{
+			if (!Mesh.IsVertex(VID))
+			{
+				continue;
+			}
+			int32 SetID = DisjointSet.Find(VID);
+			if (KeepSizeThreshold > 0 && DisjointSet.Sizes[SetID] < KeepSizeThreshold)
+			{
+				continue;
+			}
+
+			FVector3d Pt = Mesh.GetVertex(VID);
+			Neighbors.Reset();
+			VertexHash.FindPointsInBall(Pt, CloseVertexThreshold, [&Mesh, Pt](int32 OtherVID)
+				{
+					// TODO: could return MaxReal if OtherVID is already in the same component so we skip union-ing it?
+					return DistanceSquared(Pt, Mesh.GetVertex(OtherVID));
+				}, Neighbors);
+			for (int32 NbrVID : Neighbors)
+			{
+				DisjointSet.Union(SetID, NbrVID);
+			}
+			VertexHash.InsertPointUnsafe(VID, Pt);
+		}
+	}
+
+	// TODO: support more overlap strategies, currently just uses AABB
+	// Note this merges components based on overlap of their bounding boxes as computed *before* any merges; multiple passes may merge additional components
+	template<typename TriangleMeshType>
+	void ConnectOverlappingComponents(const TriangleMeshType& Mesh, int32 KeepSizeThreshold = 0)
+	{
+		TMap<int32, FAxisAlignedBox3d> SetIDToBounds;
+		for (int32 VID = 0; VID < Mesh.MaxVertexID(); VID++)
+		{
+			if (!Mesh.IsVertex(VID))
+			{
+				continue;
+			}
+			int32 SetID = DisjointSet.Find(VID);
+			if (KeepSizeThreshold > 0 && DisjointSet.Sizes[SetID] < KeepSizeThreshold)
+			{
+				continue;
+			}
+			FVector3d Pt = Mesh.GetVertex(VID);
+			FAxisAlignedBox3d& Bounds = SetIDToBounds.FindOrAdd(SetID, FAxisAlignedBox3d::Empty());
+			Bounds.Contain(Pt);
+		}
+		// For each component's bounding box
+		for (auto It = SetIDToBounds.CreateConstIterator(); It; ++It)
+		{
+			auto NextIt = It;
+			// Iterate over the subsequent component bounding boxes
+			++NextIt;
+			for (; NextIt; ++NextIt)
+			{
+				// Union the two components if their bounds overlap
+				if (It.Value().Intersects(NextIt.Value()))
+				{
+					DisjointSet.Union(It.Key(), NextIt.Key());
+				}
+			}
+		}
+	}
+
+	template<typename TriangleMeshType>
+	bool HasMultipleComponents(const TriangleMeshType& Mesh, int32 KeepSizeThreshold = 0)
+	{
+		int32 FoundComponent = -1;
+		for (int32 VID = 0; VID < Mesh.MaxVertexID(); VID++)
+		{
+			if (!Mesh.IsVertex(VID))
+			{
+				continue;
+			}
+			int32 SetID = DisjointSet.Find(VID);
+			if (KeepSizeThreshold > 0 && DisjointSet.Sizes[SetID] < KeepSizeThreshold)
+			{
+				continue;
+			}
+			if (FoundComponent == -1)
+			{
+				FoundComponent = SetID;
+			}
+			else if (FoundComponent != SetID)
+			{
+				return true;
+			}
+		}
+		return false;
+	}
+
+	// Map arbitrary set IDs to indices from 0 to k-1 (if there are k components)
+	template<typename TriangleMeshType>
+	TMap<int32, int32> MakeComponentMap(const TriangleMeshType& Mesh, int32 KeepSizeThreshold = 0)
+	{
+		TMap<int32, int32> ComponentMap;
+		int32 CurrentIdx = 0;
+		for (int32 VID = 0; VID < Mesh.MaxVertexID(); VID++)
+		{
+			if (!Mesh.IsVertex(VID))
+			{
+				continue;
+			}
+			int32 SetID = DisjointSet.Find(VID);
+			if (KeepSizeThreshold > 0 && DisjointSet.Sizes[SetID] < KeepSizeThreshold)
+			{
+				continue;
+			}
+			if (!ComponentMap.Contains(SetID))
+			{
+				ComponentMap.Add(SetID, CurrentIdx++);
+			}
+		}
+		return ComponentMap;
+	}
+
+	inline int32 GetComponent(int32 VertexID)
+	{
+		return DisjointSet.Find(VertexID);
+	}
+
+	inline int32 GetComponentSize(int32 VertexID)
+	{
+		return DisjointSet.GetSize(VertexID);
+	}
+
+	template<typename TriangleType>
+	inline int32 GetComponent(const TriangleType& Triangle)
+	{
+		return GetComponent(Triangle[0]);
+	}
+
+	inline void ConnectVertices(int32 VertexID0, int32 VertexID1)
+	{
+		DisjointSet.Union(VertexID0, VertexID1);
+	}
+
+
+protected:
+
+	FSizedDisjointSet DisjointSet;
+};
+
+
+} // end namespace UE::Geometry
+} // end namespace UE
\ No newline at end of file
