Add scalar weight maps to DynamicMeshAttributeSet
- Allow the Attribute Editor Tool to add/delete weight maps from DynamicMesh targets.
- Enable transferring scalar weight maps back and forth between DynamicMesh and MeshDescription.
- Serialize TDynamicAttributeBase::Name when its subclasses are serialized

diff --git a/Engine/Plugins/Experimental/MeshModelingToolsetExp/Source/MeshModelingToolsEditorOnlyExp/Private/AttributeEditorTool.cpp b/Engine/Plugins/Experimental/MeshModelingToolsetExp/Source/MeshModelingToolsEditorOnlyExp/Private/AttributeEditorTool.cpp
--- a/Engine/Plugins/Experimental/MeshModelingToolsetExp/Source/MeshModelingToolsEditorOnlyExp/Private/AttributeEditorTool.cpp
+++ b/Engine/Plugins/Experimental/MeshModelingToolsetExp/Source/MeshModelingToolsEditorOnlyExp/Private/AttributeEditorTool.cpp
@@ -674,12 +674,6 @@ void UAttributeEditorTool::AddNewAttribute()
 
 void UAttributeEditorTool::AddNewWeightMap()
 {
-	if (bTargetIsStaticMesh == false)
-	{
-		GetToolManager()->DisplayMessage(LOCTEXT("ErrorAddingWeightMap", "Currently cannot add Weight Maps to Dynamic Mesh Components"), EToolMessageLevel::UserWarning);
-		return;
-	}
-
 	AddNewAttribute(EAttributeEditorElementType::Vertex, EAttributeEditorAttribType::Float, FName(NewAttributeProps->NewName));
 }
 
diff --git a/Engine/Plugins/Runtime/GeometryProcessing/Source/DynamicMesh/Private/DynamicMeshEditor.cpp b/Engine/Plugins/Runtime/GeometryProcessing/Source/DynamicMesh/Private/DynamicMeshEditor.cpp
--- a/Engine/Plugins/Runtime/GeometryProcessing/Source/DynamicMesh/Private/DynamicMeshEditor.cpp
+++ b/Engine/Plugins/Runtime/GeometryProcessing/Source/DynamicMesh/Private/DynamicMeshEditor.cpp
@@ -1405,6 +1405,17 @@ int FDynamicMeshEditor::FindOrCreateDuplicateVertex(int VertexID, FMeshIndexMapp
 		NewVertexID = Mesh->AppendVertex(*Mesh, VertexID);
 		IndexMaps.SetVertex(VertexID, NewVertexID);
 		ResultOut.NewVertices.Add(NewVertexID);
+
+		if (Mesh->HasAttributes())
+		{
+			for (int WeightLayerIndex = 0; WeightLayerIndex < Mesh->Attributes()->NumWeightLayers(); ++WeightLayerIndex)
+			{
+				FDynamicMeshWeightAttribute* WeightAttr = Mesh->Attributes()->GetWeightLayer(WeightLayerIndex);
+				float Val;
+				WeightAttr->GetValue(VertexID, &Val);
+				WeightAttr->SetNewValue(NewVertexID, &Val);
+			}
+		}
 	}
 	return NewVertexID;
 }
@@ -1575,6 +1586,19 @@ void FDynamicMeshEditor::AppendMesh(const FDynamicMesh3* AppendMesh,
 			}
 		}
 
+		int NumWeightLayers = FMath::Min(Mesh->Attributes()->NumWeightLayers(), AppendMesh->Attributes()->NumWeightLayers());
+		for (int WeightLayerIndex = 0; WeightLayerIndex < NumWeightLayers; WeightLayerIndex++)
+		{
+			const FDynamicMeshWeightAttribute* FromWeights = AppendMesh->Attributes()->GetWeightLayer(WeightLayerIndex);
+			FDynamicMeshWeightAttribute* ToWeights = Mesh->Attributes()->GetWeightLayer(WeightLayerIndex);
+			for (const TPair<int32, int32>& MapVID : VertexMap.GetForwardMap())
+			{
+				float Weight;
+				FromWeights->GetValue(MapVID.Key, &Weight);
+				ToWeights->SetValue(MapVID.Value, &Weight);
+			}
+		}
+
 		for (const TPair<FName, TUniquePtr<FDynamicMeshVertexSkinWeightsAttribute>>& AttribPair : AppendMesh->Attributes()->GetSkinWeightsAttributes())
 		{
 			FDynamicMeshVertexSkinWeightsAttribute* ToAttrib = Mesh->Attributes()->GetSkinWeightsAttribute(AttribPair.Key);
@@ -1894,7 +1918,7 @@ static void AppendTriangleAttributes(const FDynamicMesh3* FromMesh, int FromTria
 
 
 // Utility function for ::AppendTriangles()
-static void AppendGenericAttributes(const FDynamicMesh3* FromMesh, FDynamicMesh3* ToMesh, FMeshIndexMappings& IndexMaps)
+static void AppendVertexAttributes(const FDynamicMesh3* FromMesh, FDynamicMesh3* ToMesh, FMeshIndexMappings& IndexMaps)
 {
 
 	if (FromMesh->HasAttributes() == false || ToMesh->HasAttributes() == false)
@@ -1902,6 +1926,19 @@ static void AppendGenericAttributes(const FDynamicMesh3* FromMesh, FDynamicMesh3
 		return;
 	}
 
+	int NumWeightLayers = FMath::Min(FromMesh->Attributes()->NumWeightLayers(), ToMesh->Attributes()->NumWeightLayers());
+	for (int WeightLayerIndex = 0; WeightLayerIndex < NumWeightLayers; WeightLayerIndex++)
+	{
+		const FDynamicMeshWeightAttribute* FromWeights = FromMesh->Attributes()->GetWeightLayer(WeightLayerIndex);
+		FDynamicMeshWeightAttribute* ToWeights = ToMesh->Attributes()->GetWeightLayer(WeightLayerIndex);
+		for (const TPair<int32, int32>& MapVID : IndexMaps.GetVertexMap().GetForwardMap())
+		{
+			float Weight;
+			FromWeights->GetValue(MapVID.Key, &Weight);
+			ToWeights->SetValue(MapVID.Value, &Weight);
+		}
+	}
+
 	// Copy skin weight and generic attributes after full IndexMaps have been created. 	
 	for (const TPair<FName, TUniquePtr<FDynamicMeshVertexSkinWeightsAttribute>>& AttribPair : FromMesh->Attributes()->GetSkinWeightsAttributes())
 	{
@@ -2002,7 +2039,7 @@ void FDynamicMeshEditor::AppendTriangles(const FDynamicMesh3* SourceMesh, const
 		//Mesh->CheckValidity(true);
 	}
  
-	AppendGenericAttributes(SourceMesh, Mesh, IndexMaps);
+	AppendVertexAttributes(SourceMesh, Mesh, IndexMaps);
 	
 }
 
@@ -2108,7 +2145,7 @@ bool FDynamicMeshEditor::SplitMesh(const FDynamicMesh3* SourceMesh, TArray<FDyna
 
 	for (int Idx = 0; Idx < NumMeshes; Idx++)
 	{
-		AppendGenericAttributes(SourceMesh, &SplitMeshes[Idx], Mappings[Idx]);
+		AppendVertexAttributes(SourceMesh, &SplitMeshes[Idx], Mappings[Idx]);
 	}
 	
 	return true;
diff --git a/Engine/Plugins/Runtime/GeometryProcessing/Source/DynamicMesh/Private/Operations/UniformTessellate.cpp b/Engine/Plugins/Runtime/GeometryProcessing/Source/DynamicMesh/Private/Operations/UniformTessellate.cpp
--- a/Engine/Plugins/Runtime/GeometryProcessing/Source/DynamicMesh/Private/Operations/UniformTessellate.cpp
+++ b/Engine/Plugins/Runtime/GeometryProcessing/Source/DynamicMesh/Private/Operations/UniformTessellate.cpp
@@ -1349,6 +1349,37 @@ namespace UniformTessellateLocals
 
 			return MakeUnique<FVertexAttributeGenerator>(Mesh, TessellationNum, InProgress, bUseParallel, FPntr);
 		}
+
+		static TUniquePtr<FVertexAttributeGenerator> CreateFromVertexAttribute(const FDynamicMesh3* Mesh,
+			const int TessellationNum,
+			FProgressCancel* InProgress,
+			const bool bUseParallel,
+			const TDynamicMeshVertexAttribute<RealType, ElementSize>* Attribute)
+		{
+			if (Attribute == nullptr)
+			{
+				return nullptr;
+			}
+
+			TFunction<void(int, RealType*)> FPntr = [Attribute](const int VertexID, RealType* Value)
+			{
+				Attribute->GetValue(VertexID, Value);
+			};
+
+			return MakeUnique<FVertexAttributeGenerator>(Mesh, TessellationNum, InProgress, bUseParallel, FPntr);
+		}
+
+		void CopyToAttribute(TDynamicMeshVertexAttribute<RealType, ElementSize>* OutAttribute)
+		{
+			RealType Value[ElementSize];
+			for (int EID = 0; EID < OutElementCount; ++EID)
+			{
+				BaseType::GetElement(EID, Value);
+				OutAttribute->SetValue(EID, Value);
+			}
+		}
+
+
 	};
 
 	/**
@@ -1842,6 +1873,24 @@ namespace UniformTessellateLocals
 
 				OutAttributes->AttachSkinWeightsAttribute(AttributeInfo.Key, SkinAttribute);
 			}
+
+			if (InAttributes->NumWeightLayers() > 0)
+			{
+				OutAttributes->SetNumWeightLayers(InAttributes->NumWeightLayers());
+				for (int Idx = 0; Idx < InAttributes->NumWeightLayers(); ++Idx)
+				{
+					TUniquePtr<FVertexAttributeGenerator<float, 1>> Generator;
+					Generator = FVertexAttributeGenerator<float, 1>::CreateFromVertexAttribute(InMesh, TessellationNum, Progress, bUseParallel, InAttributes->GetWeightLayer(Idx));
+					if (Generator->Generate() == false)
+					{
+						return false;
+					}
+
+					Generator->CopyToAttribute(OutAttributes->GetWeightLayer(Idx));
+				}
+			}
+
+
 		}
 
 		if (InMesh->HasVertexNormals()) 
diff --git a/Engine/Source/Runtime/Core/Public/UObject/UE5MainStreamObjectVersion.h b/Engine/Source/Runtime/Core/Public/UObject/UE5MainStreamObjectVersion.h
--- a/Engine/Source/Runtime/Core/Public/UObject/UE5MainStreamObjectVersion.h
+++ b/Engine/Source/Runtime/Core/Public/UObject/UE5MainStreamObjectVersion.h
@@ -197,6 +197,9 @@ struct CORE_API FUE5MainStreamObjectVersion
 		// Use special BVH for TriangleMesh, instead of the AABBTree
 		UseTriangleMeshBVH,
 
+		// FDynamicMeshAttributeSet has Weight Maps. TDynamicAttributeBase serializes its name.
+		DynamicMeshAttributesWeightMapsAndNames,
+
 		// -----<new versions can be added above this line>-------------------------------------------------
 		VersionPlusOne,
 		LatestVersion = VersionPlusOne - 1
diff --git a/Engine/Source/Runtime/GeometryCore/Private/DynamicMesh/DynamicMeshAttributeSet.cpp b/Engine/Source/Runtime/GeometryCore/Private/DynamicMesh/DynamicMeshAttributeSet.cpp
--- a/Engine/Source/Runtime/GeometryCore/Private/DynamicMesh/DynamicMeshAttributeSet.cpp
+++ b/Engine/Source/Runtime/GeometryCore/Private/DynamicMesh/DynamicMeshAttributeSet.cpp
@@ -65,6 +65,12 @@ void FDynamicMeshAttributeSet::Copy(const FDynamicMeshAttributeSet& Copy)
 		PolygroupLayers[GroupIdx].Copy(Copy.PolygroupLayers[GroupIdx]);
 	}
 
+	SetNumWeightLayers(Copy.NumWeightLayers());
+	for (int WeightLayerIdx = 0; WeightLayerIdx < NumWeightLayers(); ++WeightLayerIdx)
+	{
+		WeightLayers[WeightLayerIdx].Copy(Copy.WeightLayers[WeightLayerIdx]);
+	}
+
 	ResetRegisteredAttributes();
 
 	SkinWeightAttributes.Reset();
@@ -151,6 +157,12 @@ void FDynamicMeshAttributeSet::CompactCopy(const FCompactMaps& CompactMaps, cons
 	{
 		PolygroupLayers[GroupIdx].CompactCopy(CompactMaps, Copy.PolygroupLayers[GroupIdx]);
 	}
+	
+	SetNumWeightLayers(Copy.NumWeightLayers());
+	for (int WeightLayerIdx = 0; WeightLayerIdx < NumWeightLayers(); ++WeightLayerIdx)
+	{
+		WeightLayers[WeightLayerIdx].CompactCopy(CompactMaps, Copy.WeightLayers[WeightLayerIdx]);
+	}
 
 	ResetRegisteredAttributes();
 
@@ -197,6 +209,11 @@ void FDynamicMeshAttributeSet::CompactInPlace(const FCompactMaps& CompactMaps)
 		PolygroupLayers[GroupIdx].CompactInPlace(CompactMaps);
 	}
 
+	for (int WeightLayerIdx = 0; WeightLayerIdx < NumWeightLayers(); ++WeightLayerIdx)
+	{
+		WeightLayers[WeightLayerIdx].CompactInPlace(CompactMaps);
+	}
+
 	for (const TPair<FName, TUniquePtr<FDynamicMeshVertexSkinWeightsAttribute>>& AttribPair : SkinWeightAttributes)
 	{
 		AttribPair.Value->CompactInPlace(CompactMaps);
@@ -293,6 +310,12 @@ void FDynamicMeshAttributeSet::EnableMatchingAttributes(const FDynamicMeshAttrib
 		PolygroupLayers[GroupIdx].Initialize((int32)0);
 	}
 
+	SetNumWeightLayers(ToMatch.NumWeightLayers());
+	for (int WeightLayerIdx = 0; WeightLayerIdx < NumWeightLayers(); ++WeightLayerIdx)
+	{
+		WeightLayers[WeightLayerIdx].Initialize(0.0f);
+	}
+
 	ResetRegisteredAttributes();
 	
 	SkinWeightAttributes.Reset();
@@ -338,6 +361,11 @@ void FDynamicMeshAttributeSet::Reparent(FDynamicMesh3* NewParent)
 		PolygroupLayers[GroupIdx].Reparent(NewParent);
 	}
 
+	for (int WeightLayerIdx = 0; WeightLayerIdx < NumWeightLayers(); ++WeightLayerIdx)
+	{
+		WeightLayers[WeightLayerIdx].Reparent(NewParent);
+	}
+
 	for (const TPair<FName, TUniquePtr<FDynamicMeshVertexSkinWeightsAttribute>>& AttribPair : SkinWeightAttributes)
 	{
 		AttribPair.Value.Get()->Reparent(NewParent);
@@ -459,6 +487,41 @@ const FDynamicMeshPolygroupAttribute* FDynamicMeshAttributeSet::GetPolygroupLaye
 }
 
 
+int32 FDynamicMeshAttributeSet::NumWeightLayers() const
+{
+	return WeightLayers.Num();
+}
+
+void FDynamicMeshAttributeSet::SetNumWeightLayers(int32 Num)
+{
+	if (WeightLayers.Num() == Num)
+	{
+		return;
+	}
+	if (Num >= WeightLayers.Num())
+	{
+		for (int i = (int)WeightLayers.Num(); i < Num; ++i)
+		{
+			WeightLayers.Add(new FDynamicMeshWeightAttribute(ParentMesh));
+		}
+	}
+	else
+	{
+		WeightLayers.RemoveAt(Num, WeightLayers.Num() - Num);
+	}
+	ensure(WeightLayers.Num() == Num);
+}
+
+FDynamicMeshWeightAttribute* FDynamicMeshAttributeSet::GetWeightLayer(int Index)
+{
+	return &WeightLayers[Index];
+}
+
+const FDynamicMeshWeightAttribute* FDynamicMeshAttributeSet::GetWeightLayer(int Index) const
+{
+	return &WeightLayers[Index];
+}
+
 
 void FDynamicMeshAttributeSet::EnableMaterialID()
 {
@@ -624,12 +687,23 @@ bool FDynamicMeshAttributeSet::IsMaterialBoundaryEdge(int EdgeID) const
 void FDynamicMeshAttributeSet::OnNewVertex(int VertexID, bool bInserted)
 {
 	FDynamicMeshAttributeSetBase::OnNewVertex(VertexID, bInserted);
+
+	for (FDynamicMeshWeightAttribute& WeightLayer : WeightLayers)
+	{
+		float NewWeight = 0.0f;
+		WeightLayer.SetNewValue(VertexID, &NewWeight);
+	}
 }
 
 
 void FDynamicMeshAttributeSet::OnRemoveVertex(int VertexID)
 {
 	FDynamicMeshAttributeSetBase::OnRemoveVertex(VertexID);
+
+	for (FDynamicMeshWeightAttribute& WeightLayer : WeightLayers)
+	{
+		WeightLayer.OnRemoveVertex(VertexID);
+	}
 }
 
 
@@ -727,6 +801,12 @@ void FDynamicMeshAttributeSet::OnSplitEdge(const FDynamicMesh3::FEdgeSplitInfo&
 	{
 		PolygroupLayer.OnSplitEdge(SplitInfo);
 	}
+
+	for (FDynamicMeshWeightAttribute& WeightLayer : WeightLayers)
+	{
+		WeightLayer.OnSplitEdge(SplitInfo);
+	}
+
 }
 
 void FDynamicMeshAttributeSet::OnFlipEdge(const FDynamicMesh3::FEdgeFlipInfo & flipInfo)
@@ -754,6 +834,12 @@ void FDynamicMeshAttributeSet::OnFlipEdge(const FDynamicMesh3::FEdgeFlipInfo & f
 	{
 		PolygroupLayer.OnFlipEdge(flipInfo);
 	}
+
+	for (FDynamicMeshWeightAttribute& WeightLayer : WeightLayers)
+	{
+		WeightLayer.OnFlipEdge(flipInfo);
+	}
+
 }
 
 
@@ -782,6 +868,12 @@ void FDynamicMeshAttributeSet::OnCollapseEdge(const FDynamicMesh3::FEdgeCollapse
 	{
 		PolygroupLayer.OnCollapseEdge(collapseInfo);
 	}
+
+	for (FDynamicMeshWeightAttribute& WeightLayer : WeightLayers)
+	{
+		WeightLayer.OnCollapseEdge(collapseInfo);
+	}
+
 }
 
 void FDynamicMeshAttributeSet::OnPokeTriangle(const FDynamicMesh3::FPokeTriangleInfo & pokeInfo)
@@ -809,6 +901,11 @@ void FDynamicMeshAttributeSet::OnPokeTriangle(const FDynamicMesh3::FPokeTriangle
 	{
 		PolygroupLayer.OnPokeTriangle(pokeInfo);
 	}
+
+	for (FDynamicMeshWeightAttribute& WeightLayer : WeightLayers)
+	{
+		WeightLayer.OnPokeTriangle(pokeInfo);
+	}
 }
 
 void FDynamicMeshAttributeSet::OnMergeEdges(const FDynamicMesh3::FMergeEdgesInfo & mergeInfo)
@@ -836,6 +933,11 @@ void FDynamicMeshAttributeSet::OnMergeEdges(const FDynamicMesh3::FMergeEdgesInfo
 	{
 		PolygroupLayer.OnMergeEdges(mergeInfo);
 	}
+
+	for (FDynamicMeshWeightAttribute& WeightLayer : WeightLayers)
+	{
+		WeightLayer.OnMergeEdges(mergeInfo);
+	}
 }
 
 void FDynamicMeshAttributeSet::OnSplitVertex(const DynamicMeshInfo::FVertexSplitInfo& SplitInfo, const TArrayView<const int>& TrianglesToUpdate)
@@ -863,13 +965,19 @@ void FDynamicMeshAttributeSet::OnSplitVertex(const DynamicMeshInfo::FVertexSplit
 	{
 		PolygroupLayer.OnSplitVertex(SplitInfo, TrianglesToUpdate);
 	}
+
+	for (FDynamicMeshWeightAttribute& WeightLayer : WeightLayers)
+	{
+		WeightLayer.OnSplitVertex(SplitInfo, TrianglesToUpdate);
+	}
 }
 
 bool FDynamicMeshAttributeSet::IsSameAs(const FDynamicMeshAttributeSet& Other, bool bIgnoreDataLayout) const
 {
 	if (UVLayers.Num() != Other.UVLayers.Num() ||
 		NormalLayers.Num() != Other.NormalLayers.Num() ||
-		PolygroupLayers.Num() != Other.PolygroupLayers.Num())
+		PolygroupLayers.Num() != Other.PolygroupLayers.Num() ||
+		WeightLayers.Num() != Other.WeightLayers.Num())
 	{
 		return false;
 	}
@@ -898,6 +1006,14 @@ bool FDynamicMeshAttributeSet::IsSameAs(const FDynamicMeshAttributeSet& Other, b
 		}
 	}
 
+	for (int Idx = 0; Idx < WeightLayers.Num(); Idx++)
+	{
+		if (!WeightLayers[Idx].IsSameAs(Other.WeightLayers[Idx], bIgnoreDataLayout))
+		{
+			return false;
+		}
+	}
+
 	if (HasPrimaryColors() != Other.HasPrimaryColors())
 	{
 		return false;
@@ -1020,6 +1136,12 @@ void FDynamicMeshAttributeSet::Serialize(FArchive& Ar, const FCompactMaps* Compa
 		SerializeLayers(UVLayers, Ar, CompactMaps, bUseCompression);
 		SerializeLayers(NormalLayers, Ar, CompactMaps, bUseCompression);
 		SerializeLayers(PolygroupLayers, Ar, CompactMaps, bUseCompression);
+
+		const bool bSerializeWeightLayers = !Ar.IsLoading() || Ar.CustomVer(FUE5MainStreamObjectVersion::GUID) >= FUE5MainStreamObjectVersion::DynamicMeshAttributesWeightMapsAndNames;
+		if (bSerializeWeightLayers)
+		{
+			SerializeLayers(WeightLayers, Ar, CompactMaps, bUseCompression);
+		}
 	}
 
 	if (Ar.IsLoading())
@@ -1038,6 +1160,10 @@ void FDynamicMeshAttributeSet::Serialize(FArchive& Ar, const FCompactMaps* Compa
 		{
 			Attr.ParentMesh = ParentMesh;
 		}
+		for (FDynamicMeshWeightAttribute& Attr : WeightLayers)
+		{
+			Attr.Parent = ParentMesh;
+		}
 	}
 
 	// Use int32 here to future-proof for multiple color layers.
diff --git a/Engine/Source/Runtime/GeometryCore/Private/DynamicMesh/DynamicMeshChangeTracker.cpp b/Engine/Source/Runtime/GeometryCore/Private/DynamicMesh/DynamicMeshChangeTracker.cpp
--- a/Engine/Source/Runtime/GeometryCore/Private/DynamicMesh/DynamicMeshChangeTracker.cpp
+++ b/Engine/Source/Runtime/GeometryCore/Private/DynamicMesh/DynamicMeshChangeTracker.cpp
@@ -425,6 +425,11 @@ void FDynamicMeshAttributeSetChangeTracker::BeginChange()
 		Change->PolygroupChanges.SetNum(NumPolygroupLayers);
 	}
 
+	for (int k = 0, NumWeights = Attribs->NumWeightLayers(); k < NumWeights; k++)
+	{
+		Change->WeightChanges.Add(Attribs->GetWeightLayer(k)->NewBlankChange());
+	}
+
 	for (int k = 0, NumRegAttrib = Attribs->NumRegisteredAttributes(); k < NumRegAttrib; k++)
 	{
 		Change->RegisteredAttributeChanges.Add(Attribs->GetRegisteredAttribute(k)->NewBlankChange());
@@ -500,6 +505,13 @@ void FDynamicMeshAttributeSetChangeTracker::SaveInitialTriangle(int TriangleID)
 	{
 		Change->RegisteredAttributeChanges[k]->SaveInitialTriangle(Attribs->GetRegisteredAttribute(k), TriangleID);
 	}
+
+	int NumWeightLayers = Attribs->NumWeightLayers();
+	for (int k = 0; k < NumWeightLayers; ++k)
+	{
+		Change->WeightChanges[k]->SaveInitialTriangle(Attribs->GetWeightLayer(k), TriangleID);
+	}
+
 }
 
 
@@ -509,6 +521,13 @@ void FDynamicMeshAttributeSetChangeTracker::SaveInitialVertex(int VertexID)
 	{
 		Change->RegisteredAttributeChanges[k]->SaveInitialVertex(Attribs->GetRegisteredAttribute(k), VertexID);
 	}
+
+	int NumWeightLayers = Attribs->NumWeightLayers();
+	for (int k = 0; k < NumWeightLayers; ++k)
+	{
+		Change->WeightChanges[k]->SaveInitialVertex(Attribs->GetWeightLayer(k), VertexID);
+	}
+
 }
 
 
@@ -609,6 +628,11 @@ void FDynamicMeshAttributeSetChangeTracker::StoreAllFinalTriangles(const TArray<
 	{
 		Change->RegisteredAttributeChanges[k]->StoreAllFinalTriangles(Attribs->GetRegisteredAttribute(k), TriangleIDs);
 	}
+
+	for (int k = 0, NumWeightLayers = Attribs->NumWeightLayers(); k < NumWeightLayers; k++)
+	{
+		Change->WeightChanges[k]->StoreAllFinalTriangles(Attribs->GetWeightLayer(k), TriangleIDs);
+	}
 }
 
 
@@ -618,6 +642,11 @@ void FDynamicMeshAttributeSetChangeTracker::StoreAllFinalVertices(const TArray<i
 	{
 		Change->RegisteredAttributeChanges[k]->StoreAllFinalVertices(Attribs->GetRegisteredAttribute(k), VertexIDs);
 	}
+
+	for (int k = 0, NumWeightLayers = Attribs->NumWeightLayers(); k < NumWeightLayers; k++)
+	{
+		Change->WeightChanges[k]->StoreAllFinalVertices(Attribs->GetWeightLayer(k), VertexIDs);
+	}
 }
 
 
@@ -645,6 +674,11 @@ bool FDynamicMeshAttributeChangeSet::Apply(FDynamicMeshAttributeSet* Attributes,
 	{
 		RegisteredAttributeChanges[k]->Apply(Attributes->GetRegisteredAttribute(k), bRevert);
 	}
+	for (int k = 0, NumWeightLayers = Attributes->NumWeightLayers(); k < NumWeightLayers; k++)
+	{
+		WeightChanges[k]->Apply(Attributes->GetWeightLayer(k), bRevert);
+	}
+
 	if (ColorChange.IsSet())
 	{
 		FDynamicMeshColorOverlay* ColorLayer = Attributes->PrimaryColors();
diff --git a/Engine/Source/Runtime/GeometryCore/Public/DynamicMesh/DynamicAttribute.h b/Engine/Source/Runtime/GeometryCore/Public/DynamicMesh/DynamicAttribute.h
--- a/Engine/Source/Runtime/GeometryCore/Public/DynamicMesh/DynamicAttribute.h
+++ b/Engine/Source/Runtime/GeometryCore/Public/DynamicMesh/DynamicAttribute.h
@@ -4,6 +4,7 @@
 
 #include "DynamicMesh/DynamicMesh3.h"
 #include "MeshIndexMappings.h"
+#include "Serialization/NameAsStringProxyArchive.h"
 
 namespace UE
 {
@@ -151,7 +152,7 @@ public:
 	}
 
 
-	virtual TUniquePtr<TDynamicAttributeChangeBase<ParentType>> NewBlankChange() = 0;
+	virtual TUniquePtr<TDynamicAttributeChangeBase<ParentType>> NewBlankChange() const = 0;
 
 
 
@@ -185,6 +186,34 @@ public:
 	{
 	}
 
+	/**
+	 * Serialization operator for TDynamicAttributeBase.
+	 *
+	 * @param Ar Archive to serialize with.
+	 * @param Attr Mesh attribute to serialize.
+	 * @returns Passing down serializing archive.
+	 */
+	friend FArchive& operator<<(FArchive& Ar, TDynamicAttributeBase<ParentType>& Attr)
+	{
+		Attr.Serialize(Ar);
+		return Ar;
+	}
+
+	/**
+	* Serialize to and from an archive.
+	*
+	* @param Ar Archive to serialize with.
+	*/
+	void Serialize(FArchive& Ar)
+	{
+		Ar.UsingCustomVersion(FUE5MainStreamObjectVersion::GUID);
+		if (!Ar.IsLoading() || Ar.CustomVer(FUE5MainStreamObjectVersion::GUID) >= FUE5MainStreamObjectVersion::DynamicMeshAttributesWeightMapsAndNames)
+		{
+			FNameAsStringProxyArchive ProxyArchive(Ar);
+			ProxyArchive << Name;
+		}
+	}
+
 protected:
 
 	/** 
diff --git a/Engine/Source/Runtime/GeometryCore/Public/DynamicMesh/DynamicMeshAttributeSet.h b/Engine/Source/Runtime/GeometryCore/Public/DynamicMesh/DynamicMeshAttributeSet.h
--- a/Engine/Source/Runtime/GeometryCore/Public/DynamicMesh/DynamicMeshAttributeSet.h
+++ b/Engine/Source/Runtime/GeometryCore/Public/DynamicMesh/DynamicMeshAttributeSet.h
@@ -6,6 +6,7 @@
 #include "DynamicMesh/DynamicMeshOverlay.h"
 #include "DynamicMesh/DynamicMeshTriangleAttribute.h"
 #include "DynamicMesh/DynamicAttribute.h"
+#include "DynamicMesh/DynamicVertexAttribute.h"
 #include "GeometryTypes.h"
 #include "InfoTypes.h"
 #include "Containers/IndirectArray.h"
@@ -27,6 +28,9 @@ typedef TDynamicMeshScalarTriangleAttribute<int32> FDynamicMeshMaterialAttribute
 /** Per-triangle integer polygroup ID */
 typedef TDynamicMeshScalarTriangleAttribute<int32> FDynamicMeshPolygroupAttribute;
 
+/** Per-vertex scalar float weight */
+typedef TDynamicMeshVertexAttribute<float, 1> FDynamicMeshWeightAttribute;
+
 /** Forward declarations */
 template<typename ParentType>
 class TDynamicVertexSkinWeightsAttribute;
@@ -268,6 +272,21 @@ public:
 	/** @return the Polygroup layer at the given Index */
 	const FDynamicMeshPolygroupAttribute* GetPolygroupLayer(int Index) const;
 
+	//
+	// Weight layers
+	//
+
+	/** @return number of weight layers */
+	virtual int32 NumWeightLayers() const;
+
+	/** Set the number of weight layers */
+	virtual void SetNumWeightLayers(int32 Num);
+
+	/** @return the weight layer at the given Index */
+	FDynamicMeshWeightAttribute* GetWeightLayer(int Index);
+
+	/** @return the weight layer at the given Index */
+	const FDynamicMeshWeightAttribute* GetWeightLayer(int Index) const;
 
 
 
@@ -408,6 +427,7 @@ protected:
 
 	TUniquePtr<FDynamicMeshMaterialAttribute> MaterialIDAttrib;
 
+	TIndirectArray<FDynamicMeshWeightAttribute> WeightLayers;
 	TIndirectArray<FDynamicMeshPolygroupAttribute> PolygroupLayers;
 
 	using SkinWeightAttributesMap = TMap<FName, TUniquePtr<FDynamicMeshVertexSkinWeightsAttribute>>;
@@ -475,6 +495,10 @@ protected:
 		{
 			bValid = GetPolygroupLayer(PolygroupLayerIndex)->CheckValidity(bAllowNonmanifold, FailMode) && bValid;
 		}
+		for (int WeightLayerIndex = 0; WeightLayerIndex < NumWeightLayers(); WeightLayerIndex++)
+		{
+			bValid = GetWeightLayer(WeightLayerIndex)->CheckValidity(bAllowNonmanifold, FailMode) && bValid;
+		}
 		return bValid;
 	}
 };
diff --git a/Engine/Source/Runtime/GeometryCore/Public/DynamicMesh/DynamicMeshChangeTracker.h b/Engine/Source/Runtime/GeometryCore/Public/DynamicMesh/DynamicMeshChangeTracker.h
--- a/Engine/Source/Runtime/GeometryCore/Public/DynamicMesh/DynamicMeshChangeTracker.h
+++ b/Engine/Source/Runtime/GeometryCore/Public/DynamicMesh/DynamicMeshChangeTracker.h
@@ -92,6 +92,9 @@ typedef TDynamicMeshAttributeChange<float, 4> FDynamicMeshColorChange;
 /** Standard per-triangle integer attribute change type */
 typedef FDynamicMeshTriangleAttributeChange<int32,1> FDynamicMeshTriGroupChange;
 
+/** Standard weight map change type - 1-element float */
+typedef TDynamicMeshAttributeChange<float, 1> FDynamicMeshWeightChange;
+
 /**
  * FDynamicMeshAttributeChangeSet stores a set of UV and Normal changes for a FDynamicMesh3
  */
@@ -106,6 +109,8 @@ public:
 	TOptional<FDynamicMeshColorChange> ColorChange;
 	TOptional<FDynamicMeshTriGroupChange> MaterialIDAttribChange;
 	TArray<FDynamicMeshTriGroupChange> PolygroupChanges;
+
+	TArray<TUniquePtr<FDynamicMeshAttributeChangeBase>> WeightChanges;
 	TArray<TUniquePtr<FDynamicMeshAttributeChangeBase>> RegisteredAttributeChanges;
 
 	/** call ::Apply() on all the UV and Normal changes */
diff --git a/Engine/Source/Runtime/GeometryCore/Public/DynamicMesh/DynamicMeshTriangleAttribute.h b/Engine/Source/Runtime/GeometryCore/Public/DynamicMesh/DynamicMeshTriangleAttribute.h
--- a/Engine/Source/Runtime/GeometryCore/Public/DynamicMesh/DynamicMeshTriangleAttribute.h
+++ b/Engine/Source/Runtime/GeometryCore/Public/DynamicMesh/DynamicMeshTriangleAttribute.h
@@ -60,6 +60,8 @@ protected:
 	/** List of per-triangle attribute values */
 	TDynamicVector<AttribValueType> AttribValues;
 
+	using Super = FDynamicMeshAttributeBase;
+
 	friend class FDynamicMesh3;
 	friend class FDynamicMeshAttributeSet;
 
@@ -300,7 +302,7 @@ public:
 	}
 
 
-	virtual TUniquePtr<FDynamicMeshAttributeChangeBase> NewBlankChange() override
+	virtual TUniquePtr<FDynamicMeshAttributeChangeBase> NewBlankChange() const override
 	{
 		return MakeUnique<FDynamicMeshTriangleAttributeChange<AttribValueType, AttribDimension>>();
 	}
@@ -448,6 +450,8 @@ public:
 	*/
 	void Serialize(FArchive& Ar, const FCompactMaps* CompactMaps, bool bUseCompression)
 	{
+		Super::Serialize(Ar);
+
 		Ar.UsingCustomVersion(FUE5MainStreamObjectVersion::GUID);
 		if (Ar.IsLoading() && Ar.CustomVer(FUE5MainStreamObjectVersion::GUID) < FUE5MainStreamObjectVersion::DynamicMeshCompactedSerialization)
 		{
diff --git a/Engine/Source/Runtime/GeometryCore/Public/DynamicMesh/DynamicVertexAttribute.h b/Engine/Source/Runtime/GeometryCore/Public/DynamicMesh/DynamicVertexAttribute.h
--- a/Engine/Source/Runtime/GeometryCore/Public/DynamicMesh/DynamicVertexAttribute.h
+++ b/Engine/Source/Runtime/GeometryCore/Public/DynamicMesh/DynamicVertexAttribute.h
@@ -55,6 +55,10 @@ protected:
 	/** List of per-vertex attribute values */
 	TDynamicVector<AttribValueType> AttribValues;
 
+	using Super = TDynamicAttributeBase<ParentType>;
+
+	friend class FDynamicMeshAttributeSet;
+
 public:
 	/** Create an empty overlay */
 	TDynamicVertexAttribute() : Parent(nullptr)
@@ -129,7 +133,7 @@ public:
 	void CompactCopy(const FCompactMaps& CompactMaps, const TDynamicVertexAttribute<AttribValueType, AttribDimension, ParentType>& ToCopy)
 	{
 		TDynamicAttributeBase<ParentType>::CopyParentClassData(ToCopy);
-		check(CompactMaps.NumVertexMappings() >= 0 && static_cast<size_t>(CompactMaps.NumVertexMappings()) <= AttribValues.Num() / AttribDimension);
+		check(CompactMaps.NumVertexMappings() <= int(ToCopy.AttribValues.Num() / AttribDimension));
 
 		AttribValueType Data[AttribDimension];
 		for (int VID = 0, NumVID = CompactMaps.NumVertexMappings(); VID < NumVID; VID++)
@@ -339,7 +343,7 @@ public:
 		CopyValue(SplitInfo.OriginalVertex, SplitInfo.NewVertex);
 	}
 
-	virtual TUniquePtr<TDynamicAttributeChangeBase<ParentType>> NewBlankChange() override
+	virtual TUniquePtr<TDynamicAttributeChangeBase<ParentType>> NewBlankChange() const override
 	{
 		return MakeUnique<TDynamicVertexAttributeChange<AttribValueType, AttribDimension, ParentType>>();
 	}
@@ -369,6 +373,135 @@ public:
 		return true;
 	}
 
+	/**
+	 * Returns true if this AttributeSet is the same as Other.
+	 * @param bIgnoreDataLayout Ignore gaps and padding in the data layout, i.e. the exact index is irrelevant as long as valid entries are equal in value and in the same order.
+	 */
+	bool IsSameAs(const TDynamicVertexAttribute<AttribValueType, AttribDimension, ParentType>& Other, bool bIgnoreDataLayout) const
+	{
+		if (!bIgnoreDataLayout)
+		{
+			if (AttribValues.Num() != Other.AttribValues.Num())
+			{
+				return false;
+			}
+
+			for (int Idx = 0, NumValues = AttribValues.Num(); Idx < NumValues; Idx++)
+			{
+				if (AttribValues[Idx] != Other.AttribValues[Idx])
+				{
+					return false;
+				}
+			}
+		}
+		else
+		{
+			// bIgnoreDataLayout
+
+			FRefCountVector::IndexIterator ItVid = Parent->GetVerticesRefCounts().BeginIndices();
+			const FRefCountVector::IndexIterator ItVidEnd = Parent->GetVerticesRefCounts().EndIndices();
+			FRefCountVector::IndexIterator ItVidOther = Other.Parent->GetVerticesRefCounts().BeginIndices();
+			const FRefCountVector::IndexIterator ItVidEndOther = Other.Parent->GetVerticesRefCounts().EndIndices();
+
+			while (ItVid != ItVidEnd && ItVidOther != ItVidEndOther)
+			{
+				for (int32 i = 0; i < AttribDimension; ++i)
+				{
+					const AttribValueType AttribValue = AttribValues[*ItVid * AttribDimension + i];
+					const AttribValueType AttribValueOther = Other.AttribValues[*ItVidOther * AttribDimension + i];
+					if (AttribValue != AttribValueOther)
+					{
+						// Vertex attribute value is not the same.
+						return false;
+					}
+				}
+				++ItVid;
+				++ItVidOther;
+			}
+
+			if (ItVid != ItVidEnd || ItVidOther != ItVidEndOther)
+			{
+				// Number of vertex attribute values is not the same.
+				return false;
+			}
+		}
+
+		return true;
+	}
+
+
+	/**
+	 * Serialization operator for TDynamicVertexAttribute.
+	 *
+	 * @param Ar Archive to serialize with.
+	 * @param Attr Vertex attribute to serialize.
+	 * @returns Passing down serializing archive.
+	 */
+	friend FArchive& operator<<(FArchive& Ar, TDynamicVertexAttribute<AttribValueType, AttribDimension, ParentType>& Attr)
+	{
+		Attr.Serialize(Ar, nullptr, false);
+		return Ar;
+	}
+
+	/**
+	* Serialize to and from an archive.
+	*
+	* @param Ar Archive to serialize with.
+	* @param CompactMaps If this is not a null pointer, the mesh serialization compacted the vertex and/or triangle data using the provided mapping.
+	* @param bUseCompression Use compression for serializing bulk data.
+	*/
+	void Serialize(FArchive& Ar, const FCompactMaps* CompactMaps, bool bUseCompression)
+	{
+		Super::Serialize(Ar);
+
+		Ar.UsingCustomVersion(FUE5MainStreamObjectVersion::GUID);
+		if (Ar.IsLoading() && Ar.CustomVer(FUE5MainStreamObjectVersion::GUID) < FUE5MainStreamObjectVersion::DynamicMeshCompactedSerialization)
+		{
+			Ar << AttribValues;
+		}
+		else
+		{
+			auto SerializeVector = [](FArchive& Ar, auto& Vector, bool bUseCompression)
+			{
+				if (bUseCompression)
+				{
+					Vector.template Serialize<true, true>(Ar);
+				}
+				else
+				{
+					Vector.template Serialize<true, false>(Ar);
+				}
+			};
+
+			Ar << bUseCompression;
+
+			if (CompactMaps == nullptr || !CompactMaps->VertexMapIsSet())
+			{
+				SerializeVector(Ar, AttribValues, bUseCompression);
+			}
+			else
+			{
+				TDynamicVector<AttribValueType> AttribValuesCompact;
+				AttribValuesCompact.SetNum(Parent->VertexCount() * AttribDimension);
+
+				int32 VidCompact = 0;
+				for (int32 Vid = 0, Num = AttribValues.Num() / AttribDimension; Vid < Num; ++Vid)
+				{
+					const int32 VidMapping = CompactMaps->GetVertexMapping(Vid);
+					if (VidMapping != FCompactMaps::InvalidID)
+					{
+						for (int32 i = 0; i < AttribDimension; ++i)
+						{
+							AttribValuesCompact[VidCompact++] = AttribValues[Vid * AttribDimension + i];
+						}
+					}
+				}
+
+				SerializeVector(Ar, AttribValuesCompact, bUseCompression);
+			}
+		}
+	}
+
 
 protected:
 
diff --git a/Engine/Source/Runtime/GeometryCore/Public/DynamicMesh/DynamicVertexSkinWeightsAttribute.h b/Engine/Source/Runtime/GeometryCore/Public/DynamicMesh/DynamicVertexSkinWeightsAttribute.h
--- a/Engine/Source/Runtime/GeometryCore/Public/DynamicMesh/DynamicVertexSkinWeightsAttribute.h
+++ b/Engine/Source/Runtime/GeometryCore/Public/DynamicMesh/DynamicVertexSkinWeightsAttribute.h
@@ -306,7 +306,7 @@ public:
 		CopyValue(SplitInfo.OriginalVertex, SplitInfo.NewVertex);
 	}
 
-	TUniquePtr<TDynamicAttributeChangeBase<ParentType>> NewBlankChange() override
+	TUniquePtr<TDynamicAttributeChangeBase<ParentType>> NewBlankChange() const override
 	{
 		return MakeUnique<TDynamicVertexSkinWeightsAttributeChange<ParentType>>();
 	}
diff --git a/Engine/Source/Runtime/MeshConversion/Private/DynamicMeshToMeshDescription.cpp b/Engine/Source/Runtime/MeshConversion/Private/DynamicMeshToMeshDescription.cpp
--- a/Engine/Source/Runtime/MeshConversion/Private/DynamicMeshToMeshDescription.cpp
+++ b/Engine/Source/Runtime/MeshConversion/Private/DynamicMeshToMeshDescription.cpp
@@ -5,6 +5,7 @@
 #include "DynamicMesh/DynamicMeshAttributeSet.h"
 #include "DynamicMesh/DynamicMeshOverlay.h"
 #include "DynamicMesh/DynamicVertexSkinWeightsAttribute.h"
+#include "DynamicMesh/DynamicVertexAttribute.h"
 #include "MeshDescriptionBuilder.h"
 #include "DynamicMesh/MeshTangents.h"
 #include "Util/ColorConstants.h"
@@ -949,6 +950,9 @@ void FDynamicMeshToMeshDescription::Convert_NoSharedInstances(const FDynamicMesh
 	// convert polygroup layers
 	ConvertPolygroupLayers(MeshIn, MeshOut, IndexToTriangleIDMap);
 
+	// convert weight map layers
+	ConvertWeightLayers(MeshIn, MeshOut, MapV);
+
 	// Convert all attached skin weights, if we're converting a mesh description that originated from
 	// a USkeletalMesh.
 	for (const TTuple<FName, TUniquePtr<FDynamicMeshVertexSkinWeightsAttribute>>& AttributeInfo: MeshIn->Attributes()->GetSkinWeightsAttributes())
@@ -1010,3 +1014,43 @@ void FDynamicMeshToMeshDescription::ConvertPolygroupLayers(const FDynamicMesh3*
 		}
 	}
 }
+
+
+void FDynamicMeshToMeshDescription::ConvertWeightLayers(const FDynamicMesh3* MeshIn, FMeshDescription& MeshOut, const TArray<FVertexID>& IndexToVertexIDMap)
+{
+	if (MeshIn->Attributes() == nullptr) return;
+
+	TAttributesSet<FVertexID>& VertexAttribsSet = MeshOut.VertexAttributes();
+
+	for (int32 LayerIndex = 0; LayerIndex < MeshIn->Attributes()->NumWeightLayers(); ++LayerIndex)
+	{
+		const FDynamicMeshWeightAttribute* Weights = MeshIn->Attributes()->GetWeightLayer(LayerIndex);
+		FName LayerName = Weights->GetName();
+
+		// Find existing attribute with the same name. If not found, create a new one.
+		TVertexAttributesRef<float> Attribute;
+		if (VertexAttribsSet.HasAttribute(LayerName))
+		{
+			Attribute = VertexAttribsSet.GetAttributesRef<float>(LayerName);
+		}
+		else
+		{
+			VertexAttribsSet.RegisterAttribute<float>(LayerName, 1, 0, EMeshAttributeFlags::AutoGenerated);
+			Attribute = VertexAttribsSet.GetAttributesRef<float>(LayerName);
+		}
+		if (ensure(Attribute.IsValid()))
+		{
+			for (int32 InVertexID : MeshIn->VertexIndicesItr())
+			{
+				FVertexID OutVertexID = IndexToVertexIDMap[InVertexID];
+				float Value;
+				Weights->GetValue(InVertexID, &Value);
+				Attribute.Set(OutVertexID, Value);
+			}
+		}
+		else
+		{
+			UE_LOG(LogTemp, Warning, TEXT("FDynamicMeshToMeshDescription::ConvertWeightLayers - could not create attribute named %s"), *LayerName.ToString());
+		}
+	}
+}
diff --git a/Engine/Source/Runtime/MeshConversion/Private/MeshDescriptionToDynamicMesh.cpp b/Engine/Source/Runtime/MeshConversion/Private/MeshDescriptionToDynamicMesh.cpp
--- a/Engine/Source/Runtime/MeshConversion/Private/MeshDescriptionToDynamicMesh.cpp
+++ b/Engine/Source/Runtime/MeshConversion/Private/MeshDescriptionToDynamicMesh.cpp
@@ -470,6 +470,24 @@ void FMeshDescriptionToDynamicMesh::Convert(const FMeshDescription* MeshIn, FDyn
 		}
 	});
 
+
+	// find weight attribs
+	const TAttributesSet<FVertexID>& VertAttribsSet = MeshIn->VertexAttributes();
+	TArray<TVertexAttributesConstRef<float>> WeightAttribs;
+	TArray<FName> WeightAttribNames;
+	VertAttribsSet.ForEach([&](const FName AttributeName, auto AttributesRef)
+	{
+		if (FSkeletalMeshAttributes::IsReservedAttributeName(AttributeName)) return;
+		if (FSkeletalMeshAttributes::IsSkinWeightAttribute(AttributeName)) return;
+
+		if (VertAttribsSet.template HasAttributeOfType<float>(AttributeName))
+		{
+			WeightAttribs.Add(VertAttribsSet.GetAttributesRef<float>(AttributeName));
+			WeightAttribNames.Add(AttributeName);
+		}
+	});
+
+
 	if (!bDisableAttributes)
 	{
 		bool bFoundNonDefaultVertexInstanceColor = false;
@@ -482,6 +500,8 @@ void FMeshDescriptionToDynamicMesh::Convert(const FMeshDescription* MeshIn, FDyn
 			MeshOut.Attributes()->SetNumPolygroupLayers(PolygroupAttribs.Num());
 		}
 
+		MeshOut.Attributes()->SetNumWeightLayers(WeightAttribs.Num());
+
 		for (int UVLayerIndex = 0; UVLayerIndex < NumUVLayers; UVLayerIndex++)
 		{
 			auto UVFuture = Async(EAsyncExecution::ThreadPool, [&, UVLayerIndex, bUseSharedUVs]() // must copy UVLayerIndex here!
@@ -801,6 +821,27 @@ void FMeshDescriptionToDynamicMesh::Convert(const FMeshDescription* MeshIn, FDyn
 			Pending.Add(MoveTemp(PolygroupFuture));
 		}
 
+		// initialize weight layers
+		for (int32 WeightLayerIdx = 0; WeightLayerIdx < WeightAttribs.Num(); WeightLayerIdx++)
+		{
+			auto WeightFuture = Async(EAsyncExecution::ThreadPool, [this, &WeightAttribs, &WeightAttribNames, &MeshOut, WeightLayerIdx]()
+			{
+				TVertexAttributesConstRef<float> InputWeights = WeightAttribs[WeightLayerIdx];
+				FDynamicMeshWeightAttribute* OutputWeights = MeshOut.Attributes()->GetWeightLayer(WeightLayerIdx);
+				if (ensure(OutputWeights))
+				{
+					OutputWeights->SetName(WeightAttribNames[WeightLayerIdx]);
+					for (int32 vid : MeshOut.VertexIndicesItr())
+					{
+						FVertexID SourceVertexID = VertIDMap[vid];
+						float Value = InputWeights.Get(SourceVertexID);
+						OutputWeights->SetValue(vid, &Value);
+					}
+				}
+			});
+			Pending.Add(MoveTemp(WeightFuture));
+		}
+
 		// wait for all work to be done
 		for (TFuture<void>& Future : Pending)
 		{
diff --git a/Engine/Source/Runtime/MeshConversion/Public/DynamicMeshToMeshDescription.h b/Engine/Source/Runtime/MeshConversion/Public/DynamicMeshToMeshDescription.h
--- a/Engine/Source/Runtime/MeshConversion/Public/DynamicMeshToMeshDescription.h
+++ b/Engine/Source/Runtime/MeshConversion/Public/DynamicMeshToMeshDescription.h
@@ -158,4 +158,11 @@ protected:
 	 */
 	void ConvertPolygroupLayers(const FDynamicMesh3* MeshIn, FMeshDescription& MeshOut, const TArray<FTriangleID>& IndexToTriangleIDMap);
 
+	/**
+	* Transfer WeightLayers from DynamicMesh AttributeSet to MeshDescription.
+	* Will copy to existing MeshDescription VertexAttribute<float> if one with the same name exists.
+	* Otherwise will register a new one.
+	*/
+	void ConvertWeightLayers(const FDynamicMesh3* MeshIn, FMeshDescription& MeshOut, const TArray<FVertexID>& IndexToVertexIDMap);
+
 };
\ No newline at end of file
diff --git a/Engine/Source/Runtime/MeshConversion/Public/MeshDescriptionBuilder.h b/Engine/Source/Runtime/MeshConversion/Public/MeshDescriptionBuilder.h
--- a/Engine/Source/Runtime/MeshConversion/Public/MeshDescriptionBuilder.h
+++ b/Engine/Source/Runtime/MeshConversion/Public/MeshDescriptionBuilder.h
@@ -106,6 +106,22 @@ public:
 	 */
 	FPolygonGroupID AppendPolygonGroup(FName MaterialSlotName = NAME_None);
 
+	/**
+	 * Set the specified value for the named attribute / vertex index combination.
+	 * If a vertex attribute with the given name doesn't exist, it is created.
+	 */
+	template<typename T>
+	void SetVertexAttributeValue(FName AttributeName, FVertexID VertexID, const T& Value)
+	{
+		TAttributesSet<FVertexID>& VertexAttributes = MeshDescription->VertexAttributes();
+
+		if (!VertexAttributes.GetAttributesRef<T>(AttributeName).IsValid())
+		{
+			VertexAttributes.RegisterAttribute<T>(AttributeName);
+		}
+
+		VertexAttributes.GetAttributesRef<T>(AttributeName).Set(VertexID, Value);
+	}
 
 	/** Set MeshAttribute::Edge::IsHard to true for all edges */
 	void SetAllEdgesHardness(bool bHard);
