Fix for AABB FindIntersection with ray : fill the RayParamOut value to be able to get the intersection point

diff --git a/Engine/Source/Runtime/GeometryCore/Public/Intersection/IntrRay3AxisAlignedBox3.h b/Engine/Source/Runtime/GeometryCore/Public/Intersection/IntrRay3AxisAlignedBox3.h
--- a/Engine/Source/Runtime/GeometryCore/Public/Intersection/IntrRay3AxisAlignedBox3.h
+++ b/Engine/Source/Runtime/GeometryCore/Public/Intersection/IntrRay3AxisAlignedBox3.h
@@ -137,7 +137,7 @@ protected:
 
 	// internal functions
 
-	static bool DoClipping(RealType t0, RealType t1,
+	static bool DoClipping(RealType& t0, RealType& t1,
 		const TVector<RealType>& RayOrigin, const TVector<RealType>& RayDirection,
 		const TAxisAlignedBox3<RealType>& Box, bool solid, 
 		int& quantity, TVector<RealType>& Point0, TVector<RealType>& Point1, EIntersectionType& intrType)
