Initial version of a convex decomposition algorithm

Works by trying a bunch of plane cuts to break reduce the volume of the convex pieces, then merging them back together to remove those that were not helpful.  We need the merge step especially for meshes where it takes multiple plane cuts to start reducing the hull volumes, as the initial cuts in that case are not informed by the error metric.

diff --git a/Engine/Source/Runtime/GeometryCore/Private/CompGeom/ConvexDecomposition3.cpp b/Engine/Source/Runtime/GeometryCore/Private/CompGeom/ConvexDecomposition3.cpp
new file mode 100644
--- /dev/null
+++ b/Engine/Source/Runtime/GeometryCore/Private/CompGeom/ConvexDecomposition3.cpp
@@ -0,0 +1,1384 @@
+// Copyright Epic Games, Inc. All Rights Reserved.
+
+#include "CompGeom/ConvexDecomposition3.h"
+#include "CompGeom/ConvexHull3.h"
+#include "CompGeom/Delaunay2.h"
+#include "MeshQueries.h"
+#include "VectorUtil.h"
+#include "Spatial/PriorityOrderPoints.h"
+#include "VertexConnectedComponents.h"
+
+#include "DynamicMesh/DynamicMesh3.h"
+#include "DynamicMesh/Operations/MergeCoincidentMeshEdges.h"
+#include "DynamicMesh/MeshTransforms.h"
+
+#include "Async/ParallelFor.h"
+
+namespace UE
+{
+namespace Geometry
+{
+
+
+
+// Internal struct to run templated TriangleMeshType functions on the hull mesh
+// Assumes the convex part is not compacted (i.e. is using the dynamic mesh vertices)
+// Note this will include vertices that are not on the hull
+struct FConvexPartHullMeshAdapter
+{
+	FConvexPartHullMeshAdapter(FConvexDecomposition3::FConvexPart* ConvexPart) : ConvexPart(ConvexPart)
+	{
+	}
+
+	FConvexDecomposition3::FConvexPart* ConvexPart;
+
+	bool IsTriangle(int32 Index) const
+	{
+		return Index >= 0 && Index < ConvexPart->HullTriangles.Num();
+	}
+	bool IsVertex(int32 VID) const
+	{
+		return ConvexPart->InternalGeo.IsVertex(VID);
+	}
+	int32 MaxTriangleID() const
+	{
+		return ConvexPart->HullTriangles.Num();
+	}
+	int32 MaxVertexID() const
+	{
+		return ConvexPart->InternalGeo.MaxVertexID();
+	}
+	int32 TriangleCount() const
+	{
+		return ConvexPart->HullTriangles.Num();
+	}
+	int32 VertexCount() const
+	{
+		return ConvexPart->InternalGeo.VertexCount();
+	}
+	uint64 GetChangeStamp() const
+	{
+		return 0;
+	}
+	FIndex3i GetTriangle(int32 Index) const
+	{
+		return ConvexPart->HullTriangles[Index];
+	}
+	FVector3d GetVertex(int32 VID) const
+	{
+		return ConvexPart->InternalGeo.GetVertex(VID);
+	}
+
+	inline void GetTriVertices(int TID, FVector3d& V0, FVector3d& V1, FVector3d& V2) const
+	{
+		const FIndex3i& TriIndices = ConvexPart->HullTriangles[TID];
+		V0 = GetVertex(TriIndices.A);
+		V1 = GetVertex(TriIndices.B);
+		V2 = GetVertex(TriIndices.C);
+	}
+};
+
+// Helper functionality
+namespace
+{
+
+// Compute volume by summing the volume of tetrahedra formed by the mesh triangles and a given reference point
+// If we use a central point on the mesh (centroid/center of mass/etc) this can make the computation more robust,
+// and can also make the result less arbitrary for meshes with holes
+template <typename TriangleMeshType>
+double GetVolumeUsingReferencePoint(const TriangleMeshType& Mesh, FVector3d RefPt, double DimScaleFactor = 1)
+{
+	double Volume = 0.0;
+	for (int TriIdx = 0; TriIdx < Mesh.MaxTriangleID(); TriIdx++)
+	{
+		if (!Mesh.IsTriangle(TriIdx))
+		{
+			continue;
+		}
+
+		FVector3d V0, V1, V2;
+		Mesh.GetTriVertices(TriIdx, V0, V1, V2);
+
+		// (6x) volume of the tetrahedron formed by the triangles and the reference point
+		FVector3d V1mRef = (V1 - RefPt) * DimScaleFactor;
+		FVector3d V2mRef = (V2 - RefPt) * DimScaleFactor;
+		FVector3d N = V2mRef.Cross(V1mRef);
+
+		Volume += ((V0 - RefPt) * DimScaleFactor).Dot(N);
+	}
+
+	return Volume * (1.0 / 6.0);
+}
+
+
+// Partially-computed result of splitting a convex part by a plane -- just enough to evaluate what the error would be
+// (with the caveat that it cannot see if there are multiple components to be broken apart)
+struct FPartialCutResult
+{
+	TArray<FVector3d> CutVertices;
+	TArray<int32> CrossingEdgeIDs;
+	TArray<double> SignDists;
+
+	TArray<FIndex3i> HullTriangles[2];
+	double HullVolumes[2]{ 0,0 };
+	bool bSuccess = false;
+
+	// Compute + return the error of the partial result
+	// TODO: support multiple kinds of error
+	double Score(const FConvexDecomposition3::FConvexPart& Convex)
+	{
+		double HullMinusGeoVolume = HullVolumes[0] + HullVolumes[1] - Convex.GeoVolume;
+		return HullMinusGeoVolume;
+	}
+
+	void ApplyToGeo(TIndirectArray<FConvexDecomposition3::FConvexPart>& Decomposition, int32 OrigPartIdx, const FPlane3d& Plane, int32& OtherSideStartIdxOut, double PlaneTol, double ConnectedComponentTolerance)
+	{
+		FConvexDecomposition3::FConvexPart& OrigPart = Decomposition[OrigPartIdx];
+		bool bSourceGeometryVolumeUnreliable = OrigPart.bGeometryVolumeUnreliable;
+
+		int32 NewPartIdx = Decomposition.Add(new FConvexDecomposition3::FConvexPart);
+		FConvexDecomposition3::FConvexPart& NewPart = Decomposition[NewPartIdx];
+		int32 PartIndices[2]{ OrigPartIdx, NewPartIdx };
+
+		TArray<int32> OnCutEdges;
+		check(CrossingEdgeIDs.Num() == CutVertices.Num());
+
+		TMap<int32, int32> VIDMaps[2]; // mapping from indices used in HullTriangles -> indices used in OrigPart and NewPart
+		TArray<int32> HullCrossingVertexMap;
+		int32 HullCrossingVertexStart = OrigPart.InternalGeo.MaxVertexID();
+		HullCrossingVertexMap.SetNumUninitialized(CrossingEdgeIDs.Num());
+		// split all the saved crossing edges
+		for (int32 Idx = 0; Idx < CrossingEdgeIDs.Num(); Idx++)
+		{
+			int32 EID = CrossingEdgeIDs[Idx];
+			FDynamicMesh3::FEdge Edge = OrigPart.InternalGeo.GetEdge(EID);
+			DynamicMeshInfo::FEdgeSplitInfo SplitInfo;
+			double DistA = SignDists[Edge.Vert.A], DistB = SignDists[Edge.Vert.B];
+			double Param = DistA / (DistA - DistB);
+			EMeshResult SplitResult = OrigPart.InternalGeo.SplitEdge(EID, SplitInfo, Param);
+			if (ensure(SplitResult == EMeshResult::Ok))
+			{
+				HullCrossingVertexMap[Idx] = SplitInfo.NewVertex;
+			}
+			else
+			{
+				// technically this could happen if we hit the vertex valence limit; that shouldn't really happen in practice though ...
+				HullCrossingVertexMap[Idx] = -1;
+			}
+		}
+
+		// remove all fully on-plane triangles
+		for (int32 TID = 0, MaxTID = OrigPart.InternalGeo.MaxTriangleID(); TID < MaxTID; TID++)
+		{
+			if (!OrigPart.InternalGeo.IsTriangle(TID))
+			{
+				continue;
+			}
+			FIndex3i Tri = OrigPart.InternalGeo.GetTriangle(TID);
+			// Note any new vertex (ID beyond sign dist or un-set before) will be on the plane by construction (introduced by above edge split)
+			// SignDists was zero-initialized, so we don't need to update it with the new vertices, and can just assume IDs beyond the previous max are on the plane
+			if ((Tri.A >= SignDists.Num() || FMathd::Abs(SignDists[Tri.A]) <= PlaneTol) &&
+				(Tri.B >= SignDists.Num() || FMathd::Abs(SignDists[Tri.B]) <= PlaneTol) &&
+				(Tri.C >= SignDists.Num() || FMathd::Abs(SignDists[Tri.C]) <= PlaneTol))
+			{
+				EMeshResult Res = OrigPart.InternalGeo.RemoveTriangle(TID, true, false);
+			}
+		}
+		
+		// Init the new part geometry w/ a full copy of the original
+		NewPart.InternalGeo = OrigPart.InternalGeo;
+
+		// Delete the other half from each part
+		FDynamicMesh3* PartsMeshes[2] = { &OrigPart.InternalGeo, &NewPart.InternalGeo };
+		for (int32 VID = 0, MaxVID = SignDists.Num(); VID < MaxVID; VID++)
+		{
+			double SignDist = SignDists[VID];
+			if (SignDist > PlaneTol && PartsMeshes[0]->IsVertex(VID))
+			{
+				PartsMeshes[0]->RemoveVertex(VID, false);
+			}
+			else if (SignDist < -PlaneTol && PartsMeshes[1]->IsVertex(VID))
+			{
+				PartsMeshes[1]->RemoveVertex(VID, false);
+			}
+		}
+
+		// Fill in the surface on the cutting plane; return false if hole fill process was ambiguous / unreliable
+		auto HoleFill = [this, &Plane, &PartsMeshes, &PlaneTol]() -> bool
+		{
+			bool bSuccessFill = true;
+
+			FVector3d BasisU, BasisV;
+			VectorUtil::MakePerpVectors(Plane.Normal, BasisU, BasisV);
+
+			// Collect on-plane edges to do a hole-fill triangulation
+			// TODO: detect case where the second hole fill can just be a flipped duplicate of the first
+			TArray<int32> BoundaryEdges;
+			for (int32 Side = 0; Side < 2; Side++)
+			{
+				BasisU = -BasisU;
+				BoundaryEdges.Reset();
+				for (int32 EID : PartsMeshes[Side]->BoundaryEdgeIndicesItr())
+				{
+					const FDynamicMesh3::FEdge Edge = PartsMeshes[Side]->GetEdge(EID);
+					if ( // boundary edge is on plane
+						(Edge.Vert.A >= SignDists.Num() || FMathd::Abs(SignDists[Edge.Vert.A]) <= PlaneTol) &&
+						(Edge.Vert.B >= SignDists.Num() || FMathd::Abs(SignDists[Edge.Vert.B]) <= PlaneTol))
+					{
+						BoundaryEdges.Add(EID);
+					}
+				}
+
+
+				if (BoundaryEdges.IsEmpty())
+				{
+					continue;
+				}
+
+				int32 FirstEID = BoundaryEdges[0];
+				int32 FirstVID = PartsMeshes[Side]->GetEdge(FirstEID).Vert.A;
+				FVector3d Orig = PartsMeshes[Side]->GetVertex(FirstVID);
+
+				auto Proj2 = [&Orig, &BasisU, &BasisV](const FVector3d& P) -> FVector2d
+				{
+					FVector3d Diff = P - Orig;
+					return FVector2d(Diff.Dot(BasisU), Diff.Dot(BasisV));
+				};
+
+				const int32 MaxVerts = BoundaryEdges.Num() * 2;
+				TArray<FVector2d> Verts2; // 2D projections of the on-plane vertices
+				Verts2.Reserve(MaxVerts);
+				TArray<int32> V2ToVIDMap;
+				V2ToVIDMap.Reserve(MaxVerts);
+				TMap<int32, int32> VIDToV2Map; // map from mesh VID to 2d vertex index
+				VIDToV2Map.Reserve(MaxVerts);
+				TArray<FIndex2i> Edges2; // Edges indexing into the Verts2 array
+				Edges2.Reserve(BoundaryEdges.Num());
+				for (int32 EID : BoundaryEdges)
+				{
+					FIndex2i VIDs = PartsMeshes[Side]->GetOrientedBoundaryEdgeV(EID);
+					FIndex2i Edge2;
+					for (int32 SubIdx = 0; SubIdx < 2; SubIdx++)
+					{
+						int32 VID = VIDs[SubIdx];
+						int32* SearchV2Idx = VIDToV2Map.Find(VID);
+						int32 V2Idx;
+						if (SearchV2Idx)
+						{
+							V2Idx = *SearchV2Idx;
+						}
+						else
+						{
+							FVector2d V2D = Proj2(PartsMeshes[Side]->GetVertex(VID));
+							V2Idx = Verts2.Add(V2D);
+							ensure(V2Idx == V2ToVIDMap.Add(VID));
+							VIDToV2Map.Add(VIDs[SubIdx], V2Idx);
+						}
+						Edge2[SubIdx] = V2Idx;
+					}
+					Edges2.Add(Edge2);
+				}
+
+				FDelaunay2 Delaunay;
+				Delaunay.bAutomaticallyFixEdgesToDuplicateVertices = true;
+				TArray<FIndex3i> Triangles2; // Triangles indexing into the Verts2 array
+				if (Delaunay.Triangulate(Verts2, Edges2))
+				{
+					bool bWellDefinedResult = Delaunay.GetFilledTriangles(Triangles2, Edges2, FDelaunay2::EFillMode::NegativeWinding);
+					if (!bWellDefinedResult)
+					{
+						bSuccessFill = false;
+					}
+					// TODO: if result is not well defined, is there anything more robust we could do here?
+					// Perhaps fill based on the winding number of the input mesh? (But more expensive, and we'd have to handle ~coplanar cases as well)
+				}
+				if (!Triangles2.IsEmpty()) // Delaunay triangulation succeeded at generating some triangulation
+				{
+					for (const FIndex3i& Tri2 : Triangles2)
+					{
+						FIndex3i Tri3(
+							V2ToVIDMap[Tri2.A],
+							V2ToVIDMap[Tri2.B],
+							V2ToVIDMap[Tri2.C]
+						);
+						int TID = PartsMeshes[Side]->AppendTriangle(Tri3/*, TODO: Pass a group ID to mark that this is a hole-fill triangle*/);
+						if (TID < 0)
+						{
+							bSuccessFill = false;
+						}
+					}
+				}
+				else // Delaunay failed to generate a triangulation
+				{
+					bSuccessFill = false;
+
+					// fallback triangulation: just fill everything with triangle fans
+					// TODO: is this fallback method worth keeping?
+					// It should make the volumes more stable for the subsequent geometry, but could prevent hulls from fitting as closely as possible
+					// Leaving this here but disabled for now, to make it easier to test; likely we can remove it after some additional testing of triangulation failure cases
+					if constexpr (false)
+					{
+						int32 FanCenterVID = PartsMeshes[Side]->AppendVertex(Orig);
+						for (int32 EID : BoundaryEdges)
+						{
+							FIndex2i VIDs = PartsMeshes[Side]->GetOrientedBoundaryEdgeV(EID);
+							if (!ensure(VIDs.A >= 0)) // this would happen if an edge is already not a boundary edge; TODO: can this actually happen here? 
+							{
+								continue;
+							}
+							PartsMeshes[Side]->AppendTriangle(FanCenterVID, VIDs.B, VIDs.A/*, TODO: Pass a group ID to mark that this is a hole-fill triangle*/);
+						}
+
+						// split off any complete loops as their own separate fan, and put each fan-center vertex at its one-ring's centroid
+						TArray<int> TrianglesOut, ContiguousGroupLengths;
+						TArray<bool> GroupIsLoop;
+						DynamicMeshInfo::FVertexSplitInfo SplitInfo;
+						bool bHasSplit = false;
+						if (ensure(EMeshResult::Ok == PartsMeshes[Side]->GetVtxContiguousTriangles(FanCenterVID, TrianglesOut, ContiguousGroupLengths, GroupIsLoop)))
+						{
+							if (ContiguousGroupLengths.Num() > 1)
+							{
+								int32 NumLoops = 0;
+								int32 NumSpans = 0;
+								for (bool bIsLoop : GroupIsLoop)
+								{
+									NumLoops += int(bIsLoop);
+									NumSpans += int(!bIsLoop);
+								}
+
+								if (NumLoops > 1 || (NumLoops == 1 && NumSpans > 0))
+								{
+									bHasSplit = true;
+									
+									for (int GroupIdx = 0, GroupStartIdx = 0; GroupIdx < ContiguousGroupLengths.Num(); GroupStartIdx += ContiguousGroupLengths[GroupIdx++])
+									{
+										if (GroupIsLoop[GroupIdx] && (NumSpans > 0 || GroupIdx > 0))
+										{
+											if (ensure(EMeshResult::Ok == PartsMeshes[Side]->SplitVertex(FanCenterVID, TArrayView<const int>(TrianglesOut.GetData() + GroupStartIdx, ContiguousGroupLengths[GroupIdx]), SplitInfo)))
+											{
+												FVector3d Centroid;
+												PartsMeshes[Side]->GetVtxOneRingCentroid(SplitInfo.NewVertex, Centroid);
+												PartsMeshes[Side]->SetVertex(SplitInfo.NewVertex, Centroid, false);
+											}
+										}
+									}
+								}
+							}
+						}
+						FVector3d Centroid;
+						PartsMeshes[Side]->GetVtxOneRingCentroid(FanCenterVID, Centroid);
+						PartsMeshes[Side]->SetVertex(FanCenterVID, Centroid, false);
+					}
+				}
+			}
+
+			return bSuccessFill;
+		};
+
+		// TODO: We may be able to skip hole filling if we're using an error metric that is not volume-based
+		bool bHoleFillResultUnreliable = !HoleFill();
+
+		auto ComputeHullsIfMultipleComponents = [&Decomposition, &ConnectedComponentTolerance](int32 PartIdx)->int32
+		{
+			FConvexDecomposition3::FConvexPart& Part = Decomposition[PartIdx];
+			FDynamicMesh3& Mesh = Part.InternalGeo;
+			FVertexConnectedComponents VertexComponents;
+			VertexComponents.Init(Mesh);
+			VertexComponents.ConnectTriangles(Mesh);
+			bool bMultipleGroups = VertexComponents.HasMultipleComponents(Mesh, 3);
+			if (!bMultipleGroups)
+			{
+				return int32(Mesh.TriangleCount() > 0);
+			}
+			// multiple groups found; connect close vertices just to double check
+			VertexComponents.ConnectCloseVertices(Mesh, ConnectedComponentTolerance, 3);
+
+			// Additionally merge components if they overlap in space
+			// Note this only looks at bounding box overlap currently; TODO: extend to test more possible plane of separation?
+			VertexComponents.ConnectOverlappingComponents(Mesh, 3);
+
+			// Note: Similar to VertexComponents.MakeComponentMap, but instead of mapping to a contiguous range starting at zero,
+			//		this maps to Decomposition array indices, i.e.: PartIdx, Decomposition.Num(), Decomposition.Num() + 1, ..., Decomposition.Num() + N
+			TMap<int32, int32> ComponentIdxMap;
+			int32 NewComponents = 0;
+			int32 NextComponentIdx = PartIdx;
+			for (int32 VID : Mesh.VertexIndicesItr())
+			{
+				if (!Mesh.IsReferencedVertex(VID))
+				{
+					continue;
+				}
+				int32 DisjointSetID = VertexComponents.GetComponent(VID);
+				int32* ComponentIdx = ComponentIdxMap.Find(DisjointSetID);
+				if (!ComponentIdx)
+				{
+					int32 IdxInDecomposition = NextComponentIdx;
+					if (ComponentIdxMap.IsEmpty())
+					{
+						NextComponentIdx = Decomposition.Num();
+					}
+					else
+					{
+						NextComponentIdx++;
+					}
+					ComponentIdxMap.Add(DisjointSetID, IdxInDecomposition);
+					NewComponents++;
+				}
+			}
+
+			if (NewComponents < 2)
+			{
+				return NewComponents;
+			}
+
+			FDynamicMesh3 OrigMesh = MoveTemp(Mesh); // Prepare to copy the vertices and triangles out of the old mesh and into the new
+			Decomposition[PartIdx].Reset();
+			for (int32 Idx = 1; Idx < NewComponents; Idx++)
+			{
+				Decomposition.Add(new FConvexDecomposition3::FConvexPart);
+			}
+			
+			TMap<int32, int32> VertexMap; // Map vertices to their new VID in their new mesh (note: contains mappings into every component)
+			for (int32 VID : OrigMesh.VertexIndicesItr())
+			{
+				if (!OrigMesh.IsReferencedVertex(VID))
+				{
+					continue;
+				}
+				int32 DisjointSetID = VertexComponents.GetComponent(VID);
+				int32 DecompIdx = ComponentIdxMap[DisjointSetID];
+				int32 NewVID = Decomposition[DecompIdx].InternalGeo.AppendVertex(OrigMesh.GetVertex(VID));
+				VertexMap.Add(VID, NewVID);
+			}
+			for (int32 TID : OrigMesh.TriangleIndicesItr())
+			{
+				FIndex3i Tri = OrigMesh.GetTriangle(TID);
+				// TODO: if we use triangle groups, also get the group ID here and transfer it across
+				int32 DisjointSetID = VertexComponents.GetComponent(Tri.A);
+				int32 DecompIdx = ComponentIdxMap[DisjointSetID];
+				FIndex3i NewTri(VertexMap[Tri.A], VertexMap[Tri.B], VertexMap[Tri.C]);
+				Decomposition[DecompIdx].InternalGeo.AppendTriangle(NewTri);
+			}
+			for (const TPair<int32, int32>& ToDecompIdx : ComponentIdxMap)
+			{
+				int32 DecompIdx = ToDecompIdx.Value;
+				Decomposition[DecompIdx].ComputeHull();
+				Decomposition[DecompIdx].ComputeStats();
+			}
+
+			return NewComponents;
+		};
+
+		// Compact both side meshes and transfer hull vertices OR compute new hulls if there were multiple connected components
+		int32 ComponentsOnSide0 = 1;
+		OtherSideStartIdxOut = NewPartIdx;
+		for (int32 Side = 0; Side < 2; Side++)
+		{
+			FConvexDecomposition3::FConvexPart* Part;
+			Part = Side == 0 ? &OrigPart : &NewPart;
+
+			Part->HullTriangles.Reset(HullTriangles[Side].Num());
+			Part->HullTriangles.Reset();
+
+			// If there are multiple components, we'll need to compute all new hulls for each of them
+			int32 NumComponents = ComputeHullsIfMultipleComponents(PartIndices[Side]);
+			if (NumComponents > 1)
+			{
+				if (Side == 0) // keep the components from each side of the plane contiguous
+				{
+					ComponentsOnSide0 = NumComponents;
+					int32 LastIdx = Decomposition.Num() - 1;
+					Decomposition.Swap(NewPartIdx, LastIdx);
+					NewPartIdx = LastIdx;
+					OtherSideStartIdxOut = NewPartIdx;
+				}
+				continue;
+			}
+
+			// There was only one component -- we can re-use the convex hull we computed, just need to fix the vertex indices
+			FCompactMaps CompactMaps;
+			PartsMeshes[Side]->CompactInPlace(&CompactMaps);
+			TMap<int32, int32> AddedVertsMap;
+
+			auto RemapIndex = [this, &CompactMaps, &HullCrossingVertexStart, &HullCrossingVertexMap, &PartsMeshes, &Side, &AddedVertsMap](int32 OrigHullIdx)
+			{
+				if (OrigHullIdx < HullCrossingVertexStart)
+				{
+					return CompactMaps.GetVertexMapping(OrigHullIdx);
+				}
+				else
+				{
+					int32 CrossingVertexIdx = OrigHullIdx - HullCrossingVertexStart;
+					int32 NewIdx = HullCrossingVertexMap[CrossingVertexIdx];
+					if (NewIdx == -1)
+					{
+						int32* AddedVertPtr = AddedVertsMap.Find(CrossingVertexIdx);
+						if (!AddedVertPtr)
+						{
+							int32 AppendedVertID = PartsMeshes[Side]->AppendVertex(CutVertices[CrossingVertexIdx]);
+							AddedVertsMap.Add(CrossingVertexIdx, AppendedVertID);
+							return AppendedVertID;
+						}
+						else
+						{
+							return *AddedVertPtr;
+						}
+					}
+					else
+					{
+						NewIdx = CompactMaps.GetVertexMapping(NewIdx);
+					}
+					return NewIdx;
+				}
+			};
+			for (FIndex3i& HullTri : HullTriangles[Side])
+			{
+				HullTri.A = RemapIndex(HullTri.A);
+				HullTri.B = RemapIndex(HullTri.B);
+				HullTri.C = RemapIndex(HullTri.C);
+			}
+			Part->HullTriangles = MoveTemp(HullTriangles[Side]);
+
+			// Compute all the extra data the FConvexPart requires as well
+			Part->ComputeHullPlanes();
+			Part->GeoCenter = TMeshQueries<FDynamicMesh3>::GetMeshVerticesCentroid(Part->InternalGeo);
+			Part->GeoVolume = GetVolumeUsingReferencePoint<FDynamicMesh3>(Part->InternalGeo, Part->GeoCenter);
+			if (Part->GeoVolume < 0)
+			{
+				Part->bGeometryVolumeUnreliable = true;
+				Part->GeoVolume = 0;
+			}
+			Part->HullVolume = HullVolumes[Side];
+
+			Part->Bounds = Part->InternalGeo.GetBounds();
+			Part->HullError = Part->HullVolume - Part->GeoVolume;
+		}
+
+		// Propagate down the flag that the geometry volume may be incorrect, if it was set on the source or if the hole fill had issues
+		if (bSourceGeometryVolumeUnreliable || bHoleFillResultUnreliable)
+		{
+			OrigPart.bGeometryVolumeUnreliable = true;
+			for (int32 PartIdx = NewPartIdx; PartIdx < Decomposition.Num(); PartIdx++)
+			{
+				Decomposition[PartIdx].bGeometryVolumeUnreliable = true;
+			}
+		}
+	}
+
+	FPartialCutResult() {}
+
+	FPartialCutResult(const FConvexDecomposition3::FConvexPart& Convex, const FPlane3d& Plane, double PlaneTol)
+	{
+		SignDists.SetNumZeroed(Convex.InternalGeo.MaxVertexID());
+		TArray<bool> OnHull; // indicates if vertex was on the original hull
+		OnHull.SetNumZeroed(SignDists.Num());
+		TArray<bool> ForHull[2]; // indicates if vertex should be used when making the new hull
+		ForHull[0].SetNumZeroed(Convex.InternalGeo.MaxVertexID());
+		ForHull[1].SetNumZeroed(Convex.InternalGeo.MaxVertexID());
+
+		for (int32 VID : Convex.InternalGeo.VertexIndicesItr())
+		{
+			if (!Convex.InternalGeo.IsReferencedVertex(VID))
+			{
+				continue;
+			}
+			double SD = Plane.DistanceTo(Convex.InternalGeo.GetVertex(VID));
+			SignDists[VID] = SD;
+			if (SD < -PlaneTol)
+			{
+				ForHull[0][VID] = true;
+			}
+			else if (SD > PlaneTol)
+			{
+				ForHull[1][VID] = true;
+			}
+		}
+		for (int32 EID : Convex.InternalGeo.EdgeIndicesItr())
+		{
+			FIndex2i EdgeVID = Convex.InternalGeo.GetEdgeV(EID);
+			double DistA = SignDists[EdgeVID.A];
+			double DistB = SignDists[EdgeVID.B];
+			FInterval1d Interval = FInterval1d::MakeFromUnordered(DistA, DistB);
+			
+			bool OnSide[2]{ Interval.Min < -PlaneTol, Interval.Max > PlaneTol };
+			if (OnSide[0] || OnSide[1])
+			{
+				if (OnSide[0] && OnSide[1]) // crossing case
+				{
+					CrossingEdgeIDs.Add(EID);
+					double t = DistA / (DistA - DistB);
+					FVector3d OnPlane = Lerp(
+						Convex.InternalGeo.GetVertex(EdgeVID.A),
+						Convex.InternalGeo.GetVertex(EdgeVID.B), t);
+					CutVertices.Add(OnPlane); // crossing == always goes to both hulls
+					// ForHull values already marked for the endpoints, by the per-vertex loop when signed distances were calculated
+				}
+				else // one-side case
+				{
+					int32 Side = OnSide[0] ? 0 : 1;
+					ForHull[Side][EdgeVID.A] = true;
+					ForHull[Side][EdgeVID.B] = true;
+				}
+			}
+			// else {} Note: fully on-plane case does not need anything
+		}
+
+		// compute the hulls, s.t. the hull triangle indices are referencing the original internal geo vertices w/ the cut vertices appended
+
+		// crossing range is the range of signed distances spanned by original hull triangles that cross the plane
+		// within that range, we need to also consider non-hull vertices when creating the new hulls
+		if (!Convex.HullTriangles.IsEmpty())
+		{
+			FInterval1d CrossingRange(0, 0);
+			for (const FIndex3i& HullTri : Convex.HullTriangles)
+			{
+				OnHull[HullTri.A] = true;
+				OnHull[HullTri.B] = true;
+				OnHull[HullTri.C] = true;
+				FVector3d Signs(
+					SignDists[HullTri.A],
+					SignDists[HullTri.B],
+					SignDists[HullTri.C]
+				);
+				FInterval1d TriRange(Signs.GetMin(), Signs.GetMax());
+				if (TriRange.Contains(0))
+				{
+					CrossingRange.Contain(TriRange);
+				}
+			}
+			// filter out the non-hull vertices that are not in the crossing range, since these are covered by the re-used hull points
+			for (int32 VID : Convex.InternalGeo.VertexIndicesItr())
+			{
+				if (!OnHull[VID] && !CrossingRange.Contains(SignDists[VID]))
+				{
+					ForHull[0][VID] = false;
+					ForHull[1][VID] = false;
+				}
+			}
+		}
+
+		for (int32 Side = 0; Side < 2; Side++)
+		{
+			FConvexHull3d HullCompute;
+			bool bOK = HullCompute.Solve(Convex.InternalGeo.MaxVertexID() + CutVertices.Num(),
+				[this, &Convex](int32 Index)
+				{
+					int32 MaxVID = Convex.InternalGeo.MaxVertexID();
+					return Index < MaxVID ? Convex.InternalGeo.GetVertex(Index) : CutVertices[Index - MaxVID];
+				},
+				[this, &Convex, &ForHull, Side](int32 Index) { return Index < Convex.InternalGeo.MaxVertexID() ? ForHull[Side][Index] : true; });
+			if (!bOK)
+			{
+				// don't cut the hull on a plane that would lead to a degenerate hull on either side
+				CutVertices.Empty();
+				CrossingEdgeIDs.Empty();
+				bSuccess = false;
+				return;
+			}
+			HullTriangles[Side] = HullCompute.MoveTriangles();
+
+			// Custom hull volume calculation to account for the triangles indexing partly into InternalGeo, partly into CutVertices
+			HullVolumes[Side] = 0;
+			int32 MaxVID = Convex.InternalGeo.MaxVertexID();
+			for (const FIndex3i& Tri : HullTriangles[Side])
+			{
+				FVector3d Verts[3];
+				for (int32 SubIdx = 0; SubIdx < 3; SubIdx++)
+				{
+					int32 VID = Tri[SubIdx];
+					Verts[SubIdx] = VID < MaxVID ? Convex.InternalGeo.GetVertex(VID) : CutVertices[VID - MaxVID];
+				}
+
+				// Get cross product of edges and (un-normalized) normal vector.
+				FVector3d V1mRef = (Verts[1] - Convex.GeoCenter);
+				FVector3d V2mRef = (Verts[2] - Convex.GeoCenter);
+				FVector3d N = V2mRef.Cross(V1mRef);
+
+				HullVolumes[Side] += ((Verts[0] - Convex.GeoCenter)).Dot(N);
+			}
+			HullVolumes[Side] *= (1.0 / 6.0);
+		}
+
+		bSuccess = true;
+	}
+
+};
+
+} // end anonymous namespace for helper functionality
+
+void FConvexDecomposition3::FConvexPart::Compact()
+{
+	ensure(!bIsCompact);
+
+	// remove all triangles, and all verties that are not reference by the hull
+	// and remap HullTriangles to the new compact indices
+	TMap<int32, int32> VertexRemapIDs;
+	FDynamicMesh3 VertexMesh;
+	for (FIndex3i& Tri : HullTriangles)
+	{
+		for (int32 SubIdx = 0; SubIdx < 3; SubIdx++)
+		{
+			int32 OldVID = Tri[SubIdx];
+			int32* NewVIDFound = VertexRemapIDs.Find(OldVID);
+			int32 NewVID;
+			if (NewVIDFound)
+			{
+				NewVID = *NewVIDFound;
+			}
+			else
+			{
+				NewVID = VertexMesh.AppendVertex(InternalGeo.GetVertex(OldVID));
+				VertexRemapIDs.Add(OldVID, NewVID);
+			}
+			Tri[SubIdx] = NewVID;
+		}
+	}
+	InternalGeo = MoveTemp(VertexMesh);
+	bIsCompact = true;
+}
+
+void FConvexDecomposition3::InitializeFromMesh(const FDynamicMesh3& SourceMesh, bool bMergeEdges)
+{
+	Decomposition.Empty();
+	FConvexPart* Convex = new FConvexPart(SourceMesh, bMergeEdges, ResultTransform);
+	if (Convex->IsFailed())
+	{
+		delete Convex;
+		return;
+	}
+	Decomposition.Add(Convex);
+}
+
+FConvexDecomposition3::FConvexPart::FConvexPart(const FDynamicMesh3& SourceMesh, bool bMergeEdges, FTransformSRT3d& TransformOut)
+{
+	// Copy out the source mesh
+	InternalGeo.Copy(SourceMesh, false, false, false, false);
+
+	// Transform the mesh to a standard unit-cube-at-origin space, so threshold have a consistent meaning
+	FAxisAlignedBox3d InitialBounds = InternalGeo.GetBounds();
+	double InvScaleFactor = FMath::Clamp(InitialBounds.MaxDim(), KINDA_SMALL_NUMBER, 1e8);
+	double ScaleFactor = 1.0 / InvScaleFactor;
+	FTransformSRT3d MeshTransform(-ScaleFactor * InitialBounds.Center());
+	MeshTransform.SetScale(ScaleFactor * MeshTransform.GetScale());
+	MeshTransforms::ApplyTransform(InternalGeo, MeshTransform);
+	// Return the inverse transform by reference, so we can put the results back into the original space
+	TransformOut = FTransformSRT3d(InitialBounds.Center());
+	TransformOut.SetScale(FVector3d(InvScaleFactor, InvScaleFactor, InvScaleFactor));
+	
+	// Weld close edges so we can sample convex edges
+	if (bMergeEdges)
+	{
+		FMergeCoincidentMeshEdges MergeEdges(&InternalGeo);
+		MergeEdges.Apply();
+	}
+
+	// Compute hull and standard measurements (volume, center, bounds, etc)
+	ComputeHull();
+	ComputeStats();
+}
+
+void FConvexDecomposition3::FConvexPart::ComputeHull(bool bComputePlanes)
+{
+	FConvexHull3d HullCompute;
+	bool bOK = HullCompute.Solve(InternalGeo.MaxVertexID(),
+		[this](int32 Index) { return InternalGeo.GetVertex(Index); },
+		[this](int32 Index) { return InternalGeo.IsVertex(Index); });
+	if (!bOK)
+	{
+		int32 Dimension = HullCompute.GetDimension();
+		bFailed = true;
+		return;
+	}
+	HullTriangles = HullCompute.MoveTriangles();
+	if (bComputePlanes)
+	{
+		ComputeHullPlanes();
+	}
+}
+
+void FConvexDecomposition3::FConvexPart::ComputeHullPlanes()
+{
+	HullPlanes.SetNum(HullTriangles.Num());
+	for (int32 TriIdx = 0; TriIdx < HullTriangles.Num(); TriIdx++)
+	{
+		const FIndex3i& Tri = HullTriangles[TriIdx];
+		FVector3d VA = InternalGeo.GetVertex(Tri.A);
+		FVector3d Normal = VectorUtil::Normal(VA, InternalGeo.GetVertex(Tri.B), InternalGeo.GetVertex(Tri.C));
+		HullPlanes[TriIdx] = FPlane3d(Normal, VA);
+	}
+}
+
+void FConvexDecomposition3::FConvexPart::ComputeStats()
+{
+	check(!bIsCompact); // cannot compute GeoVolume if the part is already compacted
+	FConvexPartHullMeshAdapter HullAdapter(this);
+	GeoCenter = TMeshQueries<FDynamicMesh3>::GetMeshVerticesCentroid(InternalGeo);
+	GeoVolume = GetVolumeUsingReferencePoint<FDynamicMesh3>(InternalGeo, GeoCenter);
+	HullVolume = GetVolumeUsingReferencePoint<FConvexPartHullMeshAdapter>(HullAdapter, GeoCenter);
+	Bounds = InternalGeo.GetBounds();
+	if (GeoVolume < 0)
+	{
+		bGeometryVolumeUnreliable = true;
+		GeoVolume = 0;
+	}
+	HullError = HullVolume - GeoVolume;
+}
+
+int32 FConvexDecomposition3::SplitWorst(bool bCanSkipUnreliableGeoVolumes)
+{
+	if (Decomposition.Num() == 0)
+	{
+		return 0;
+	}
+
+	double WorstError = Decomposition[0].HullError;
+	int32 WorstIdx = 0;
+	bool bCanSkipCurrent = bCanSkipUnreliableGeoVolumes && Decomposition[0].bGeometryVolumeUnreliable;
+	for (int32 PartIdx = 1; PartIdx < Decomposition.Num(); PartIdx++)
+	{
+		if (Decomposition[PartIdx].HullError > WorstError)
+		{
+			if (bCanSkipUnreliableGeoVolumes && !bCanSkipCurrent && Decomposition[PartIdx].bGeometryVolumeUnreliable)
+			{
+				continue;
+			}
+			WorstIdx = PartIdx;
+			WorstError = Decomposition[PartIdx].HullError;
+			bCanSkipCurrent = bCanSkipCurrent && Decomposition[PartIdx].bGeometryVolumeUnreliable;
+		}
+	}
+
+	FConvexPart& Part = Decomposition[WorstIdx];
+
+	double OrigHullVolume = Decomposition[WorstIdx].HullVolume;
+
+	TArray<FPlane3d> CandidatePlanes;
+
+	// Always start with the major axis planes
+	FVector3d Center = Part.Bounds.Center();
+	CandidatePlanes.Emplace(FVector3d::XAxisVector, Center);
+	CandidatePlanes.Emplace(FVector3d::YAxisVector, Center);
+	CandidatePlanes.Emplace(FVector3d::ZAxisVector, Center);
+	FVector3d Extents = Part.Bounds.Extents();
+	int32 MaxBoundsDim = 0;
+	if (Extents.Y > Extents.X)
+	{
+		MaxBoundsDim = 1;
+	}
+	if (Extents.Z > Extents[MaxBoundsDim])
+	{
+		MaxBoundsDim = 2;
+	}
+
+	// make candidate planes out of convex and boundary edges
+	const int32 MaxConvexEdgePlanes = 20; // TODO: Expose this as a parameter
+
+	// Structure of arrays for planes that come from convex (or open) edges
+	constexpr int32 MaxPerEdge = 5;
+	struct FConvexEdgeCandidates
+	{
+		TArray<FVector3d> Centers;
+		TArray<float> Depths;
+		TArray<float> Convexities; // will favor planes that come from edges that are deeper in the hull
+		TArray<TArray<FPlane3d, TFixedAllocator<MaxPerEdge>>> PlaneArrays;
+
+		void AddSingle(FConvexPart& ConvPart, const FVector3d& Center, const FVector3d& Normal, double Convexity = 0)
+		{
+			Centers.Add(Center);
+			PlaneArrays.Emplace_GetRef().Add(FPlane3d(Normal, Center));
+			Depths.Add((float)ConvPart.GetHullDepth(Center));
+			Convexities.Add((float)Convexity);
+		}
+
+		void AddMultiple(FConvexPart& ConvPart, const FVector3d& Center, FVector3d Normals[MaxPerEdge], int32 NumToAdd, double Convexity)
+		{
+			TArray<FPlane3d, TFixedAllocator<MaxPerEdge>> PlaneSet;
+			for (int32 NormalIdx = 0; NormalIdx < NumToAdd; NormalIdx++)
+			{
+				if (Normals[NormalIdx].Normalize())
+				{
+					PlaneSet.Add(FPlane3d(Normals[NormalIdx], Center));
+				}
+			}
+			if (PlaneSet.IsEmpty())
+			{
+				return;
+			}
+
+			Centers.Add(Center);
+			PlaneArrays.Add(PlaneSet);
+
+			constexpr int32 DepthIdx = 0, ConvexityIdx = 1;
+			Depths.Add((float)ConvPart.GetHullDepth(Center));
+			Convexities.Add((float)Convexity);
+		}
+	};
+
+	const double CosHalfAngleExtraSamplesThreshold = FMathd::Cos(FMathd::DegToRad * ConvexEdgeAngleMoreSamplesThreshold * .5);
+	const double CosHalfAngleIsConvexThreshold = FMathd::Cos(FMathd::DegToRad * ConvexEdgeAngleThreshold * .5);
+
+	FConvexEdgeCandidates ConvexEdgeCandidates;
+	for (int EID : Part.InternalGeo.EdgeIndicesItr())
+	{
+		FDynamicMesh3::FEdge Edge = Part.InternalGeo.GetEdge(EID);
+		FVector3d NormalA = Part.InternalGeo.GetTriNormal(Edge.Tri.A);
+
+		constexpr bool bIncludeBoundaryEdges = false;
+		if (Edge.Tri.B < 0) // Boundary edge; try adding a candidate plane on these just in case
+		{
+			if (!bIncludeBoundaryEdges)
+			{
+				continue;
+			}
+			FVector3d VA = Part.InternalGeo.GetVertex(Edge.Vert.A);
+			FVector3d VB = Part.InternalGeo.GetVertex(Edge.Vert.B);
+			FVector3d EdgeDir = VA - VB;
+			FVector3d CutNormal = EdgeDir.Cross(NormalA);
+			if (!CutNormal.Normalize())
+			{
+				continue;
+			}
+			ConvexEdgeCandidates.AddSingle(Part, (VA + VB) * .5, CutNormal);
+		}
+		else // Non-boundary edge; compute the angle on the edge
+		{
+			FVector3d NormalB = Part.InternalGeo.GetTriNormal(Edge.Tri.B);
+			FVector3d AvgNormal = NormalA + NormalB;
+			if (!AvgNormal.Normalize())
+			{
+				continue;
+			}
+			FIndex2i OppV = Part.InternalGeo.GetEdgeOpposingV(EID);
+			FVector3d OppAPos = Part.InternalGeo.GetVertex(OppV.A);
+			FVector3d VA = Part.InternalGeo.GetVertex(Edge.Vert.A);
+			FVector3d VB = Part.InternalGeo.GetVertex(Edge.Vert.B);
+			FVector3d EdgeDir = VA - VB;
+			FVector3d OppAEdge = OppAPos - VA;
+			if (!EdgeDir.Normalize())
+			{
+				continue;
+			}
+			// Direction of triangle A, perpendicular to the shared edge
+			FVector3d OppADir = OppAEdge - OppAEdge.Dot(EdgeDir) * EdgeDir;
+			if (!OppADir.Normalize())
+			{
+				continue;
+			}
+
+			double CosHalfAngle = AvgNormal.Dot(OppADir);
+			if (CosHalfAngle < CosHalfAngleIsConvexThreshold)
+			{
+				// not convex enough
+				continue;
+			}
+
+			// Direction of triangle B, perpendicular to the shared edge (by symmetry)
+			FVector3d OppBDir = OppADir - 2 * (OppADir - OppADir.Dot(AvgNormal) * AvgNormal);
+
+			FVector3d PlaneNormals[5];
+			PlaneNormals[0] = EdgeDir.Cross(AvgNormal);
+			PlaneNormals[1] = EdgeDir.Cross(OppADir);
+			PlaneNormals[2] = EdgeDir.Cross(OppBDir);
+			int32 NumPlanes = 3;
+			if (CosHalfAngle < CosHalfAngleExtraSamplesThreshold) // Sample additional cutting planes if the spanned angle is large enough
+			{
+				NumPlanes = 5;
+				PlaneNormals[3] = (PlaneNormals[0] + PlaneNormals[1]);
+				PlaneNormals[4] = (PlaneNormals[0] + PlaneNormals[2]);
+			}
+			ConvexEdgeCandidates.AddMultiple(Part, (VA + VB) * .5, PlaneNormals, NumPlanes, CosHalfAngle);
+		}
+	}
+
+	int32 InitialNumPlanes = CandidatePlanes.Num();
+	FPriorityOrderPoints OrderConvexEdges;
+	OrderConvexEdges.ComputeUniformSpaced(ConvexEdgeCandidates.Centers, ConvexEdgeCandidates.Depths, ConvexEdgeCandidates.Convexities, MaxConvexEdgePlanes);
+	for (int32 OrderIdx = 0; OrderIdx < OrderConvexEdges.Order.Num() && CandidatePlanes.Num() < InitialNumPlanes + MaxConvexEdgePlanes; OrderIdx++)
+	{
+		int32 CvxEdgeCandIdx = OrderConvexEdges.Order[OrderIdx];
+		CandidatePlanes.Append(ConvexEdgeCandidates.PlaneArrays[CvxEdgeCandIdx]);
+	}
+
+	// TODO: Consider doing parallel evaluation of candidate planes; however then we'd need to store the partial cut result for all of them
+	double LowestError = FMathd::MaxReal;
+	int32 BestPlaneIdx = -1;
+	FPartialCutResult BestCutResult;
+
+	for (int32 PlaneIdx = 0; PlaneIdx < CandidatePlanes.Num(); PlaneIdx++)
+	{
+		const FPlane3d& Plane = CandidatePlanes[PlaneIdx];
+		FPartialCutResult PlaneResult(Part, Plane, OnPlaneTolerance);
+		if (!PlaneResult.bSuccess)
+		{
+			continue;
+		}
+		double PlaneError = PlaneResult.Score(Part);
+		if (PlaneIdx < 3)
+		{
+			if (MaxBoundsDim == PlaneIdx) // favor the larger axis (make predictable cuts when none of the planes are reducing the hull volume)
+			{
+				PlaneError *= CutLargestAxisErrorScale;
+			}
+		}
+		if (PlaneError < LowestError)
+		{
+			BestPlaneIdx = PlaneIdx;
+			LowestError = PlaneError;
+			BestCutResult = MoveTemp(PlaneResult);
+		}
+	}
+
+	if (BestPlaneIdx == -1)
+	{
+		return 0;
+	}
+
+	int32 NewPartsStartIdx = Decomposition.Num();
+	int32 OtherSideStartIdx = -1;
+	BestCutResult.ApplyToGeo(Decomposition, WorstIdx, CandidatePlanes[BestPlaneIdx], OtherSideStartIdx, OnPlaneTolerance, ConnectedComponentTolerance);
+
+	UpdateProximitiesAfterSplit(WorstIdx, NewPartsStartIdx, CandidatePlanes[BestPlaneIdx], OtherSideStartIdx, OrigHullVolume);
+
+	return Decomposition.Num() - NewPartsStartIdx;
+}
+
+void FConvexDecomposition3::UpdateProximitiesAfterSplit(int32 SplitIdx, int32 NewIdxStart, FPlane3d CutPlane, int32 SecondSideIdxStart, double OrigHullVolume)
+{
+	auto IsMeshOnPlane = [this](int32 DecompIdx, const FPlane3d& Plane)
+	{
+		const FDynamicMesh3& Mesh = Decomposition[DecompIdx].InternalGeo;
+		const TArray<FIndex3i>& HullTriangles = Decomposition[DecompIdx].HullTriangles;
+		for (const FIndex3i& Tri : HullTriangles)
+		{
+			for (int32 SubIdx = 0; SubIdx < 3; SubIdx++)
+			{
+				FVector3d V = Mesh.GetVertex(Tri[SubIdx]);
+				// Note by construction the geometry should not be crossing far through the plane; otherwise would need to look at spans / couldn't just check for small abs ...
+				bool bOnPlane = FMathd::Abs(Plane.DistanceTo(V)) < ProximityTolerance;
+				if (bOnPlane)
+				{
+					return true;
+				}
+			}
+		}
+		return false;
+	};
+
+	bool bNoExtraComponents = NewIdxStart == SecondSideIdxStart && (SecondSideIdxStart + 1) == Decomposition.Num();
+
+	// First update the links that were connected to the previous SplitIdx
+
+	TArray<int32> ToRemoveProximities;
+	int32 StartNewProximitiesIdx = Proximities.Num();
+	for (auto It = DecompositionToProximity.CreateConstKeyIterator(SplitIdx); It; ++It)
+	{
+		int32 ProxIdx = It.Value();
+		checkSlow(Proximities[ProxIdx].Link.Contains(SplitIdx));
+		// Note: intentionally not taking any const references to the FProximity data because the array may reallocate as it expands below
+		FPlane3d Plane = Proximities[ProxIdx].Plane;
+		FIndex2i Link = Proximities[ProxIdx].Link;
+		int32 OtherDecoIdx = Link.A;
+		if (OtherDecoIdx == SplitIdx)
+		{
+			OtherDecoIdx = Link.B;
+		}
+
+		// If the newly split mesh is still on the plane associated with this previous link, assume the link is still worth considering for merges
+		FAxisAlignedBox3d Bounds = Decomposition[SplitIdx].Bounds;
+		Bounds.Expand(ProximityTolerance);
+		if (IsMeshOnPlane(SplitIdx, Plane) && Bounds.Intersects(Decomposition[OtherDecoIdx].Bounds))
+		{
+			// can re-use the old proximity; links are exactly the same
+			Proximities[ProxIdx].ClearMergedVolume();
+		}
+		else
+		{
+			// If the cut mesh no longer contacts the plane, remove the proximity
+			ToRemoveProximities.Add(ProxIdx);
+		}
+
+		// For all new decomposition elements, consider creating a new proximity link
+		for (int32 NewDecoIdx = NewIdxStart; NewDecoIdx < Decomposition.Num(); NewDecoIdx++)
+		{
+			Bounds = Decomposition[NewDecoIdx].Bounds;
+			Bounds.Expand(ProximityTolerance);
+			if (IsMeshOnPlane(NewDecoIdx, Plane) && Bounds.Intersects(Decomposition[OtherDecoIdx].Bounds))
+			{
+				Proximities.Emplace(FIndex2i(NewDecoIdx, OtherDecoIdx), Plane, true);
+			}
+		}
+	}
+
+	// Connect up new edges in the Decomposition->Proximity map
+	for (int32 ProxIdx = StartNewProximitiesIdx; ProxIdx < Proximities.Num(); ProxIdx++)
+	{
+		DecompositionToProximity.Add(Proximities[ProxIdx].Link.A, ProxIdx);
+		DecompositionToProximity.Add(Proximities[ProxIdx].Link.B, ProxIdx);
+	}
+
+	// Delete removed edges (including updating the map as needed)
+	// Note this must happen *after* connecting up new edges, or else the StartNewProximitiesIdx will be wrong
+	for (int32 ToRemove : ToRemoveProximities)
+	{
+		DeleteProximity(ToRemove, true);
+	}
+
+	// Create new proximity links across the new plane
+	bool bOrigOnPlane = IsMeshOnPlane(SplitIdx, CutPlane);
+	TArray<bool> NewPartsOnPlane; // note: this array is only non-empty when connected components were split off
+	NewPartsOnPlane.SetNum(SecondSideIdxStart - NewIdxStart);
+	for (int32 Idx = NewIdxStart; Idx < SecondSideIdxStart; Idx++)
+	{
+		NewPartsOnPlane[Idx - NewIdxStart] = IsMeshOnPlane(Idx, CutPlane);
+	}
+	int32 BeforeNumProximities = Proximities.Num();
+	auto TryLink = [this, &CutPlane](int32 DecoAIdx, int32 DecoBIdx)
+	{
+		FAxisAlignedBox3d Bounds = Decomposition[DecoAIdx].Bounds;
+		Bounds.Expand(ProximityTolerance);
+		if (Bounds.Intersects(Decomposition[DecoBIdx].Bounds))
+		{
+			int32 ProxIdx = Proximities.Emplace(FIndex2i(DecoAIdx, DecoBIdx), CutPlane, true);
+			DecompositionToProximity.Add(DecoAIdx, ProxIdx);
+			DecompositionToProximity.Add(DecoBIdx, ProxIdx);
+		}
+	};
+	for (int32 DecoSideBIdx = SecondSideIdxStart; DecoSideBIdx < Decomposition.Num(); DecoSideBIdx++)
+	{
+		if (IsMeshOnPlane(DecoSideBIdx, CutPlane))
+		{
+			if (bOrigOnPlane)
+			{
+				TryLink(DecoSideBIdx, SplitIdx);
+			}
+			for (int32 Idx = NewIdxStart; Idx < SecondSideIdxStart; Idx++)
+			{
+				if (NewPartsOnPlane[Idx - NewIdxStart])
+				{
+					TryLink(DecoSideBIdx, Idx);
+				}
+			}
+		}
+	}
+	// If there were no extra components, and a proximity was added, the merge cost is just the pre-split error
+	if (bNoExtraComponents && Proximities.Num() == BeforeNumProximities + 1)
+	{
+		Proximities.Last().MergedVolume = OrigHullVolume;
+	}
+}
+
+int32 FConvexDecomposition3::MergeBest(int32 NumMerges/*, double MaxErrorTolerance*/)
+{
+	// Support having a max error tolerance
+
+	int32 MergeNum = 0;
+
+	TMap<int32, TUniquePtr<FConvexPart>> ProximityComputedParts;
+
+	// initialize all parts with no merged hulls volume as having just their own hull volume (assuming no merges were yet performed on these parts)
+	for (FConvexPart& Part : Decomposition)
+	{
+		if (Part.SumHullsVolume == -FMathd::MaxReal)
+		{
+			Part.SumHullsVolume = Part.HullVolume;
+		}
+	}
+
+
+	auto CreateMergedPart = [this](FProximity& Prox) -> TUniquePtr<FConvexPart>
+	{
+		TUniquePtr<FConvexPart> NewPart = MakeUnique<FConvexPart>();
+		NewPart->GeoCenter = FVector3d::ZeroVector;
+		NewPart->SumHullsVolume = 0;
+		for (int32 SubIdx = 0; SubIdx < 2; SubIdx++)
+		{
+			// TODO: If we support a non-volume-based error, will need to also append triangles (and do mapping across), and cannot compact the result
+			const FConvexPart& OldPart = Decomposition[Prox.Link[SubIdx]];
+			for (FVector3d V : OldPart.InternalGeo.VerticesItr())
+			{
+				NewPart->InternalGeo.AppendVertex(V);
+			}
+			NewPart->GeoCenter += OldPart.GeoCenter;
+			NewPart->GeoVolume += OldPart.GeoVolume;
+			NewPart->SumHullsVolume += OldPart.SumHullsVolume;
+		}
+		NewPart->GeoCenter *= .5;
+		NewPart->ComputeHull();
+		FConvexPartHullMeshAdapter HullAdapter(NewPart.Get());
+		NewPart->HullVolume = GetVolumeUsingReferencePoint<FConvexPartHullMeshAdapter>(HullAdapter, NewPart->GeoCenter);
+		NewPart->Bounds = NewPart->InternalGeo.GetBounds();
+		NewPart->HullError = NewPart->HullVolume - NewPart->GeoVolume;
+		NewPart->Compact();
+
+		return NewPart;
+	};
+
+	for (; MergeNum < NumMerges; MergeNum++)
+	{
+		double BestKnownCost = FMathd::MaxReal;
+		int32 BestKnownIdx = -1;
+		for (int32 ProxIdx = 0; ProxIdx < Proximities.Num(); ProxIdx++)
+		{
+			double MergeCost;
+			if (Proximities[ProxIdx].HasMergedVolume())
+			{
+				MergeCost = Proximities[ProxIdx].GetMergeCost(Decomposition);
+			}
+			else
+			{
+				// compute the actual cost by computing what the merge part would be
+				TUniquePtr<FConvexPart>& NewPart = ProximityComputedParts.Add(ProxIdx, CreateMergedPart(Proximities[ProxIdx]));
+				Proximities[ProxIdx].MergedVolume = NewPart->HullVolume;
+				MergeCost = Proximities[ProxIdx].GetMergeCost(Decomposition);
+			}
+
+			if (MergeCost < BestKnownCost)
+			{
+				BestKnownCost = MergeCost;
+				BestKnownIdx = ProxIdx;
+			}
+		}
+
+		// Perform the best merge and update proximity data
+
+		if (BestKnownIdx == -1)
+		{
+			break;
+		}
+
+		// Get the proximity that we will merge -- note this proximity will be removed from the Proximities array
+		FProximity Prox = Proximities[BestKnownIdx];
+		int32 DecoToKeep = Prox.Link.A, DecoToRm = Prox.Link.B;
+		if (DecoToKeep > DecoToRm)
+		{
+			Swap(DecoToKeep, DecoToRm); // keep the smaller index (this makes the book-keeping update below slightly easier)
+		}
+
+		// Update the "to keep" index with the new, merged FConvexPart
+		if (ProximityComputedParts.Contains(BestKnownIdx))
+		{
+			Decomposition[DecoToKeep] = MoveTemp(*ProximityComputedParts[BestKnownIdx]);
+			ProximityComputedParts.Remove(BestKnownIdx);
+		}
+		else
+		{
+			TUniquePtr<FConvexPart> NewPart = CreateMergedPart(Prox);
+			Decomposition[DecoToKeep] = MoveTemp(*NewPart);
+		}
+
+		// Do all the book-keeping to delete the other old FConvexPart and update all Proximity data
+
+		// Clear all pre-computed merge info related to the newly-merged parts, saving out only what parts the merge should be linked to
+		TSet<int32> NewLinks;
+		for (int32 ProxIdx = 0; ProxIdx < Proximities.Num(); ProxIdx++)
+		{
+			FIndex2i Link = Proximities[ProxIdx].Link;
+			if (Link.Contains(DecoToRm) || Link.Contains(DecoToKeep))
+			{
+				// Clear everything associated with the old link
+				Proximities[ProxIdx].ClearMergedVolume();
+				ProximityComputedParts.Remove(ProxIdx);
+				NewLinks.Add(Link.A);
+				NewLinks.Add(Link.B);
+				DecompositionToProximity.RemoveSingle(Link.A, ProxIdx);
+				DecompositionToProximity.RemoveSingle(Link.B, ProxIdx);
+
+				// Swap-Remove the proximity links
+				int32 LastProxIdx = Proximities.Num() - 1;
+				Proximities.RemoveAtSwap(ProxIdx, 1, false);
+
+				// Update the proximity that was swapped back to this position (if any)
+				if (ProxIdx < LastProxIdx)
+				{
+					FIndex2i SwappedLink = Proximities[ProxIdx].Link;
+					DecompositionToProximity.RemoveSingle(SwappedLink.A, LastProxIdx);
+					DecompositionToProximity.RemoveSingle(SwappedLink.B, LastProxIdx);
+					DecompositionToProximity.Add(SwappedLink.A, ProxIdx);
+					DecompositionToProximity.Add(SwappedLink.B, ProxIdx);
+					TUniquePtr<FConvexPart>* FoundPart = ProximityComputedParts.Find(LastProxIdx);
+					if (FoundPart)
+					{
+						FConvexPart* Part = FoundPart->Release();
+						ProximityComputedParts.Remove(LastProxIdx);
+						ProximityComputedParts.Add(ProxIdx, TUniquePtr<FConvexPart>(Part)); // re-create with the new key
+					}
+					
+					// Next iteration, we'll reconsider the just-updated element
+					ProxIdx--;
+				}
+			}
+		}
+		NewLinks.Remove(DecoToKeep);
+		NewLinks.Remove(DecoToRm);
+
+		// Remove the unused convex, updating references to the part that was swapped back
+		int32 LastIdx = Decomposition.Num() - 1;
+		if (DecoToRm != LastIdx)
+		{
+			// Before swap-removing DecoToRm, update references that used to point to the last element of the array to now point to DecoToRm
+			[this](int32 OrigIdx, int32 NewIdx)
+			{
+				TArray<int32> ProxIndices;
+				DecompositionToProximity.MultiFind(OrigIdx, ProxIndices, false);
+				DecompositionToProximity.Remove(OrigIdx);
+				for (int32 ProxIdx : ProxIndices)
+				{
+					// any link to the new location should be removed before calling this Update function
+					checkSlow(!Proximities[ProxIdx].Link.Contains(NewIdx));
+
+					int32 SubIdx = Proximities[ProxIdx].Link.IndexOf(OrigIdx);
+					Proximities[ProxIdx].Link[SubIdx] = NewIdx;
+					DecompositionToProximity.Add(NewIdx, ProxIdx);
+				}
+			}(LastIdx, DecoToRm);
+		}
+		Decomposition.RemoveAtSwap(DecoToRm, 1, false);
+
+		// Add new proximities for all new links to the merged part
+		for (int32 ToLink : NewLinks)
+		{
+			if (ToLink == LastIdx)
+			{
+				ToLink = DecoToRm;
+			}
+			int32 ProxIdx = Proximities.Emplace(FIndex2i(DecoToKeep, ToLink), FPlane3d(), false /*bPlaneSeparates*/);
+			DecompositionToProximity.Add(DecoToKeep, ProxIdx);
+			DecompositionToProximity.Add(ToLink, ProxIdx);
+		}
+	}
+
+	return MergeNum;
+}
+
+void FConvexDecomposition3::DeleteProximity(int32 ProxIdx, bool bDeleteMapReferences)
+{
+	if (bDeleteMapReferences)
+	{
+		const FProximity& OldProx = Proximities[ProxIdx];
+		DecompositionToProximity.RemoveSingle(OldProx.Link.A, ProxIdx);
+		DecompositionToProximity.RemoveSingle(OldProx.Link.B, ProxIdx);
+	}
+	checkSlow(DecompositionToProximity.FindKey(ProxIdx) == nullptr); // verify there are no more references to this Proximity's index
+
+	if (ProxIdx == Proximities.Num() - 1)
+	{
+		Proximities.RemoveAt(ProxIdx, 1, false);
+		return;
+	}
+	else
+	{
+		// Remove by swapping the last element to the new slot, and removing the last element
+		// Then update links for that swapped-in proximity
+		int32 LastProxIdx = Proximities.Num() - 1;
+		Proximities.RemoveAtSwap(ProxIdx, 1, false);
+		const FProximity& NewProx = Proximities[ProxIdx];
+		DecompositionToProximity.RemoveSingle(NewProx.Link.A, LastProxIdx);
+		DecompositionToProximity.RemoveSingle(NewProx.Link.B, LastProxIdx);
+		DecompositionToProximity.Add(NewProx.Link.A, ProxIdx);
+		DecompositionToProximity.Add(NewProx.Link.B, ProxIdx);
+	}
+}
+
+
+} // end namespace UE::Geometry
+} // end namespace UE
\ No newline at end of file
diff --git a/Engine/Source/Runtime/GeometryCore/Public/CompGeom/ConvexDecomposition3.h b/Engine/Source/Runtime/GeometryCore/Public/CompGeom/ConvexDecomposition3.h
new file mode 100644
--- /dev/null
+++ b/Engine/Source/Runtime/GeometryCore/Public/CompGeom/ConvexDecomposition3.h
@@ -0,0 +1,302 @@
+// Copyright Epic Games, Inc. All Rights Reserved.
+
+#pragma once
+
+#include "CoreMinimal.h"
+#include "VectorTypes.h"
+#include "IndexTypes.h"
+#include "MathUtil.h"
+#include "BoxTypes.h"
+#include "PlaneTypes.h"
+
+#include "Containers/IndirectArray.h"
+
+#include "DynamicMesh/DynamicMesh3.h"
+
+namespace UE {
+namespace Geometry {
+
+// TODO: To support meshes where volume is not well defined (e.g., open boundaries or many self-intersecting parts), we'll need alternative error metrics
+enum class EConvexErrorMethod
+{
+	//GeometryToHullRayDistance,
+	//HullToGeometryDistance,
+	//BothDistances,
+	RelativeVolume
+};
+
+class GEOMETRYCORE_API FConvexDecomposition3
+{
+
+public:
+
+	FConvexDecomposition3() : ResultTransform(FTransformSRT3d::Identity())
+	{
+	}
+
+	FConvexDecomposition3(const FDynamicMesh3& SourceMesh, bool bMergeEdges = true)
+	{
+		InitializeFromMesh(SourceMesh, bMergeEdges);
+	}
+
+	void InitializeFromMesh(const FDynamicMesh3& SourceMesh, bool bMergeEdges);
+
+	//
+	// Settings
+	//
+
+	// Note we automatically transform the input to a centered-unit-cube space, and tolerances are measured in that space
+	// In other words, distance-based tolerances are expressed as a fraction of the bounding box max axis
+	
+	// Threshold for considering two convex parts close enough to consider merging
+	double ProximityTolerance = 1e-3;
+	// Threshold for considering two components to be touching
+	double ConnectedComponentTolerance = 1e-3;
+	// Maximum 'on plane' distance for points during planar cuts
+	double OnPlaneTolerance = 1e-7;
+
+	// Scale factor for error on the largest axis; slightly favoring cutting on the largest axis helps get more consistent result when no initial cut will reduce error (e.g. a torus)
+	double CutLargestAxisErrorScale = .99;
+
+	// Angle at edge above which we take 5 sample planes instead of 3, to better cover the wider range (expressed in degrees)
+	double ConvexEdgeAngleMoreSamplesThreshold = 30; 
+	// Angle at edge above which we don't sample it as a convex edge (expressed in degrees)
+	double ConvexEdgeAngleThreshold = 170;
+
+	// TODO: Provide hull approximation options?
+
+
+	/**
+	 * Compute a decomposition with the desired number of hulls
+	 * 
+	 * @param NumOutputHulls		Number of convex hulls to use in the final convex decomposition
+	 * @param NumAdditionalSplits	How far to go beyond the target number of hulls when initially the mesh into pieces -- larger values will require more computation but can find better convex decompositions
+	 */
+	void Compute(int32 NumOutputHulls, int32 NumAdditionalSplits = 10)
+	{
+		for (int32 SplitIdx = 0; SplitIdx < NumOutputHulls + NumAdditionalSplits; SplitIdx++)
+		{
+			SplitWorst(bool(SplitIdx % 2));
+		}
+		Compact();
+		MergeBest(NumHulls() - NumOutputHulls);
+	}
+
+	// Split the worst convex part, and return the increase in the total number of convex parts after splitting (can be more than 1 if result has multiple separate connected components)
+	// Note: could return 0 if no splits were possible
+	// @param bCanSkipUnreliableGeoVolumes		if true, don't split hulls where we have questionable geometry volume results, unless there is no hull with good geometry volume results
+	int32 SplitWorst(bool bCanSkipUnreliableGeoVolumes = false);
+
+	// Merge the pair of convexes in the decomposition that least-increases the error
+	// @param NumMerges			The target number of merges to perform.  Intermediate results can be used across merges, so it is best to do all merges in one call.
+	// @param ErrorTolerance	TODO: Refuse to merge if the resulting error would be greater than this value
+	// @return					The number of merges performed
+	int32 MergeBest(int32 NumMerges/*, double MaxErrorTolerance = FMathd::MaxReal*/);
+
+	// TODO: Implement a MergeSmall method to provide a way to ensure removal of low-quality parts
+	// Use merges remove convex hulls that are too small or too flat
+	// @param MinVolumeTolerance		Hulls with a smaller volume than this will be removed by merges (units in cm^3)
+	// @param MinThicknessTolerance		Hulls thinner than this will be removed by merges (units in cm)
+	// Note hull thickness is measured by first fitting a plane to the hull points then measuring their span in the plane normal direction
+	//int32 MergeSmall(double MinVolumeTolerance, double MinThicknessTolerance = FMathd::MaxReal);
+
+	// Compact the decomposition representation to the minimal needed for output: Just the hull vertices and triangles, no internal geometry
+	// Only call this if no more decomposition will be performed, but merges can still be performed on the compact representation
+	void Compact()
+	{
+		for (FConvexPart& Part : Decomposition)
+		{
+			Part.Compact();
+		}
+	}
+
+	//
+	// Results
+	//
+
+	int32 NumHulls() const
+	{
+		return Decomposition.Num();
+	}
+
+	/** @return convex hull triangles */
+	TArray<FIndex3i> const& GetTriangles(int32 HullIdx) const
+	{
+		return Decomposition[HullIdx].HullTriangles;
+	}
+
+	template<typename RealType>
+	TArray<TVector<RealType>> GetVertices(int32 HullIdx, bool bTransformedToOutput = true) const
+	{
+		TArray<TVector<RealType>> Vertices;
+		Vertices.SetNum(Decomposition[HullIdx].InternalGeo.MaxVertexID());
+		for (int32 VID = 0; VID < Decomposition[HullIdx].InternalGeo.MaxVertexID(); VID++)
+		{
+			if (Decomposition[HullIdx].InternalGeo.IsVertex(VID))
+			{
+				Vertices[VID] = (TVector<RealType>) ResultTransform.TransformPosition(Decomposition[HullIdx].InternalGeo.GetVertex(VID));
+			}
+		}
+		return Vertices;
+	}
+
+	// Get the non-hull geometry underlying a part of the convex decomposition
+	const FDynamicMesh3& GetInternalMesh(int32 HullIdx) const
+	{
+		return Decomposition[HullIdx].InternalGeo;
+	}
+
+	// Representation of a convex hull in the decomposition + associated information to help further split or merge
+	struct FConvexPart
+	{
+		FConvexPart() {}
+		FConvexPart(const FDynamicMesh3& SourceMesh, bool bMergeEdges, FTransformSRT3d& TransformOut);
+
+		void Reset()
+		{
+			InternalGeo.Clear();
+			HullTriangles.Reset();
+			HullPlanes.Reset();
+			HullVolume = 0;
+			GeoVolume = 0;
+			GeoCenter = FVector3d::ZeroVector;
+			Bounds = FAxisAlignedBox3d::Empty();
+			HullError = FMathd::MaxReal;
+			bIsCompact = false;
+			bFailed = false;
+			bGeometryVolumeUnreliable = false;
+		}
+
+		// Underlying geometry represented by the convex part
+		// Note: If IsCompact(), this will only store vertices
+		FDynamicMesh3 InternalGeo;
+
+		inline bool IsCompact() const
+		{
+			return bIsCompact;
+		}
+
+		inline bool IsFailed() const
+		{
+			return bFailed;
+		}
+
+		TArray<FIndex3i> HullTriangles;
+		TArray<FPlane3d> HullPlanes; // 1:1 with HullTriangles
+
+		// Measurements of the geo and hull, to be used when evaluating potential further splits
+		double HullVolume = 0, GeoVolume = 0;
+		double SumHullsVolume = -FMathd::MaxReal; // Sum of volume of any hulls that have been merged to form this hull; only valid if the FConvexPart was created by merging
+		FVector3d GeoCenter = FVector3d::ZeroVector; // Some central point of the geometry (e.g., an average of vertices)
+		FAxisAlignedBox3d Bounds;
+		bool bGeometryVolumeUnreliable = false;
+
+		// Approximate value indicating how badly the hull mismatches the input geometry
+		double HullError = FMathd::MaxReal;
+
+		void ComputeHullPlanes();
+
+		// get the depth of a point inside the hull (positive values are inside the hull, negative outside)
+		double GetHullDepth(FVector3d Pt)
+		{
+			int32 NumPlanes = HullPlanes.Num();
+			ensure(HullTriangles.Num() == NumPlanes);
+			if (NumPlanes == 0)
+			{
+				return 0;
+			}
+			double MaxDist = HullPlanes[0].DistanceTo(Pt);
+			for (int32 Idx = 0; Idx < NumPlanes; Idx++)
+			{
+				const FPlane3d& Plane = HullPlanes[Idx];
+				double Dist = Plane.DistanceTo(Pt);
+				MaxDist = FMath::Max(Dist, MaxDist);
+			}
+			return -MaxDist;
+		}
+
+		// Helper to create hull after InternalGeo is set
+		void ComputeHull(bool bComputePlanes = true);
+
+		// Helper to compute volumes, centroid, bounds, etc after the convex part is initialized
+		void ComputeStats();
+
+		void Compact();
+
+	protected:
+		bool bIsCompact = false;
+		bool bFailed = false;
+	};
+
+	// All convex hulls parts in the convex decomposition
+	// Note this must be an indirect array, not a TArray, because FConvexPart holds an FDynamicMesh3 mesh, which is not trivially movable
+	TIndirectArray<FConvexPart> Decomposition;
+
+	/// Transform taking the result hull vertices back to the original space of the inputs; automatically applied by GetVertices()
+	FTransformSRT3d ResultTransform = FTransformSRT3d::Identity();
+
+	// Edge in a graph of merge candidates, indicating which Decomposition parts could be merged together
+	// Note a proximity link does not guarantee that the geometry (or convex hulls) are touching
+	struct FProximity
+	{
+		FProximity() = default;
+		FProximity(const FIndex2i& Link, const FPlane3d& Plane, bool bPlaneSeparates) : Link(Link), Plane(Plane), bPlaneSeparates(bPlaneSeparates)
+		{}
+
+		FIndex2i Link{ -1, -1 };
+		
+		// the shared plane between the two hulls (because adjacent hulls are formed by planar cuts in the splitting step)
+		// Note: Plane becomes invalid in the merge-up phase; it is used to accelerate filtering in the split-down phase
+		FPlane3d Plane;
+		
+		// Whether the geometry is separated by Plane
+		bool bPlaneSeparates = false;
+		
+		// What the volume of the convex hull would be if we merged these two parts
+		double MergedVolume = VolumeNotComputed();
+
+		double GetMergeCost(const TIndirectArray<FConvexPart>& DecompositionIn) const
+		{
+			if (!HasMergedVolume() || !DecompositionIn.IsValidIndex(Link.A) || !DecompositionIn.IsValidIndex(Link.B))
+			{
+				return FMathd::MaxReal;
+			}
+			ensure(DecompositionIn[Link.A].SumHullsVolume != -FMathd::MaxReal && DecompositionIn[Link.B].SumHullsVolume != -FMathd::MaxReal);
+			return MergedVolume - DecompositionIn[Link.A].SumHullsVolume - DecompositionIn[Link.B].SumHullsVolume;
+		}
+
+		bool HasMergedVolume() const
+		{
+			return MergedVolume != VolumeNotComputed();
+		}
+		void ClearMergedVolume()
+		{
+			MergedVolume = VolumeNotComputed();
+		}
+
+		constexpr double VolumeNotComputed() const { return FMathd::MaxReal; }
+		
+	};
+
+	// All potential merge edges between Decomposition FConvexParts
+	TArray<FProximity> Proximities;
+
+	// Mapping from Decomposition indices to Proximity indices
+	TMultiMap<int32, int32> DecompositionToProximity;
+
+	// Helper function to delete a proximity relationship, and fix index references to account for the new array ordering
+	// @param bDeleteMapReferences		If true, also update references in the DecompositionToProximity map
+	void DeleteProximity(int32 ProxIdx, bool bDeleteMapReferences);
+
+	// Helper function to update all proximity data after splitting an FConvexPart apart
+	// @param SplitIdx				The index of the FConvexPart that has been split into multiple parts.  Now contains the first new part.
+	// @param NewIdxStart			The first index in the Decomposition array of new FConvexParts aside from SplitIdx.  Everything from this to the end of the Decomposition array must be new.
+	// @param CutPlane				The plane used for the split
+	// @param SecondSideIdxStart	The first index in the Decomposition array of new FConvexParts that were on the opposite side of the plane from SplitIdx.  Will be different from NewIdxStart only if the first part was split into multiple components.
+	// @param OrigHullVolume		The volume of the convex hull *before* the split was performed -- used to precompute the merge cost.
+	void UpdateProximitiesAfterSplit(int32 SplitIdx, int32 NewIdxStart, FPlane3d CutPlane, int32 SecondSideIdxStart, double OrigHullVolume);
+};
+
+} // end namespace UE::Geometry
+} // end namespace UE
+
diff --git a/Engine/Source/Runtime/GeometryCore/Public/CompGeom/ConvexHull3.h b/Engine/Source/Runtime/GeometryCore/Public/CompGeom/ConvexHull3.h
--- a/Engine/Source/Runtime/GeometryCore/Public/CompGeom/ConvexHull3.h
+++ b/Engine/Source/Runtime/GeometryCore/Public/CompGeom/ConvexHull3.h
@@ -15,8 +15,6 @@ namespace Geometry {
 
 using namespace UE::Math;
 
-template <typename RealType> struct TConvexHull3Internal;
-
 /**
  * Helper class to find the dimensions spanned by a point cloud
  * and (if it spans 3 dimensions) the indices of four 'extreme' points
@@ -123,6 +121,12 @@ public:
 		return Hull;
 	}
 
+	/** @return convex hull triangles by a move */
+	TArray<FIndex3i>&& MoveTriangles()
+	{
+		return MoveTemp(Hull);
+	}
+
 	/** 
 	 * Convert an already-computed convex hull into a halfspace representation
 	 * Following the logic of ContainmentQueries3.h, all halfspaces are oriented "outwards,"
