Add optional error tolerances to guide convex decomposition

diff --git a/Engine/Source/Runtime/GeometryCore/Private/CompGeom/ConvexDecomposition3.cpp b/Engine/Source/Runtime/GeometryCore/Private/CompGeom/ConvexDecomposition3.cpp
--- a/Engine/Source/Runtime/GeometryCore/Private/CompGeom/ConvexDecomposition3.cpp
+++ b/Engine/Source/Runtime/GeometryCore/Private/CompGeom/ConvexDecomposition3.cpp
@@ -824,18 +824,25 @@ void FConvexDecomposition3::FConvexPart::ComputeStats()
 	HullError = HullVolume - GeoVolume;
 }
 
-int32 FConvexDecomposition3::SplitWorst(bool bCanSkipUnreliableGeoVolumes)
+int32 FConvexDecomposition3::SplitWorst(bool bCanSkipUnreliableGeoVolumes, double ErrorTolerance)
 {
 	if (Decomposition.Num() == 0)
 	{
 		return 0;
 	}
 
+	double VolumeTolerance = ConvertDistanceToleranceToLocalVolumeTolerance(ErrorTolerance);
+
 	double WorstError = Decomposition[0].HullError;
+	bool bErrorAboveTolerance = VolumeTolerance <= 0 || Decomposition[0].HullError > VolumeTolerance;
 	int32 WorstIdx = 0;
 	bool bCanSkipCurrent = bCanSkipUnreliableGeoVolumes && Decomposition[0].bGeometryVolumeUnreliable;
 	for (int32 PartIdx = 1; PartIdx < Decomposition.Num(); PartIdx++)
 	{
+		if (Decomposition[PartIdx].HullError > VolumeTolerance)
+		{
+			bErrorAboveTolerance = true;
+		}
 		if (Decomposition[PartIdx].HullError > WorstError)
 		{
 			if (bCanSkipUnreliableGeoVolumes && !bCanSkipCurrent && Decomposition[PartIdx].bGeometryVolumeUnreliable)
@@ -848,6 +855,15 @@ int32 FConvexDecomposition3::SplitWorst(bool bCanSkipUnreliableGeoVolumes)
 		}
 	}
 
+	// Stop splitting if we see no errors above tolerance
+	// Note if bCanSkipUnreliableGeoVolumes==true, we may still split a below-tolerance part
+	//  -- TODO: consider what to do in this case!
+	//		It will just add one extra split currently as we alternate toggling bCanSkipUnreliableGeoVolumes off, but the behavior may be confusing.
+	if (!bErrorAboveTolerance)
+	{
+		return 0;
+	}
+
 	FConvexPart& Part = Decomposition[WorstIdx];
 
 	double OrigHullVolume = Decomposition[WorstIdx].HullVolume;
@@ -1167,9 +1183,10 @@ void FConvexDecomposition3::UpdateProximitiesAfterSplit(int32 SplitIdx, int32 Ne
 	}
 }
 
-int32 FConvexDecomposition3::MergeBest(int32 NumMerges/*, double MaxErrorTolerance*/)
+int32 FConvexDecomposition3::MergeBest(int32 NumMerges, double MaxErrorTolerance)
 {
 	// Support having a max error tolerance
+	double VolumeTolerance = ConvertDistanceToleranceToLocalVolumeTolerance(MaxErrorTolerance);
 
 	int32 MergeNum = 0;
 
@@ -1213,7 +1230,7 @@ int32 FConvexDecomposition3::MergeBest(int32 NumMerges/*, double MaxErrorToleran
 		return NewPart;
 	};
 
-	for (; MergeNum < NumMerges; MergeNum++)
+	for (; VolumeTolerance > 0 || MergeNum < NumMerges; MergeNum++)
 	{
 		double BestKnownCost = FMathd::MaxReal;
 		int32 BestKnownIdx = -1;
@@ -1239,6 +1256,11 @@ int32 FConvexDecomposition3::MergeBest(int32 NumMerges/*, double MaxErrorToleran
 			}
 		}
 
+		if (VolumeTolerance > 0 && BestKnownCost > VolumeTolerance)
+		{
+			break;
+		}
+
 		// Perform the best merge and update proximity data
 
 		if (BestKnownIdx == -1)
diff --git a/Engine/Source/Runtime/GeometryCore/Public/CompGeom/ConvexDecomposition3.h b/Engine/Source/Runtime/GeometryCore/Public/CompGeom/ConvexDecomposition3.h
--- a/Engine/Source/Runtime/GeometryCore/Public/CompGeom/ConvexDecomposition3.h
+++ b/Engine/Source/Runtime/GeometryCore/Public/CompGeom/ConvexDecomposition3.h
@@ -71,27 +71,42 @@ public:
 	 * 
 	 * @param NumOutputHulls		Number of convex hulls to use in the final convex decomposition
 	 * @param NumAdditionalSplits	How far to go beyond the target number of hulls when initially the mesh into pieces -- larger values will require more computation but can find better convex decompositions
+	 * @param ErrorTolerance		Stop splitting when hulls have error less than this (expressed in cm; will be cubed for volumetric error)
 	 */
-	void Compute(int32 NumOutputHulls, int32 NumAdditionalSplits = 10)
+	void Compute(int32 NumOutputHulls, int32 NumAdditionalSplits = 10, double ErrorTolerance = 0.0)
 	{
-		for (int32 SplitIdx = 0; SplitIdx < NumOutputHulls + NumAdditionalSplits; SplitIdx++)
+		int32 TargetNumSplits = NumOutputHulls + NumAdditionalSplits;
+		for (int32 SplitIdx = 0; SplitIdx < TargetNumSplits; SplitIdx++)
 		{
-			SplitWorst(bool(SplitIdx % 2));
+			int32 NumNewParts = SplitWorst(bool(SplitIdx % 2), ErrorTolerance);
+			if (NumNewParts == 0)
+			{
+				break;
+			}
 		}
 		Compact();
-		MergeBest(NumHulls() - NumOutputHulls);
+		MergeBest(NumHulls() - NumOutputHulls, ErrorTolerance);
 	}
 
 	// Split the worst convex part, and return the increase in the total number of convex parts after splitting (can be more than 1 if result has multiple separate connected components)
 	// Note: could return 0 if no splits were possible
 	// @param bCanSkipUnreliableGeoVolumes		if true, don't split hulls where we have questionable geometry volume results, unless there is no hull with good geometry volume results
-	int32 SplitWorst(bool bCanSkipUnreliableGeoVolumes = false);
+	int32 SplitWorst(bool bCanSkipUnreliableGeoVolumes = false, double ErrorTolerance = 0.0);
+
+	// Merge the pairs of convex hulls in the decomposition that will least increase the error
+	// @param NumMerges				The target number of merges to perform.  Intermediate results can be used across merges, so it is best to do all merges in one call.
+	// @param ErrorTolerance		If > 0, continue to merge (if there are possible merges) until the resulting error would be greater than this value. Overrides NumMerges as the stopping condition.
+	// Note ErrorTolerance is expressed in cm, and will be cubed for volumetric error.
+	// @return						The number of merges performed
+	int32 MergeBest(int32 NumMerges, double ErrorTolerance = 0);
 
-	// Merge the pair of convexes in the decomposition that least-increases the error
-	// @param NumMerges			The target number of merges to perform.  Intermediate results can be used across merges, so it is best to do all merges in one call.
-	// @param ErrorTolerance	TODO: Refuse to merge if the resulting error would be greater than this value
-	// @return					The number of merges performed
-	int32 MergeBest(int32 NumMerges/*, double MaxErrorTolerance = FMathd::MaxReal*/);
+	// simple helper to convert an error tolerance expressed in world space to a local-space volume tolerance
+	inline double ConvertDistanceToleranceToLocalVolumeTolerance(double DistTolerance) const
+	{
+		double LocalDist = DistTolerance / ResultTransform.GetScale().GetMax();
+		double LocalVolume = LocalDist * LocalDist * LocalDist;
+		return LocalVolume;
+	}
 
 	// TODO: Implement a MergeSmall method to provide a way to ensure removal of low-quality parts
 	// Use merges remove convex hulls that are too small or too flat
@@ -140,6 +155,27 @@ public:
 		return Vertices;
 	}
 
+	FDynamicMesh3 GetHullMesh(int32 HullIdx) const
+	{
+		FDynamicMesh3 HullMesh;
+		for (int32 VID = 0; VID < Decomposition[HullIdx].InternalGeo.MaxVertexID(); VID++)
+		{
+			if (Decomposition[HullIdx].InternalGeo.IsVertex(VID))
+			{
+				HullMesh.AppendVertex(ResultTransform.TransformPosition(Decomposition[HullIdx].InternalGeo.GetVertex(VID)));
+			}
+			else
+			{
+				HullMesh.AppendVertex(FVector3d::ZeroVector);
+			}
+		}
+		for (const FIndex3i& Tri : Decomposition[HullIdx].HullTriangles)
+		{
+			HullMesh.AppendTriangle(Tri);
+		}
+		return HullMesh;
+	}
+
 	// Get the non-hull geometry underlying a part of the convex decomposition
 	const FDynamicMesh3& GetInternalMesh(int32 HullIdx) const
 	{
