improve convex decompositions by adding a concept of 'too thin' parts that the decomposition must remove in the merge phase

also fix handling of splitting into connected components when some of the components don't have a valid convex hull (by bailing on the component split in that case)

diff --git a/Engine/Source/Runtime/GeometryCore/Private/CompGeom/ConvexDecomposition3.cpp b/Engine/Source/Runtime/GeometryCore/Private/CompGeom/ConvexDecomposition3.cpp
--- a/Engine/Source/Runtime/GeometryCore/Private/CompGeom/ConvexDecomposition3.cpp
+++ b/Engine/Source/Runtime/GeometryCore/Private/CompGeom/ConvexDecomposition3.cpp
@@ -26,11 +26,11 @@ namespace Geometry
 // Note this will include vertices that are not on the hull
 struct FConvexPartHullMeshAdapter
 {
-	FConvexPartHullMeshAdapter(FConvexDecomposition3::FConvexPart* ConvexPart) : ConvexPart(ConvexPart)
+	FConvexPartHullMeshAdapter(const FConvexDecomposition3::FConvexPart* ConvexPart) : ConvexPart(ConvexPart)
 	{
 	}
 
-	FConvexDecomposition3::FConvexPart* ConvexPart;
+	const FConvexDecomposition3::FConvexPart* ConvexPart;
 
 	bool IsTriangle(int32 Index) const
 	{
@@ -203,6 +203,18 @@ struct FPartialCutResult
 			}
 		}
 
+		// Remove unreferenced vertices, which may be created by the above deletion (i.e., vertices on the plane that were only connected to now-deleted triangles)
+		for (int32 Side = 0; Side < 2; Side++)
+		{
+			for (int32 VID = 0, MaxVID = PartsMeshes[Side]->MaxVertexID(); VID < MaxVID; VID++)
+			{
+				if (PartsMeshes[Side]->IsVertex(VID) && !PartsMeshes[Side]->IsReferencedVertex(VID))
+				{
+					PartsMeshes[Side]->RemoveVertex(VID, false);
+				}
+			}
+		}
+
 		// Fill in the surface on the cutting plane; return false if hole fill process was ambiguous / unreliable
 		auto HoleFill = [this, &Plane, &PartsMeshes, &PlaneTol]() -> bool
 		{
@@ -397,45 +409,21 @@ struct FPartialCutResult
 			// Note this only looks at bounding box overlap currently; TODO: extend to test more possible plane of separation?
 			VertexComponents.ConnectOverlappingComponents(Mesh, 3);
 
-			// Note: Similar to VertexComponents.MakeComponentMap, but instead of mapping to a contiguous range starting at zero,
-			//		this maps to Decomposition array indices, i.e.: PartIdx, Decomposition.Num(), Decomposition.Num() + 1, ..., Decomposition.Num() + N
-			TMap<int32, int32> ComponentIdxMap;
-			int32 NewComponents = 0;
-			int32 NextComponentIdx = PartIdx;
-			for (int32 VID : Mesh.VertexIndicesItr())
-			{
-				if (!Mesh.IsReferencedVertex(VID))
-				{
-					continue;
-				}
-				int32 DisjointSetID = VertexComponents.GetComponent(VID);
-				int32* ComponentIdx = ComponentIdxMap.Find(DisjointSetID);
-				if (!ComponentIdx)
-				{
-					int32 IdxInDecomposition = NextComponentIdx;
-					if (ComponentIdxMap.IsEmpty())
-					{
-						NextComponentIdx = Decomposition.Num();
-					}
-					else
-					{
-						NextComponentIdx++;
-					}
-					ComponentIdxMap.Add(DisjointSetID, IdxInDecomposition);
-					NewComponents++;
-				}
-			}
-
+			// Get a mapping from component set IDs to contiguous indices from 0 to k-1 (for k components)
+			TMap<int32, int32> ComponentToIdxMap = VertexComponents.MakeComponentMap(Mesh, 3);
+			
+			int32 NewComponents = ComponentToIdxMap.Num();
 			if (NewComponents < 2)
 			{
 				return NewComponents;
 			}
 
-			FDynamicMesh3 OrigMesh = MoveTemp(Mesh); // Prepare to copy the vertices and triangles out of the old mesh and into the new
-			Decomposition[PartIdx].Reset();
-			for (int32 Idx = 1; Idx < NewComponents; Idx++)
+
+			const FDynamicMesh3& OrigMesh = Mesh; // Prepare to copy the vertices and triangles out of the old mesh and into the new parts
+			TArray<TUniquePtr<FConvexDecomposition3::FConvexPart>> NewParts;
+			for (int32 Idx = 0; Idx < NewComponents; Idx++)
 			{
-				Decomposition.Add(new FConvexDecomposition3::FConvexPart);
+				NewParts.Add(MakeUnique<FConvexDecomposition3::FConvexPart>());
 			}
 			
 			TMap<int32, int32> VertexMap; // Map vertices to their new VID in their new mesh (note: contains mappings into every component)
@@ -446,8 +434,8 @@ struct FPartialCutResult
 					continue;
 				}
 				int32 DisjointSetID = VertexComponents.GetComponent(VID);
-				int32 DecompIdx = ComponentIdxMap[DisjointSetID];
-				int32 NewVID = Decomposition[DecompIdx].InternalGeo.AppendVertex(OrigMesh.GetVertex(VID));
+				int32 NewPartIdx = ComponentToIdxMap[DisjointSetID];
+				int32 NewVID = NewParts[NewPartIdx]->InternalGeo.AppendVertex(OrigMesh.GetVertex(VID));
 				VertexMap.Add(VID, NewVID);
 			}
 			for (int32 TID : OrigMesh.TriangleIndicesItr())
@@ -455,15 +443,33 @@ struct FPartialCutResult
 				FIndex3i Tri = OrigMesh.GetTriangle(TID);
 				// TODO: if we use triangle groups, also get the group ID here and transfer it across
 				int32 DisjointSetID = VertexComponents.GetComponent(Tri.A);
-				int32 DecompIdx = ComponentIdxMap[DisjointSetID];
+				int32 NewPartIdx = ComponentToIdxMap[DisjointSetID];
 				FIndex3i NewTri(VertexMap[Tri.A], VertexMap[Tri.B], VertexMap[Tri.C]);
-				Decomposition[DecompIdx].InternalGeo.AppendTriangle(NewTri);
+				NewParts[NewPartIdx]->InternalGeo.AppendTriangle(NewTri);
 			}
-			for (const TPair<int32, int32>& ToDecompIdx : ComponentIdxMap)
+			bool bHasFailedComponents = false;
+			for (const TPair<int32, int32>& ToNewPartIdx : ComponentToIdxMap)
 			{
-				int32 DecompIdx = ToDecompIdx.Value;
-				Decomposition[DecompIdx].ComputeHull();
-				Decomposition[DecompIdx].ComputeStats();
+				int32 NewPartIdx = ToNewPartIdx.Value;
+				bHasFailedComponents = !NewParts[NewPartIdx]->ComputeHull() || bHasFailedComponents;
+				if (bHasFailedComponents)
+				{
+					break;
+				}
+				NewParts[NewPartIdx]->ComputeStats();
+			}
+			if (bHasFailedComponents)
+			{
+				// failed to build hull for some components; fall back to not splitting the components in this case
+				return 1;
+			}
+
+			// Successfully built the connected components: transfer them into the decomposition
+			Decomposition[PartIdx] = MoveTemp(*NewParts[0]);
+			NewParts[0].Reset();
+			for (int32 Idx = 1; Idx < NewParts.Num(); Idx++)
+			{
+				Decomposition.Add(NewParts[Idx].Release());
 			}
 
 			return NewComponents;
@@ -777,7 +783,7 @@ FConvexDecomposition3::FConvexPart::FConvexPart(const FDynamicMesh3& SourceMesh,
 	ComputeStats();
 }
 
-void FConvexDecomposition3::FConvexPart::ComputeHull(bool bComputePlanes)
+bool FConvexDecomposition3::FConvexPart::ComputeHull(bool bComputePlanes)
 {
 	FConvexHull3d HullCompute;
 	bool bOK = HullCompute.Solve(InternalGeo.MaxVertexID(),
@@ -787,13 +793,14 @@ void FConvexDecomposition3::FConvexPart::ComputeHull(bool bComputePlanes)
 	{
 		int32 Dimension = HullCompute.GetDimension();
 		bFailed = true;
-		return;
+		return false;
 	}
 	HullTriangles = HullCompute.MoveTriangles();
 	if (bComputePlanes)
 	{
 		ComputeHullPlanes();
 	}
+	return true;
 }
 
 void FConvexDecomposition3::FConvexPart::ComputeHullPlanes()
@@ -887,7 +894,6 @@ int32 FConvexDecomposition3::SplitWorst(bool bCanSkipUnreliableGeoVolumes, doubl
 	}
 
 	// make candidate planes out of convex and boundary edges
-	const int32 MaxConvexEdgePlanes = 20; // TODO: Expose this as a parameter
 
 	// Structure of arrays for planes that come from convex (or open) edges
 	constexpr int32 MaxPerEdge = 5;
@@ -1183,15 +1189,144 @@ void FConvexDecomposition3::UpdateProximitiesAfterSplit(int32 SplitIdx, int32 Ne
 	}
 }
 
-int32 FConvexDecomposition3::MergeBest(int32 NumMerges, double MaxErrorTolerance)
+int32 FConvexDecomposition3::MergeBest(int32 TargetNumParts, double MaxErrorTolerance, double MinThicknessToleranceWorldSpace, bool bAllowCompact)
 {
 	// Support having a max error tolerance
 	double VolumeTolerance = ConvertDistanceToleranceToLocalVolumeTolerance(MaxErrorTolerance);
+	double MinThicknessTolerance = ConvertDistanceToleranceToLocalSpace(MinThicknessToleranceWorldSpace);
 
 	int32 MergeNum = 0;
 
 	TMap<int32, TUniquePtr<FConvexPart>> ProximityComputedParts;
 
+	TFunctionRef<bool(const FConvexPart&)> IsPartBelowSizeTolerance = MinThicknessTolerance > 0 ? 
+	(TFunctionRef<bool(const FConvexPart&)>)[MinThicknessTolerance](const FConvexPart& Part) -> bool
+	{ // Check the part vs the thickness tolerance
+		// First check the already-computed AABB
+		if (Part.Bounds.MinDim() < MinThicknessTolerance)
+		{
+			return true;
+		}
+		// Now check if it's a thin part that's just not aligned to a major axis
+		
+		// Use the extreme points on the AABB to judge whether it's a flat part
+		// Note: TExtremePoints3 is likely ok for detecting very-thin parts, but could be off by ~2x easily. To more precisely measure thickness,
+		// consider switching to CompGeom/DiTOrientedBox.h (could be important for a relatively large size tolerance)
+		TExtremePoints3<double> ExtremePoints(Part.InternalGeo.MaxVertexID(), [&Part](int32 VID) {return Part.InternalGeo.GetVertex(VID);}, [&Part](int32 VID) {return Part.InternalGeo.IsVertex(VID);});
+		if (UNLIKELY(ExtremePoints.Dimension < 3)) // Points were all coplanar (Shouldn't happen here because we wouldn't have a valid hull at all in this case)
+		{
+			return true;
+		}
+		// By convention, the last basis vector is in the direction most likely to be 'thin'
+		FVector3d SmallDir = ExtremePoints.Basis[2];
+		double TetHeight = FMathd::Max(
+			FMathd::Abs(SmallDir.Dot(Part.InternalGeo.GetVertex(ExtremePoints.Extreme[0]) - Part.InternalGeo.GetVertex(ExtremePoints.Extreme[3]))),
+			FMathd::Abs(SmallDir.Dot(Part.InternalGeo.GetVertex(ExtremePoints.Extreme[0]) - Part.InternalGeo.GetVertex(ExtremePoints.Extreme[2]))));
+		if (TetHeight > MinThicknessTolerance)
+		{
+			return false;
+		}
+		if (TetHeight * 2 < MinThicknessTolerance)
+		{
+			return true;
+		}
+		FInterval1d SmallAxisInterval;
+		for (FVector3d V : Part.InternalGeo.VerticesItr())
+		{
+			SmallAxisInterval.Contain(V.Dot(SmallDir));
+		}
+		return SmallAxisInterval.Extent() < MinThicknessTolerance;
+	} : // No tolerance set: Always return false
+	(TFunctionRef<bool(const FConvexPart&)>)[](const FConvexPart& Part) -> bool
+	{
+		return false;
+	};
+
+	int32 MustMergeCount = 0;
+	if (MinThicknessTolerance > 0)
+	{
+		for (int32 PartIdx = 0; PartIdx < Decomposition.Num(); PartIdx++)
+		{
+			FConvexPart& Part = Decomposition[PartIdx];
+			Part.bMustMerge = IsPartBelowSizeTolerance(Part);
+			MustMergeCount += Part.bMustMerge;
+		}
+	}
+	// TODO: if MustMergeCount is >= NumParts-1, we could shortcut the process and just return a single convex hull
+
+	// Try to further-split long/wide thin parts -- parts that are not thin in all directions, e.g. parts with max bounding box extent wider than MinThicknessTolerance
+	if (MustMergeCount > 0)
+	{
+		TArray<int32> PartsToConsider;
+		for (int32 OrigPartIdx = 0, OrigNumDecomposition = Decomposition.Num(); OrigPartIdx < OrigNumDecomposition; OrigPartIdx++)
+		{
+			FConvexPart& OrigPart = Decomposition[OrigPartIdx];
+			if (OrigPart.bMustMerge && OrigPart.Bounds.MaxDim() > MinThicknessTolerance)
+			{
+				TArray<FPlane3d> CutPlanes;
+				// Collect neighboring planes that could split this part
+				for (auto It = DecompositionToProximity.CreateConstKeyIterator(OrigPartIdx); It; ++It)
+				{
+					int32 OtherPartIdx = Proximities[It.Value()].Link.OtherElement(OrigPartIdx);
+					for (auto OtherIt = DecompositionToProximity.CreateConstKeyIterator(OtherPartIdx); OtherIt; ++OtherIt)
+					{
+						const FProximity& Prox = Proximities[OtherIt.Value()];
+						if (Prox.bPlaneSeparates && !Prox.Link.Contains(OrigPartIdx))
+						{
+							CutPlanes.Add(Prox.Plane);
+						}
+					}
+				}
+
+				// Attempt splits
+				PartsToConsider.Reset();
+				PartsToConsider.Add(OrigPartIdx);
+				for (const FPlane3d& Plane : CutPlanes)
+				{
+					for (int32 Idx = 0, Num = PartsToConsider.Num(); Idx < Num; Idx++)
+					{
+						int32 PartIdx = PartsToConsider[Idx];
+						FConvexPart& PartToSplit = Decomposition[PartIdx];
+						if (PartToSplit.Bounds.MaxDim() < MinThicknessTolerance)
+						{
+							continue;
+						}
+						FPartialCutResult PlaneResult(PartToSplit, Plane, OnPlaneTolerance);
+						if (!PlaneResult.bSuccess)
+						{
+							continue;
+						}
+
+						double OrigHullVolume = PartToSplit.HullVolume;
+						int32 NewPartsStartIdx = Decomposition.Num();
+						int32 OtherSideStartIdx = -1;
+						PlaneResult.ApplyToGeo(Decomposition, PartIdx, Plane, OtherSideStartIdx, OnPlaneTolerance, ConnectedComponentTolerance);
+						UpdateProximitiesAfterSplit(PartIdx, NewPartsStartIdx, Plane, OtherSideStartIdx, OrigHullVolume);
+						PartToSplit.bMustMerge = true;
+						for (int32 NewIdx = NewPartsStartIdx; NewIdx < Decomposition.Num(); NewIdx++)
+						{
+							PartsToConsider.Add(NewIdx);
+							Decomposition[NewIdx].bMustMerge = true;
+						}
+					}
+				}
+			}
+		}
+
+		// recompute MustMergeCount to account for possible splits
+		MustMergeCount = 0;
+		for (int32 PartIdx = 0; PartIdx < Decomposition.Num(); PartIdx++)
+		{
+			FConvexPart& Part = Decomposition[PartIdx];
+			MustMergeCount += Part.bMustMerge;
+		}
+	}
+
+	if (bAllowCompact)
+	{
+		Compact();
+	}
+
 	// initialize all parts with no merged hulls volume as having just their own hull volume (assuming no merges were yet performed on these parts)
 	for (FConvexPart& Part : Decomposition)
 	{
@@ -1202,7 +1337,7 @@ int32 FConvexDecomposition3::MergeBest(int32 NumMerges, double MaxErrorTolerance
 	}
 
 
-	auto CreateMergedPart = [this](FProximity& Prox) -> TUniquePtr<FConvexPart>
+	auto CreateMergedPart = [this, &IsPartBelowSizeTolerance](FProximity& Prox) -> TUniquePtr<FConvexPart>
 	{
 		TUniquePtr<FConvexPart> NewPart = MakeUnique<FConvexPart>();
 		NewPart->GeoCenter = FVector3d::ZeroVector;
@@ -1225,17 +1360,29 @@ int32 FConvexDecomposition3::MergeBest(int32 NumMerges, double MaxErrorTolerance
 		NewPart->HullVolume = GetVolumeUsingReferencePoint<FConvexPartHullMeshAdapter>(HullAdapter, NewPart->GeoCenter);
 		NewPart->Bounds = NewPart->InternalGeo.GetBounds();
 		NewPart->HullError = NewPart->HullVolume - NewPart->GeoVolume;
+		NewPart->bMustMerge = // only re-evaluate part thickness if both source parts were below tolerance
+			Decomposition[Prox.Link[0]].bMustMerge &&
+			Decomposition[Prox.Link[1]].bMustMerge &&
+			IsPartBelowSizeTolerance(*NewPart);
 		NewPart->Compact();
 
 		return NewPart;
 	};
 
-	for (; VolumeTolerance > 0 || MergeNum < NumMerges; MergeNum++)
+	for (; MustMergeCount > 0 || VolumeTolerance > 0 || Decomposition.Num() > TargetNumParts; MergeNum++)
 	{
 		double BestKnownCost = FMathd::MaxReal;
+		bool bOnlyAllowMustMerges = VolumeTolerance == 0 && Decomposition.Num() <= TargetNumParts;
 		int32 BestKnownIdx = -1;
 		for (int32 ProxIdx = 0; ProxIdx < Proximities.Num(); ProxIdx++)
 		{
+			bool bIncludesMustMergePart =
+				Decomposition[Proximities[ProxIdx].Link.A].bMustMerge ||
+				Decomposition[Proximities[ProxIdx].Link.B].bMustMerge;
+			bool bOnlyMustMergeParts =
+				Decomposition[Proximities[ProxIdx].Link.A].bMustMerge &&
+				Decomposition[Proximities[ProxIdx].Link.B].bMustMerge;
+
 			double MergeCost;
 			if (Proximities[ProxIdx].HasMergedVolume())
 			{
@@ -1249,6 +1396,21 @@ int32 FConvexDecomposition3::MergeBest(int32 NumMerges, double MaxErrorTolerance
 				MergeCost = Proximities[ProxIdx].GetMergeCost(Decomposition);
 			}
 
+			if (!bIncludesMustMergePart && ( // if there's no must-merge part, skip if ...
+					(VolumeTolerance > 0 && MergeCost > VolumeTolerance) // it won't pass the volume tolerance
+						|| 
+					bOnlyAllowMustMerges // or we're currently only considering must-merge parts (e.g., already at or below desired number of parts)
+				))
+			{
+				continue;
+			}
+
+			// Apply a bias term to favor getting rid of 'must merge' parts first
+			if (bIncludesMustMergePart && !bOnlyMustMergeParts)
+			{
+				MergeCost -= BiasToRemoveTooThinParts;
+			}
+
 			if (MergeCost < BestKnownCost)
 			{
 				BestKnownCost = MergeCost;
@@ -1256,10 +1418,6 @@ int32 FConvexDecomposition3::MergeBest(int32 NumMerges, double MaxErrorTolerance
 			}
 		}
 
-		if (VolumeTolerance > 0 && BestKnownCost > VolumeTolerance)
-		{
-			break;
-		}
 
 		// Perform the best merge and update proximity data
 
@@ -1276,6 +1434,8 @@ int32 FConvexDecomposition3::MergeBest(int32 NumMerges, double MaxErrorTolerance
 			Swap(DecoToKeep, DecoToRm); // keep the smaller index (this makes the book-keeping update below slightly easier)
 		}
 
+		MustMergeCount -= int32(Decomposition[DecoToKeep].bMustMerge) + int32(Decomposition[DecoToRm].bMustMerge);
+
 		// Update the "to keep" index with the new, merged FConvexPart
 		if (ProximityComputedParts.Contains(BestKnownIdx))
 		{
@@ -1288,6 +1448,8 @@ int32 FConvexDecomposition3::MergeBest(int32 NumMerges, double MaxErrorTolerance
 			Decomposition[DecoToKeep] = MoveTemp(*NewPart);
 		}
 
+		MustMergeCount += int32(Decomposition[DecoToKeep].bMustMerge);
+
 		// Do all the book-keeping to delete the other old FConvexPart and update all Proximity data
 
 		// Clear all pre-computed merge info related to the newly-merged parts, saving out only what parts the merge should be linked to
diff --git a/Engine/Source/Runtime/GeometryCore/Public/CompGeom/ConvexDecomposition3.h b/Engine/Source/Runtime/GeometryCore/Public/CompGeom/ConvexDecomposition3.h
--- a/Engine/Source/Runtime/GeometryCore/Public/CompGeom/ConvexDecomposition3.h
+++ b/Engine/Source/Runtime/GeometryCore/Public/CompGeom/ConvexDecomposition3.h
@@ -63,6 +63,14 @@ public:
 	// Angle at edge above which we don't sample it as a convex edge (expressed in degrees)
 	double ConvexEdgeAngleThreshold = 170;
 
+	// A bias term to favor merging away 'too thin' parts before merging other parts, expressed as a volume in units s.t. the longest axis of the original input bounding box is 1 unit long
+	// Because these merges *must* be done, it's generally safer to do them earlier, rather than hoping the error associated with them will go down after other merges are performed
+	double BiasToRemoveTooThinParts = .1;
+
+	// Maximum number of convex edges to sample for possible cutting planes when calling SplitWorst() to generate an initial convex decomposition
+	// Larger values will cost more to run but can let the algorithm find a cleaner decomposition
+	int32 MaxConvexEdgePlanes = 50;
+
 	// TODO: Provide hull approximation options?
 
 
@@ -71,9 +79,10 @@ public:
 	 * 
 	 * @param NumOutputHulls		Number of convex hulls to use in the final convex decomposition
 	 * @param NumAdditionalSplits	How far to go beyond the target number of hulls when initially the mesh into pieces -- larger values will require more computation but can find better convex decompositions
-	 * @param ErrorTolerance		Stop splitting when hulls have error less than this (expressed in cm; will be cubed for volumetric error)
+	 * @param ErrorTolerance		Stop splitting when hulls have error less than this (expressed in cm; will be cubed for volumetric error). Overrides NumOutputHulls if specified
+	 * @param MinThicknessTolerance	Optionally specify a minimum thickness (in cm) for convex parts; parts below this thickness will always be merged away. Overrides NumOutputHulls and ErrorTolerance when needed
 	 */
-	void Compute(int32 NumOutputHulls, int32 NumAdditionalSplits = 10, double ErrorTolerance = 0.0)
+	void Compute(int32 NumOutputHulls, int32 NumAdditionalSplits = 10, double ErrorTolerance = 0.0, double MinThicknessTolerance = 0)
 	{
 		int32 TargetNumSplits = NumOutputHulls + NumAdditionalSplits;
 		for (int32 SplitIdx = 0; SplitIdx < TargetNumSplits; SplitIdx++)
@@ -84,8 +93,9 @@ public:
 				break;
 			}
 		}
-		Compact();
-		MergeBest(NumHulls() - NumOutputHulls, ErrorTolerance);
+		
+		constexpr bool bAllowCompact = true;
+		MergeBest(NumOutputHulls, ErrorTolerance, MinThicknessTolerance, bAllowCompact);
 	}
 
 	// Split the worst convex part, and return the increase in the total number of convex parts after splitting (can be more than 1 if result has multiple separate connected components)
@@ -93,30 +103,39 @@ public:
 	// @param bCanSkipUnreliableGeoVolumes		if true, don't split hulls where we have questionable geometry volume results, unless there is no hull with good geometry volume results
 	int32 SplitWorst(bool bCanSkipUnreliableGeoVolumes = false, double ErrorTolerance = 0.0);
 
-	// Merge the pairs of convex hulls in the decomposition that will least increase the error
-	// @param NumMerges				The target number of merges to perform.  Intermediate results can be used across merges, so it is best to do all merges in one call.
-	// @param ErrorTolerance		If > 0, continue to merge (if there are possible merges) until the resulting error would be greater than this value. Overrides NumMerges as the stopping condition.
-	// Note ErrorTolerance is expressed in cm, and will be cubed for volumetric error.
+	// Merge the pairs of convex hulls in the decomposition that will least increase the error.  Intermediate results can be used across merges, so it is best to do all merges in one call.
+	// @param TargetNumParts		The target number of parts for the decomposition; will be overriden by non-default ErrorTolerance or MinPartThickness
+	// @param ErrorTolerance		If > 0, continue to merge (if there are possible merges) until the resulting error would be greater than this value. Overrides TargetNumParts as the stopping condition.
+	//								Note: ErrorTolerance is expressed in cm, and will be cubed for volumetric error.
+	// @param MinPartThickness		Optionally specify a minimum thickness (in cm) for convex parts; parts below this thickness will always be merged away. Overrides TargetNumParts and ErrorTolerance when needed.
+	//								Note: These parts may be further split so they can be merged into multiple hulls
+	// @param bAllowCompact			Allow the algorithm to discard underlying geometry once it will no longer be used, resulting in a smaller representation & faster merges
 	// @return						The number of merges performed
-	int32 MergeBest(int32 NumMerges, double ErrorTolerance = 0);
+	int32 MergeBest(int32 TargetNumParts, double ErrorTolerance = 0, double MinThicknessTolerance = 0, bool bAllowCompact = true);
 
 	// simple helper to convert an error tolerance expressed in world space to a local-space volume tolerance
 	inline double ConvertDistanceToleranceToLocalVolumeTolerance(double DistTolerance) const
 	{
-		double LocalDist = DistTolerance / ResultTransform.GetScale().GetMax();
+		double LocalDist = ConvertDistanceToleranceToLocalSpace(DistTolerance);
 		double LocalVolume = LocalDist * LocalDist * LocalDist;
 		return LocalVolume;
 	}
 
-	// TODO: Implement a MergeSmall method to provide a way to ensure removal of low-quality parts
-	// Use merges remove convex hulls that are too small or too flat
-	// @param MinVolumeTolerance		Hulls with a smaller volume than this will be removed by merges (units in cm^3)
-	// @param MinThicknessTolerance		Hulls thinner than this will be removed by merges (units in cm)
-	// Note hull thickness is measured by first fitting a plane to the hull points then measuring their span in the plane normal direction
-	//int32 MergeSmall(double MinVolumeTolerance, double MinThicknessTolerance = FMathd::MaxReal);
+	// simple helper to convert an error tolerance expressed in world space to local space
+	inline double ConvertDistanceToleranceToLocalSpace(double DistTolerance) const
+	{
+		double ScaleFactor = ResultTransform.GetScale().GetMax();
+		if (!ensure(ScaleFactor >= FMathd::Epsilon)) // If ScaleFactor is zero, ResultTransform is trying to collapse all the hulls to a degenerate space -- likely indicates an incorrect ResultTransform
+		{
+			ScaleFactor = FMathd::Epsilon;
+		}
+		double LocalDist = DistTolerance / ScaleFactor;
+		return LocalDist;
+	}
 
 	// Compact the decomposition representation to the minimal needed for output: Just the hull vertices and triangles, no internal geometry
-	// Only call this if no more decomposition will be performed, but merges can still be performed on the compact representation
+	// This will prevent some further processing on the representation; to ensure it is called only when the geometry is no longer needed,
+	// instead of calling this directly, call MergeBest with bAllowCompact=true
 	void Compact()
 	{
 		for (FConvexPart& Part : Decomposition)
@@ -195,12 +214,14 @@ public:
 			HullPlanes.Reset();
 			HullVolume = 0;
 			GeoVolume = 0;
+			SumHullsVolume = -FMathd::MaxReal;
 			GeoCenter = FVector3d::ZeroVector;
 			Bounds = FAxisAlignedBox3d::Empty();
 			HullError = FMathd::MaxReal;
 			bIsCompact = false;
 			bFailed = false;
 			bGeometryVolumeUnreliable = false;
+			bMustMerge = false;
 		}
 
 		// Underlying geometry represented by the convex part
@@ -227,6 +248,9 @@ public:
 		FAxisAlignedBox3d Bounds;
 		bool bGeometryVolumeUnreliable = false;
 
+		// Flag indicating the hull should be merged into another part during a MergeBest() call
+		bool bMustMerge = false;
+
 		// Approximate value indicating how badly the hull mismatches the input geometry
 		double HullError = FMathd::MaxReal;
 
@@ -252,7 +276,7 @@ public:
 		}
 
 		// Helper to create hull after InternalGeo is set
-		void ComputeHull(bool bComputePlanes = true);
+		bool ComputeHull(bool bComputePlanes = true);
 
 		// Helper to compute volumes, centroid, bounds, etc after the convex part is initialized
 		void ComputeStats();
