Add some missing #includes

#rnx
#trivial

diff --git a/Engine/Source/Runtime/GeometryCore/Public/Intersection/IntrRay3Triangle3.h b/Engine/Source/Runtime/GeometryCore/Public/Intersection/IntrRay3Triangle3.h
--- a/Engine/Source/Runtime/GeometryCore/Public/Intersection/IntrRay3Triangle3.h
+++ b/Engine/Source/Runtime/GeometryCore/Public/Intersection/IntrRay3Triangle3.h
@@ -7,6 +7,7 @@
 #include "VectorTypes.h"
 #include "TriangleTypes.h"
 #include "VectorUtil.h"
+#include "Math/Ray.h"
 
 namespace UE
 {
diff --git a/Engine/Source/Runtime/GeometryCore/Public/PlaneTypes.h b/Engine/Source/Runtime/GeometryCore/Public/PlaneTypes.h
--- a/Engine/Source/Runtime/GeometryCore/Public/PlaneTypes.h
+++ b/Engine/Source/Runtime/GeometryCore/Public/PlaneTypes.h
@@ -6,6 +6,7 @@
 
 #include "VectorTypes.h"
 #include "VectorUtil.h"
+#include "Math/Plane.h"
 
 namespace UE
 {
