MeshWindingNumberGrid.h: Don't proceed through Compute() if the bounding volume is empty.

Avoids overflow errors from msvc for certain target configurations.

#rnx

diff --git a/Engine/Source/Runtime/GeometryCore/Public/Spatial/MeshWindingNumberGrid.h b/Engine/Source/Runtime/GeometryCore/Public/Spatial/MeshWindingNumberGrid.h
--- a/Engine/Source/Runtime/GeometryCore/Public/Spatial/MeshWindingNumberGrid.h
+++ b/Engine/Source/Runtime/GeometryCore/Public/Spatial/MeshWindingNumberGrid.h
@@ -82,6 +82,11 @@ public:
 		// figure out origin & dimensions
 		FAxisAlignedBox3d bounds = FastWinding->GetTree()->GetBoundingBox();
 
+		if (bounds.IsEmpty())
+		{
+			return;
+		}
+
 		float fBufferWidth = BufferCells * (float)CellSize;
 		GridOrigin = (FVector3f)bounds.Min - fBufferWidth * FVector3f::One();
 		FVector3f max = (FVector3f)bounds.Max + fBufferWidth * FVector3f::One();
