Add IsInsideOrOn method to the circle type

diff --git a/Engine/Source/Runtime/GeometryCore/Public/CircleTypes.h b/Engine/Source/Runtime/GeometryCore/Public/CircleTypes.h
--- a/Engine/Source/Runtime/GeometryCore/Public/CircleTypes.h
+++ b/Engine/Source/Runtime/GeometryCore/Public/CircleTypes.h
@@ -101,6 +101,10 @@ struct TCircle2
 		return Center.DistanceSquared(Point) < Radius*Radius;
 	}
 
+	bool IsInsideOrOn(const TVector2<RealType>& Point) const
+	{
+		return TVector2<RealType>::DistSquared(Center, Point) <= Radius*Radius;
+	}
 
 	RealType SignedDistance(const TVector2<RealType>& Point) const
 	{
