Make a new triangulation fill mode that supports explicit boundary and hole edges

Update Arrangement2d's triangulation functions to use that fill mode, and also to be explicit about whether the boundary group is expected or not

diff --git a/Engine/Plugins/Runtime/GeometryProcessing/Source/GeometryAlgorithms/Private/Arrangement2d.cpp b/Engine/Plugins/Runtime/GeometryProcessing/Source/GeometryAlgorithms/Private/Arrangement2d.cpp
--- a/Engine/Plugins/Runtime/GeometryProcessing/Source/GeometryAlgorithms/Private/Arrangement2d.cpp
+++ b/Engine/Plugins/Runtime/GeometryProcessing/Source/GeometryAlgorithms/Private/Arrangement2d.cpp
@@ -88,7 +88,7 @@ bool FArrangement2d::AttemptTriangulate(TArray<FIntVector>& Triangles, TArray<in
 	// TODO: Currently this just constructs an FIndex3i array and then copies it to an FIntVector array; if we need to construct FIntVector versions faster, we could directly construct them here instead
 	// However this would require writing FIntVector versions of the GetTriangles functions in FDelaunay2; ideally we would instead mainly call the FIndex3i variant
 	TArray<FIndex3i> TrianglesInd3i;
-	bool bResult = AttemptTriangulate(TrianglesInd3i, SkippedEdges, BoundaryEdgeGroupID);
+	bool bResult = TriangulateInternal(TrianglesInd3i, true, BoundaryEdgeGroupID, false, -1, true, &SkippedEdges);
 	Triangles.SetNumUninitialized(TrianglesInd3i.Num());
 	for (int32 Idx = 0; Idx < TrianglesInd3i.Num(); Idx++)
 	{
@@ -97,7 +97,7 @@ bool FArrangement2d::AttemptTriangulate(TArray<FIntVector>& Triangles, TArray<in
 	return bResult;
 }
 
-bool FArrangement2d::AttemptTriangulate(TArray<FIndex3i> &Triangles, TArray<int32> &SkippedEdges, int32 BoundaryEdgeGroupID)
+bool FArrangement2d::TriangulateInternal(TArray<FIndex3i>& Triangles, bool bHasBoundaryGroupID, int32 BoundaryEdgeGroupID, bool bHasHoleGroupID, int32 HoleEdgeGroupID, bool bLegacyKeepTrianglesIfBoundaryNotFound, TArray<int32>* SkippedEdges)
 {
 	Triangles.Empty();
 
@@ -151,7 +151,7 @@ bool FArrangement2d::AttemptTriangulate(TArray<FIndex3i> &Triangles, TArray<int3
 	bool bInsertConstraintFailure = false;
 	bool bBoundaryTrackingFailure = false;
 
-	TArray<FIndex2i> AllEdges, BoundaryEdges;
+	TArray<FIndex2i> AllEdges, BoundaryEdges, HoleEdges;
 
 	for (int EdgeIdx : Graph.EdgeIndices())
 	{
@@ -162,10 +162,14 @@ bool FArrangement2d::AttemptTriangulate(TArray<FIndex3i> &Triangles, TArray<int3
 			Edge.B = InputIndices[Edge.B];
 		}
 		FIndex2i& AddedEdge = AllEdges.Emplace_GetRef(Edge.A, Edge.B);
-		if (Edge.Group == BoundaryEdgeGroupID)
+		if (bHasBoundaryGroupID && Edge.Group == BoundaryEdgeGroupID)
 		{
 			BoundaryEdges.Add(AddedEdge);
 		}
+		if (bHasHoleGroupID && Edge.Group == HoleEdgeGroupID)
+		{
+			HoleEdges.Add(AddedEdge);
+		}
 	}
 
 	Delaunay.ConstrainEdges(InputVertices, AllEdges);
@@ -182,18 +186,27 @@ bool FArrangement2d::AttemptTriangulate(TArray<FIndex3i> &Triangles, TArray<int3
 		if (!Delaunay.HasEdge(FIndex2i(Edge.A, Edge.B), false))
 		{
 			bInsertConstraintFailure = true;
-			SkippedEdges.Add(EdgeIdx);
+			if (SkippedEdges)
+			{
+				SkippedEdges->Add(EdgeIdx);
+			}
 			if (Edge.Group == BoundaryEdgeGroupID)
 			{
 				bBoundaryTrackingFailure = true;
 			}
 		}
 	}
-	
-	if (!bBoundaryTrackingFailure && BoundaryEdges.Num() > 0)
+
+	if (bHasBoundaryGroupID && (!bLegacyKeepTrianglesIfBoundaryNotFound || (!bBoundaryTrackingFailure && BoundaryEdges.Num() > 0)))
 	{
-		Triangles = Delaunay.GetFilledTriangles(BoundaryEdges, FDelaunay2::EFillMode::Solid);
-		ensure(Triangles.Num()); // technically it could be valid to not have any triangles after eating the outside-boundary ones, but it would be unusual and could also be a bug
+		if (bHasHoleGroupID)
+		{
+			Delaunay.GetFilledTriangles(Triangles, BoundaryEdges, HoleEdges);
+		}
+		else
+		{
+			Triangles = Delaunay.GetFilledTriangles(BoundaryEdges, FDelaunay2::EFillMode::Solid);
+		}
 	}
 	else
 	{
@@ -212,3 +225,22 @@ bool FArrangement2d::AttemptTriangulate(TArray<FIndex3i> &Triangles, TArray<int3
 	return !bInsertConstraintFailure;
 }
 
+bool FArrangement2d::AttemptTriangulate(TArray<FIndex3i>& Triangles, TArray<int32>& SkippedEdges, int32 BoundaryEdgeGroupID)
+{
+	return TriangulateInternal(Triangles, true, BoundaryEdgeGroupID, false, -1, true, &SkippedEdges);
+}
+
+bool FArrangement2d::TriangulateWithBoundaryAndHoles(TArray<FIndex3i>& Triangles, int32 BoundaryEdgeGroupID, int32 HoleEdgeGroupID)
+{
+	return TriangulateInternal(Triangles, true, BoundaryEdgeGroupID, true, HoleEdgeGroupID, false, nullptr);
+}
+
+bool FArrangement2d::TriangulateWithBoundary(TArray<FIndex3i>& Triangles, int32 BoundaryEdgeGroupID)
+{
+	return TriangulateInternal(Triangles, true, BoundaryEdgeGroupID, false, -1, false, nullptr);
+}
+
+bool FArrangement2d::Triangulate(TArray<FIndex3i>& Triangles)
+{
+	return TriangulateInternal(Triangles, false, -1, false, -1, false, nullptr);
+}
diff --git a/Engine/Plugins/Runtime/GeometryProcessing/Source/GeometryAlgorithms/Public/Arrangement2d.h b/Engine/Plugins/Runtime/GeometryProcessing/Source/GeometryAlgorithms/Public/Arrangement2d.h
--- a/Engine/Plugins/Runtime/GeometryProcessing/Source/GeometryAlgorithms/Public/Arrangement2d.h
+++ b/Engine/Plugins/Runtime/GeometryProcessing/Source/GeometryAlgorithms/Public/Arrangement2d.h
@@ -54,6 +54,7 @@ struct FArrangement2d
 
 	/**
 	 * Attempts to triangulates the arrangement with a constrained Delaunay triangulation
+	 * NOTE: Will return all triangles if no edges found with the BoundaryEdgeGroupID
 	 * NOTE: May fail if arrangement has self-intersections
 	 * 
 	 * Triangles: Output triangles (as indices into Graph vertices)
@@ -61,11 +62,43 @@ struct FArrangement2d
 	 * BoundaryEdgeGroupID: ID of edges corresponding to a boundary; if we have a closed loop of these boundary edges on output triangulation, will discard triangles outside this
 	 * return: false if triangulation algo knows it failed (note Triangles may still have some triangulation of the input in this case; for example it may just be missing some required edges)
 	 */
+	UE_DEPRECATED(5.1, "Please use the Triangulate or TriangulateWithBoundary functions instead, which are explicit about whether a BoundaryEdgeGroupID should be present")
 	bool GEOMETRYALGORITHMS_API AttemptTriangulate(TArray<FIndex3i>& Triangles, TArray<int32>& SkippedEdges, int32 BoundaryEdgeGroupID);
 
 	// Variant of AttemptTriangulate using FIntVector instead of FIndex3i; Note this incurs an extra copy of the triangle array
+	UE_DEPRECATED(5.1, "Please use the Triangulate or TriangulateWithBoundary functions instead, which are explicit about whether a BoundaryEdgeGroupID should be present")
 	bool GEOMETRYALGORITHMS_API AttemptTriangulate(TArray<FIntVector>& Triangles, TArray<int32>& SkippedEdges, int32 BoundaryEdgeGroupID);
 
+	/**
+	 * Attempts to triangulate the arrangement with a constrained Delaunay triangulation
+	 * NOTE: May fail if arrangement has self-intersections
+	 * 
+	 * @param Triangles				Output triangles (as indices into Graph vertices)
+	 * @param BoundaryEdgeGroupID	ID of edges corresponding to a boundary: triangles outside these edges will be removed
+	 * @param HoleEdgeGroupID		ID of edges corresponding to internal holes: triangles inside these edges will be removed
+	 * @return						false if triangulation algo knows it failed; will still likely have some triangulation even in this case
+	 */
+	bool GEOMETRYALGORITHMS_API TriangulateWithBoundaryAndHoles(TArray<FIndex3i>& Triangles, int32 BoundaryEdgeGroupID, int32 HoleEdgeGroupID);
+
+	/**
+	 * Attempts to triangulate the arrangement with a constrained Delaunay triangulation
+	 * NOTE: May fail if arrangement has self-intersections
+	 *
+	 * @param Triangles				Output triangles (as indices into Graph vertices)
+	 * @param BoundaryEdgeGroupID	ID of edges corresponding to a boundary: triangles outside these edges will be removed
+	 * @return						false if triangulation algo knows it failed; will still likely have some triangulation even in this case
+	 */
+	bool GEOMETRYALGORITHMS_API TriangulateWithBoundary(TArray<FIndex3i>& Triangles, int32 BoundaryEdgeGroupID);
+
+	/**
+	 * Attempts to triangulate the arrangement with a constrained Delaunay triangulation
+	 * NOTE: May fail if arrangement has self-intersections
+	 *
+	 * @param Triangles				Output triangles (as indices into Graph vertices)
+	 * @return						false if triangulation algo knows it failed; will still likely have some triangulation even in this case
+	 */
+	bool GEOMETRYALGORITHMS_API Triangulate(TArray<FIndex3i>& Triangles);
+
 	/**
 	 * Check if current Graph has self-intersections; not optimized, only for debugging
 	 */
@@ -538,6 +571,10 @@ protected:
 
 		return num_hits > 0;
 	}
+
+private:
+	// Full-featured implementation for all Triangulate function variants to call
+	bool GEOMETRYALGORITHMS_API TriangulateInternal(TArray<FIndex3i>& Triangles, bool bHasBoundaryEdgeGroupID, int32 BoundaryEdgeGroupID, bool bHasHoleGroupID, int32 HoleEdgeGroupID, bool bLegacyKeepTrianglesIfBoundaryNotFound, TArray<int32>* SkippedEdges);
 };
 
 
diff --git a/Engine/Source/Runtime/GeometryCore/Private/CompGeom/Delaunay2.cpp b/Engine/Source/Runtime/GeometryCore/Private/CompGeom/Delaunay2.cpp
--- a/Engine/Source/Runtime/GeometryCore/Private/CompGeom/Delaunay2.cpp
+++ b/Engine/Source/Runtime/GeometryCore/Private/CompGeom/Delaunay2.cpp
@@ -134,7 +134,7 @@ struct FDelaunay2Connectivity
 	}
 
 	// return the triangle indices rotated s.t. the smallest vertex is first (and the winding is unchanged)
-	inline FIndex3i AsUniqueTriangle(const FIndex2i& Edge, int32 Vertex)
+	inline FIndex3i AsUniqueTriangle(const FIndex2i& Edge, int32 Vertex) const
 	{
 		if (Edge.A < Edge.B)
 		{
@@ -160,7 +160,7 @@ struct FDelaunay2Connectivity
 		}
 	}
 
-	bool GetFilledTriangles(TArray<FIndex3i>& TrianglesOut, TArrayView<const FIndex2i> Edges, FDelaunay2::EFillMode FillMode)
+	bool GetFilledTriangles(TArray<FIndex3i>& TrianglesOut, TArrayView<const FIndex2i> Edges, FDelaunay2::EFillMode FillMode) const
 	{
 		bool bWellDefinedResult = true;
 
@@ -176,7 +176,7 @@ struct FDelaunay2Connectivity
 			int32 Sign = 1;
 			if (Unoriented.A > Unoriented.B)
 			{
-				Swap(Unoriented.A, Unoriented.B);
+				Unoriented.Swap();
 				if (FillMode != FDelaunay2::EFillMode::Solid)
 				{
 					Sign = -1;
@@ -254,7 +254,7 @@ struct FDelaunay2Connectivity
 			//  and depth-first traversal would be more sensitive to starting triangle in that case
 			PickIdx = (PickIdx + 1) % ToWalk.Num();
 			const FWalk Walk = ToWalk[PickIdx];
-			ToWalk.RemoveAtSwap(PickIdx);
+			ToWalk.RemoveAtSwap(PickIdx, 1, false);
 			int32 Vert = EdgeToVert[Walk.Edge];
 			FIndex3i UniqueTri = AsUniqueTriangle(Walk.Edge, Vert);
 			if (UniqueTri.A < 0) // it's a ghost
@@ -304,6 +304,132 @@ struct FDelaunay2Connectivity
 		return bWellDefinedResult;
 	}
 
+	bool GetFilledTriangles(TArray<FIndex3i>& TrianglesOut, TArrayView<const FIndex2i> BoundaryEdges, TArrayView<const FIndex2i> HoleEdges) const
+	{
+		TSet<FIndex2i> BoundarySet, HoleSet;
+		auto AddToSet = [this](TArrayView<const FIndex2i> Edges, TSet<FIndex2i>& Set)
+		{
+			for (FIndex2i Edge : Edges)
+			{
+				if (bTrackDuplicateVertices)
+				{
+					FixDuplicatesOnEdge(Edge);
+				}
+				Edge.Sort();
+				Set.Add(Edge);
+			}
+		};
+		AddToSet(BoundaryEdges, BoundarySet);
+		AddToSet(HoleEdges, HoleSet);
+
+		TMap<FIndex3i, bool> InBoundary;
+		struct FWalk
+		{
+			FIndex2i Edge;
+		};
+		TArray<FWalk> ToWalkOuter;
+		TArray<FWalk> ToWalkInner;
+
+		auto WalkOuter = [&ToWalkOuter, &ToWalkInner, &HoleEdges, &BoundaryEdges](const FIndex2i& Edge)
+		{
+			FIndex2i Unoriented = Edge;
+			Unoriented.Sort();
+			if (HoleEdges.Contains(Unoriented))
+			{
+				return;
+			}
+			FIndex2i ReverseEdge(Edge.B, Edge.A);
+			if (BoundaryEdges.Contains(Unoriented))
+			{
+				ToWalkInner.Add({ ReverseEdge });
+			}
+			else
+			{
+				ToWalkOuter.Add({ ReverseEdge });
+			}
+		};
+
+		auto WalkInner = [&ToWalkInner, &HoleEdges](const FIndex2i& Edge)
+		{
+			FIndex2i Unoriented = Edge;
+			Unoriented.Sort();
+			if (HoleEdges.Contains(Unoriented))
+			{
+				return;
+			}
+			FIndex2i ReverseEdge(Edge.B, Edge.A);
+			ToWalkInner.Add({ ReverseEdge });
+		};
+
+		// First walk all the ghost triangles, through their non-ghost edge
+		int32 GhostTriCount = 0;
+		for (const TPair<FIndex2i, int32>& EdgeV : EdgeToVert)
+		{
+			if (EdgeV.Value == GhostIndex)
+			{
+				WalkOuter(EdgeV.Key);
+				GhostTriCount++;
+			}
+		}
+
+		// Next walk up to the outer boundary
+		while (!ToWalkOuter.IsEmpty())
+		{
+			const FWalk Walk = ToWalkOuter.Pop(false);
+			int32 Vert = EdgeToVert[Walk.Edge];
+			FIndex3i UniqueTri = AsUniqueTriangle(Walk.Edge, Vert);
+			if (UniqueTri.A < 0) // it's a ghost
+			{
+				continue;
+			}
+			bool* InBoundaryPtr = InBoundary.Find(UniqueTri);
+			if (InBoundaryPtr)
+			{
+				continue;
+			}
+			InBoundary.Add(UniqueTri, false);
+			WalkOuter(FIndex2i(Walk.Edge.B, Vert));
+			WalkOuter(FIndex2i(Vert, Walk.Edge.A));
+		}
+
+		// Finally walk up to the inner hole boundary
+		int32 SolidTriCount = 0;
+		while (!ToWalkInner.IsEmpty())
+		{
+			const FWalk Walk = ToWalkInner.Pop(false);
+			int32 Vert = EdgeToVert[Walk.Edge];
+			FIndex3i UniqueTri = AsUniqueTriangle(Walk.Edge, Vert);
+			if (UniqueTri.A < 0) // it's a ghost
+			{
+				continue;
+			}
+			bool* InBoundaryPtr = InBoundary.Find(UniqueTri);
+			if (InBoundaryPtr)
+			{
+				continue;
+			}
+			InBoundary.Add(UniqueTri, true);
+			WalkInner(FIndex2i(Walk.Edge.B, Vert));
+			WalkInner(FIndex2i(Vert, Walk.Edge.A));
+		}
+
+		TrianglesOut.Reset();
+		
+		TrianglesOut.Reserve(InBoundary.Num());
+		EnumerateTriangles([&InBoundary, &TrianglesOut](const FIndex2i& Edge, int32 Vertex) -> bool
+			{
+				FIndex3i Triangle(Vertex, Edge.A, Edge.B);
+				bool* InBoundaryPtr = InBoundary.Find(Triangle);
+				if (InBoundaryPtr && *InBoundaryPtr)
+				{
+					TrianglesOut.Add(Triangle);
+				}
+				return true;
+			}, true /*bSkipGhosts*/);
+		
+		return true;
+	}
+
 	static bool IsGhost(const FIndex2i& Edge, int32 Vertex)
 	{
 		return Edge.A == GhostIndex || Edge.B == GhostIndex || Vertex == GhostIndex;
@@ -378,19 +504,19 @@ struct FDelaunay2Connectivity
 		DuplicateVertices.Add(Duplicate, Orig);
 	}
 
-	bool HasDuplicateTracking()
+	bool HasDuplicateTracking() const
 	{
 		return bTrackDuplicateVertices;
 	}
 
-	void FixDuplicatesOnEdge(FIndex2i& Edge)
+	void FixDuplicatesOnEdge(FIndex2i& Edge) const
 	{
-		int32* OrigA = DuplicateVertices.Find(Edge.A);
+		const int32* OrigA = DuplicateVertices.Find(Edge.A);
 		if (OrigA)
 		{
 			Edge.A = *OrigA;
 		}
-		int32* OrigB = DuplicateVertices.Find(Edge.B);
+		const int32* OrigB = DuplicateVertices.Find(Edge.B);
 		if (OrigB)
 		{
 			Edge.B = *OrigB;
@@ -429,7 +555,7 @@ struct FDelaunay2Connectivity
 	// VisitFunctionType is expected to take (FIndex2i Edge, int32 Vertex) and return bool
 	// Returning false from VisitFn will end the enumeration early
 	template<typename VisitFunctionType>
-	void EnumerateOrientedEdges(VisitFunctionType VisitFn)
+	void EnumerateOrientedEdges(VisitFunctionType VisitFn) const
 	{
 		for (const TPair<FIndex2i, int32>& EdgeVert : EdgeToVert)
 		{
@@ -445,7 +571,7 @@ struct FDelaunay2Connectivity
 	// Returning false from VisitFn will end the enumeration early
 	// @param bSkipGhosts If true, skip ghost triangles (triangles connected to the ghost vertex)
 	template<typename VisitFunctionType>
-	void EnumerateTriangles(VisitFunctionType VisitFn, bool bSkipGhosts = false)
+	void EnumerateTriangles(VisitFunctionType VisitFn, bool bSkipGhosts = false) const
 	{
 		for (const TPair<FIndex2i, int32>& EdgeVert : EdgeToVert)
 		{
@@ -1171,7 +1297,7 @@ void FDelaunay2::GetTrianglesAndAdjacency(TArray<FIndex3i>& Triangles, TArray<FI
 	}
 }
 
-bool FDelaunay2::GetFilledTriangles(TArray<FIndex3i>& TrianglesOut, TArrayView<const FIndex2i> Edges, EFillMode FillMode)
+bool FDelaunay2::GetFilledTriangles(TArray<FIndex3i>& TrianglesOut, TArrayView<const FIndex2i> Edges, EFillMode FillMode) const
 {
 	if (!ensure(Connectivity.IsValid()))
 	{
@@ -1180,6 +1306,15 @@ bool FDelaunay2::GetFilledTriangles(TArray<FIndex3i>& TrianglesOut, TArrayView<c
 	return Connectivity->GetFilledTriangles(TrianglesOut, Edges, FillMode);
 }
 
+bool FDelaunay2::GetFilledTriangles(TArray<FIndex3i>& TrianglesOut, TArrayView<const FIndex2i> BoundaryEdges, TArrayView<const FIndex2i> HoleEdges) const
+{
+	if (!ensure(Connectivity.IsValid()))
+	{
+		return false;
+	}
+	return Connectivity->GetFilledTriangles(TrianglesOut, BoundaryEdges, HoleEdges);
+}
+
 bool FDelaunay2::IsDelaunay(TArrayView<const FVector2f> Vertices, TArrayView<const FIndex2i> SkipEdges) const
 {
 	if (ensure(Connectivity.IsValid()))
diff --git a/Engine/Source/Runtime/GeometryCore/Public/CompGeom/Delaunay2.h b/Engine/Source/Runtime/GeometryCore/Public/CompGeom/Delaunay2.h
--- a/Engine/Source/Runtime/GeometryCore/Public/CompGeom/Delaunay2.h
+++ b/Engine/Source/Runtime/GeometryCore/Public/CompGeom/Delaunay2.h
@@ -188,7 +188,18 @@ public:
 	 * @return							true if the result was well defined and consistent; false otherwise. Solid fill mode always has a well defined result; Winding-number-based fills do not if the edges have open spans.
 	 *									Note the triangulation will still be filled by best-effort even if the function returns false.
 	 */
-	bool GetFilledTriangles(TArray<FIndex3i>& TrianglesOut, TArrayView<const FIndex2i> Edges, EFillMode FillMode = EFillMode::PositiveWinding);
+	bool GetFilledTriangles(TArray<FIndex3i>& TrianglesOut, TArrayView<const FIndex2i> Edges, EFillMode FillMode = EFillMode::PositiveWinding) const;
+
+	/**
+	 * Get (by reference) the triangles that are inside the given edges, removing the outside-boundary triangles and the inside-hole triangles
+	 * 
+	 * @param TrianglesOut					Will be filled with a subset of the triangulation that is 'inside' the given edges, as defined by the FillMode
+	 * @param BoundaryEdges					Constrained edges in the triangulation that define the boundary of the desired shape
+	 * @param HoleEdges						Constrained edges in the triangulation that define inner holes of the desired shape
+	 * @return								true if any result was successfully computed (including an empty result).
+	 *										Currently only returns false if Triangulate() has not been called yet.
+	 */
+	bool GetFilledTriangles(TArray<FIndex3i>& TrianglesOut, TArrayView<const FIndex2i> BoundaryEdges, TArrayView<const FIndex2i> HoleEdges) const;
 
 	// @return true if this is a constrained Delaunay triangulation
 	bool IsConstrained() const
diff --git a/Engine/Source/Runtime/GeometryCore/Public/IndexTypes.h b/Engine/Source/Runtime/GeometryCore/Public/IndexTypes.h
--- a/Engine/Source/Runtime/GeometryCore/Public/IndexTypes.h
+++ b/Engine/Source/Runtime/GeometryCore/Public/IndexTypes.h
@@ -106,9 +106,15 @@ PRAGMA_ENABLE_DEPRECATION_WARNINGS
 
 	inline void Swap()
 	{
-		int tmp = A;
-		A = B;
-		B = tmp;
+		::Swap(A, B);
+	}
+
+	inline void Sort()
+	{
+		if (A > B)
+		{
+			Swap();
+		}
 	}
 
 	/**
