fix type mismatches in GeometryCore DynamicVector and RefCountVector

diff --git a/Engine/Source/Runtime/GeometryCore/Public/Util/DynamicVector.h b/Engine/Source/Runtime/GeometryCore/Public/Util/DynamicVector.h
--- a/Engine/Source/Runtime/GeometryCore/Public/Util/DynamicVector.h
+++ b/Engine/Source/Runtime/GeometryCore/Public/Util/DynamicVector.h
@@ -171,7 +171,7 @@ public:
 				SetCurBlock(Count);
 
 				// Remove empty block.
-				Blocks.Truncate(Count / BlockSize, false);
+				Blocks.Truncate(int32(Count / BlockSize), false);
 			}
 		}
 		else
@@ -466,7 +466,7 @@ private:
 				{					
 					for (uint32 Index = 0; Index < NumBlocks; ++Index)
 					{
-						const SIZE_T NumCopy = FMath::template Min<uint32>(NumCopyRemaining, BlockSize);
+						const SIZE_T NumCopy = FMath::template Min<SIZE_T>(NumCopyRemaining, (SIZE_T)BlockSize);
 						FMemory::Memcpy(BufferPtr, Elements[Index]->GetData(), NumCopy * sizeof(Type));
 						BufferPtr += NumCopy;
 						NumCopyRemaining -= BlockSize;
@@ -481,7 +481,7 @@ private:
 					for (uint32 Index = 0; Index < NumBlocks; ++Index)
 					{
 						ArrayType *const NewElement = new ArrayType;
-						const SIZE_T NumCopy = FMath::template Min<uint32>(NumCopyRemaining, BlockSize);
+						const SIZE_T NumCopy = FMath::template Min<SIZE_T>(NumCopyRemaining, (SIZE_T)BlockSize);
 						FMemory::Memcpy(NewElement->GetData(), BufferPtr, NumCopy * sizeof(Type));
 						Add(NewElement);
 						BufferPtr += NumCopy;
@@ -569,7 +569,7 @@ private:
 		// Reset block index for the last item and used item count within the last block.
 		// This is similar to what happens when computing the indices in operator[], but we additionally account for (1) the vector being empty and (2) that the
 		// used item count within the last block needs to be one more than the index of the last item. 
-		const int32 LastItemIndex = Count - 1;
+		const int32 LastItemIndex = int32(Count - 1);
 		CurBlock = Count != 0 ? LastItemIndex >> nShiftBits : 0;
 		CurBlockUsed = Count != 0 ? (LastItemIndex & BlockIndexBitmask) + 1 : 0;
 	}
@@ -694,25 +694,25 @@ public:
 	{
 		check(N >= 2);
 		return FIndex2i(
-			Data[TopIndex * N + 0],
-			Data[TopIndex * N + 1]);
+			(int)Data[TopIndex * N + 0],
+			(int)Data[TopIndex * N + 1]);
 	}
 	inline FIndex3i AsIndex3(unsigned int TopIndex) const
 	{
 		check(N >= 3);
 		return FIndex3i(
-			Data[TopIndex * N + 0],
-			Data[TopIndex * N + 1],
-			Data[TopIndex * N + 2]);
+			(int)Data[TopIndex * N + 0],
+			(int)Data[TopIndex * N + 1],
+			(int)Data[TopIndex * N + 2]);
 	}
 	inline FIndex4i AsIndex4(unsigned int TopIndex) const
 	{
 		check(N >= 4);
 		return FIndex4i(
-			Data[TopIndex * N + 0],
-			Data[TopIndex * N + 1],
-			Data[TopIndex * N + 2],
-			Data[TopIndex * N + 3]);
+			(int)Data[TopIndex * N + 0],
+			(int)Data[TopIndex * N + 1],
+			(int)Data[TopIndex * N + 2],
+			(int)Data[TopIndex * N + 3]);
 	}
 
 private:
diff --git a/Engine/Source/Runtime/GeometryCore/Public/Util/RefCountVector.h b/Engine/Source/Runtime/GeometryCore/Public/Util/RefCountVector.h
--- a/Engine/Source/Runtime/GeometryCore/Public/Util/RefCountVector.h
+++ b/Engine/Source/Runtime/GeometryCore/Public/Util/RefCountVector.h
@@ -136,7 +136,7 @@ public:
 			int j = (int)RefCounts.GetLength();
 			while (j < Index)
 			{
-				int InvalidCount = INVALID_REF_COUNT;	// required on older clang because a constexpr can't be passed by ref
+				unsigned short InvalidCount = INVALID_REF_COUNT;	// required on older clang because a constexpr can't be passed by ref
 				RefCounts.Add(InvalidCount);
 				FreeIndices.Add(j);
 				++j;
@@ -178,7 +178,7 @@ public:
 			int j = (int)RefCounts.GetLength();
 			while (j < Index)
 			{
-				int InvalidCount = INVALID_REF_COUNT;	// required on older clang because a constexpr can't be passed by ref
+				unsigned short InvalidCount = INVALID_REF_COUNT;	// required on older clang because a constexpr can't be passed by ref
 				RefCounts.Add(InvalidCount);
 				++j;
 			}
@@ -595,8 +595,8 @@ public:
 			return false;
 		}
 
-		const int32 Num = FMath::Max(Lhs.GetMaxIndex(), Rhs.GetMaxIndex());
-		for (int32 Idx = 0; Idx < Num; ++Idx)
+		const size_t Num = FMath::Max(Lhs.GetMaxIndex(), Rhs.GetMaxIndex());
+		for (size_t Idx = 0; Idx < Num; ++Idx)
 		{
 			const bool LhsIsValid = Lhs.IsValid(Idx);
 			if (LhsIsValid != Rhs.IsValid(Idx))
