Geometery Core: extends the blocked 3d grids and adds narrow band sparse SDF generation using these blocked grids.
unittests are included for the blocked grids and sdf generation, as well as benchmarks for sdf generation.
( This CL flushes out an old shelf of Ryan's.)

diff --git a/Engine/Source/Runtime/GeometryCore/Private/Implicit/SDFCalculationUtils.cpp b/Engine/Source/Runtime/GeometryCore/Private/Implicit/SDFCalculationUtils.cpp
new file mode 100644
--- /dev/null
+++ b/Engine/Source/Runtime/GeometryCore/Private/Implicit/SDFCalculationUtils.cpp
@@ -0,0 +1,104 @@
+// Copyright Epic Games, Inc. All Rights Reserved.
+
+
+#include "Implicit/SDFCalculationUtils.h"
+
+using namespace UE::Geometry;
+
+// find distance x0 is from triangle x1-x2-x3
+double UE::Geometry::PointTriangleDistance(const FVector3d& x0, const FVector3d& x1, const FVector3d& x2, const FVector3d& x3)
+{
+	// first find barycentric coordinates of closest point on infinite plane
+	FVector3d x13 = (x1 - x3);
+	FVector3d x23 = (x2 - x3);
+	FVector3d x03 = (x0 - x3);
+	double m13 = x13.SquaredLength(), m23 = x23.SquaredLength(), d = x13.Dot(x23);
+	double invdet = 1.0 / FMath::Max(m13 * m23 - d * d, 1e-30);
+	double a = x13.Dot(x03), b = x23.Dot(x03);
+	// the barycentric coordinates themselves
+	double w23 = invdet * (m23 * a - d * b);
+	double w31 = invdet * (m13 * b - d * a);
+	double w12 = 1 - w23 - w31;
+	if (w23 >= 0 && w31 >= 0 && w12 >= 0) // if we're inside the triangle
+	{
+		return Distance(x0, w23*x1 + w31*x2 + w12*x3);
+	}
+	else // we have to clamp to one of the edges
+	{
+		if (w23 > 0) // this rules out edge 2-3 for us
+		{
+			return FMath::Min(PointSegmentDistance(x0, x1, x2), PointSegmentDistance(x0, x1, x3));
+		}
+		else if (w31 > 0) // this rules out edge 1-3
+		{
+			return FMath::Min(PointSegmentDistance(x0, x1, x2), PointSegmentDistance(x0, x2, x3));
+		}
+		else // w12 must be >0, ruling out edge 1-2
+		{
+			return FMath::Min(PointSegmentDistance(x0, x1, x3), PointSegmentDistance(x0, x2, x3));
+		}
+	}
+}
+
+
+
+// find distance x0 is from triangle x1-x2-x3
+float UE::Geometry::PointTriangleDistance(const FVector3f& x0, const FVector3f& x1, const FVector3f& x2, const FVector3f& x3)
+{
+	// first find barycentric coordinates of closest point on infinite plane
+	FVector3f x13 = (x1 - x3);
+	FVector3f x23 = (x2 - x3);
+	FVector3f x03 = (x0 - x3);
+	float m13 = x13.SquaredLength(), m23 = x23.SquaredLength(), d = x13.Dot(x23);
+	float invdet = 1.0f / FMath::Max(m13 * m23 - d * d, 1e-30f);
+	float a = x13.Dot(x03), b = x23.Dot(x03);
+	// the barycentric coordinates themselves
+	float w23 = invdet * (m23 * a - d * b);
+	float w31 = invdet * (m13 * b - d * a);
+	float w12 = 1 - w23 - w31;
+	if (w23 >= 0 && w31 >= 0 && w12 >= 0) // if we're inside the triangle
+	{
+		return Distance(x0, w23*x1 + w31*x2 + w12*x3);
+	}
+	else // we have to clamp to one of the edges
+	{
+		if (w23 > 0) // this rules out edge 2-3 for us
+		{
+			return FMath::Min(PointSegmentDistance(x0, x1, x2), PointSegmentDistance(x0, x1, x3));
+		}
+		else if (w31 > 0) // this rules out edge 1-3
+		{
+			return FMath::Min(PointSegmentDistance(x0, x1, x2), PointSegmentDistance(x0, x2, x3));
+		}
+		else // w12 must be >0, ruling out edge 1-2
+		{
+			return FMath::Min(PointSegmentDistance(x0, x1, x3), PointSegmentDistance(x0, x2, x3));
+		}
+	}
+}
+
+
+
+// robust test of (X0,Y0) in the triangle (X1,Y1)-(X2,Y2)-(X3,Y3)
+// if true is returned, the barycentric coordinates are set in A,B,C.
+bool UE::Geometry::PointInTriangle2d(double X0, double Y0, double X1, double Y1, double X2, double Y2, double X3, double Y3, double& A, double& B, double& C)
+{
+	A = B = C = 0;
+	X1 -= X0; X2 -= X0; X3 -= X0;
+	Y1 -= Y0; Y2 -= Y0; Y3 -= Y0;
+	int signa = Orientation(X2, Y2, X3, Y3, A);
+	if (signa == 0) return false;
+	int signb = Orientation(X3, Y3, X1, Y1, B);
+	if (signb != signa) return false;
+	int signc = Orientation(X1, Y1, X2, Y2, C);
+	if (signc != signa) return false;
+	double sum = A + B + C;
+	// if the SOS signs match and are nonzero, there's no way all of A, B, and C are zero.
+	// TODO: is this mathematically impossible? can we just remove the check?
+	checkSlow(sum != 0); // TEXT("TCachingMeshSDF::PointInTriangle2d: impossible config?"));
+	A /= sum;
+	B /= sum;
+	C /= sum;
+	return true;
+}
+
diff --git a/Engine/Source/Runtime/GeometryCore/Public/Generators/MarchingCubes.h b/Engine/Source/Runtime/GeometryCore/Public/Generators/MarchingCubes.h
--- a/Engine/Source/Runtime/GeometryCore/Public/Generators/MarchingCubes.h
+++ b/Engine/Source/Runtime/GeometryCore/Public/Generators/MarchingCubes.h
@@ -130,7 +130,7 @@ public:
 
 		if (bEnableValueCaching)
 		{
-			BlockedCornerValuesGrid = FBlockedDenseGrid3f(CellDimensions.X + 1, CellDimensions.Y + 1, CellDimensions.Z + 1, FMathf::MaxReal);
+			BlockedCornerValuesGrid.Reset(CellDimensions.X + 1, CellDimensions.Y + 1, CellDimensions.Z + 1, FMathf::MaxReal);
 		}
 		InitHashTables();
 		ResetMesh();
@@ -170,11 +170,11 @@ public:
 		{
 			if (bEnableValueCaching)
 			{
-				BlockedCornerValuesGrid = FBlockedDenseGrid3f(CellDimensions.X + 1, CellDimensions.Y + 1, CellDimensions.Z + 1, FMathf::MaxReal);
+				BlockedCornerValuesGrid.Reset(CellDimensions.X + 1, CellDimensions.Y + 1, CellDimensions.Z + 1, FMathf::MaxReal);
 			}
 			if (bParallelCompute)
 			{
-				BlockedDoneCells = FBlockedDenseGrid3i(CellDimensions.X, CellDimensions.Y, CellDimensions.Z, 0);
+				BlockedDoneCells.Reset(CellDimensions.X, CellDimensions.Y, CellDimensions.Z, 0);
 			}
 		}
 		else
@@ -605,7 +605,7 @@ protected:
 		FGridCell Cell;
 		int vertTArray[12];
 
-		BlockedDoneCells = FBlockedDenseGrid3i(CellDimensions.X, CellDimensions.Y, CellDimensions.Z, 0);
+		BlockedDoneCells.Reset(CellDimensions.X, CellDimensions.Y, CellDimensions.Z, 0);
 
 		TArray<FVector3i> stack;
 
diff --git a/Engine/Source/Runtime/GeometryCore/Public/Implicit/SDFCalculationUtils.h b/Engine/Source/Runtime/GeometryCore/Public/Implicit/SDFCalculationUtils.h
new file mode 100644
--- /dev/null
+++ b/Engine/Source/Runtime/GeometryCore/Public/Implicit/SDFCalculationUtils.h
@@ -0,0 +1,65 @@
+// Copyright Epic Games, Inc. All Rights Reserved.
+
+#pragma once
+
+#include "Math/Vector.h"
+#include "VectorTypes.h"
+
+namespace UE
+{
+namespace Geometry
+{
+	using namespace UE::Math;
+
+	// These are utility functions shared by the various SDF calculation implementations
+	// (CachingMeshSDF, SweepignMeshSDF, SparseNarrowBandMeshSDF)
+
+
+	GEOMETRYCORE_API double PointTriangleDistance(const FVector3d& x0, const FVector3d& x1, const FVector3d& x2, const FVector3d& x3);
+
+	GEOMETRYCORE_API float PointTriangleDistance(const FVector3f& x0, const FVector3f& x1, const FVector3f& x2, const FVector3f& x3);
+
+
+	GEOMETRYCORE_API bool PointInTriangle2d(double X0, double Y0, double X1, double Y1, double X2, double Y2, double X3, double Y3, double& A, double& B, double& C);
+
+	// find distance x0 is from segment x1-x2
+	template<typename RealType> 
+	RealType PointSegmentDistance(const TVector<RealType>& x0, const TVector<RealType>& x1, const TVector<RealType>& x2)
+	{
+		TVector<RealType> DX = x2 - x1;
+		RealType m2 = DX.SquaredLength();
+		// find parameter value of closest point on segment
+		RealType s12 = (DX.Dot(x2 - x0) / m2);
+		if (s12 < 0)
+		{
+			s12 = 0;
+		}
+		else if (s12 > 1)
+		{
+			s12 = 1;
+		}
+		// and find the distance
+		return Distance(x0, s12*x1 + (1.0 - s12)*x2);
+	}
+
+
+	// calculate twice signed area of triangle (0,0)-(X1,Y1)-(X2,Y2)
+	// return an SOS-determined sign (-1, +1, or 0 only if it's a truly degenerate triangle)
+	template<typename RealType> 
+	int Orientation(RealType X1, RealType Y1, RealType X2, RealType Y2, RealType& TwiceSignedArea)
+	{
+		TwiceSignedArea = Y1 * X2 - X1 * Y2;
+		if (TwiceSignedArea > 0) return 1;
+		else if (TwiceSignedArea < 0) return -1;
+		else if (Y2 > Y1) return 1;
+		else if (Y2 < Y1) return -1;
+		else if (X1 > X2) return 1;
+		else if (X1 < X2) return -1;
+		else return 0; // only true when X1==X2 and Y1==Y2
+	}
+
+
+
+
+}
+}
\ No newline at end of file
diff --git a/Engine/Source/Runtime/GeometryCore/Public/Implicit/SparseNarrowBandMeshSDF.h b/Engine/Source/Runtime/GeometryCore/Public/Implicit/SparseNarrowBandMeshSDF.h
new file mode 100644
--- /dev/null
+++ b/Engine/Source/Runtime/GeometryCore/Public/Implicit/SparseNarrowBandMeshSDF.h
@@ -0,0 +1,1336 @@
+// Copyright Epic Games, Inc. All Rights Reserved.
+
+// Port of geometry3Sharp MeshSignedDistanceGrid
+
+#pragma once
+
+#include "MathUtil.h"
+#include "MeshQueries.h"
+#include "Util/IndexUtil.h"
+#include "Spatial/MeshAABBTree3.h"
+#include "Spatial/BlockedLayout3.h"
+#include "Spatial/BlockedDenseGrid3.h"
+#include "Async/ParallelFor.h"
+#include "Implicit/GridInterpolant.h"
+#include "Implicit/SDFCalculationUtils.h"
+#include "Containers/StaticArray.h"
+
+#include <atomic> // note: UE FGenericPlatformAtomics are deprecated. 
+
+namespace UE
+{
+namespace Geometry
+{
+
+
+/**
+ * Compute discretely-sampled (ie gridded) signed distance field for a mesh within a specified narrow band.
+ * 
+ * The basic approach is, first compute exact Distances in a narrow band, and then
+ * The resulting unsigned Grid is then signed using ray-intersection counting, which
+ * is also computed on the Grid, so no BVH is necessary
+ * 
+ * The underlying grid data structures are blocked grids where spatial blocks are allocated
+ * as needed, so the memory footprint will generally be much lower than the TSweepingMeshSDF, or TCachingMeshSDF.
+ * But if distances in the full grid domain are desired (not just a narrow band), then the TSweepingMeshSDF should be used.
+ *
+ * You can optionally provide a spatial data structure to allow faster computation of
+ * narrow-band Distances; this is most beneficial if we want a wider band 
+ * 
+ * Caveats:
+ *  - the "narrow band" is based on triangle bounding boxes, so it is not necessarily
+ *    that "narrow" if you have large triangles on a diagonal to Grid axes
+ *
+ * 
+ * This code is based on the implementation found at [REDACTED_PROFILE_URL]/SDFGen
+ */
+template <class TriangleMeshType>
+class TSparseNarrowBandMeshSDF
+{
+public:
+	// INPUTS
+
+	const TriangleMeshType* Mesh = nullptr;
+	TMeshAABBTree3<TriangleMeshType>* Spatial = nullptr;       // required for the NarrowBand_SpatialFloodFill method.
+	TUniqueFunction<bool(const FVector3d&)> IsInsideFunction;  // if not supplied, the intersection counting method is used.
+
+	// Relates the grids index space, to mesh distance units.
+	float CellSize;
+
+	// Width of the band around triangles for which exact Distances are computed
+	// (In fact this is conservative, the band is often larger locally)
+	int ExactBandWidth = 1;
+
+	// Bounds of Grid will be expanded this much in positive and negative directions.
+	// Useful for if you want field to extend outwards.
+	FVector3d ExpandBounds = FVector3d::Zero();
+
+	// Most of this parallelizes very well, makes a huge speed difference
+	bool bUseParallel = true;
+
+	// If the number of cells in any dimension may exceed this, CellSize will be automatically increased to keep cell count reasonable
+	int ApproxMaxCellsPerDimension = 4096;
+
+	// The narrow band is always computed exactly, and the full Grid will be signed if bComputeSigns is true.
+	enum EComputeModes
+	{
+		NarrowBandOnly = 1,
+		NarrowBand_SpatialFloodFill = 2
+	};
+	EComputeModes ComputeMode = EComputeModes::NarrowBandOnly;
+
+	// how wide of narrow band should we compute. This value is 
+	// currently only used NarrowBand_SpatialFloodFill, as
+	// we can efficiently explore the space
+	// (in that case ExactBandWidth is only used to initialize the grid extents, and this is used instead during the flood fill)
+	double NarrowBandMaxDistance = 0;
+
+	// If true, sign of full grid will be computed using either IsInsideFunction (if provided) or crossing as controlled by the InsideMode
+	bool bComputeSigns = true;
+
+	// What counts as "inside" the mesh. Crossing count does not use triangle
+	// Orientation, so inverted faces are fine, but overlapping shells or self intersections
+	// will be filled using even/odd rules (as seen along X axis...)
+	// Winding count is basically mesh winding number, handles overlap shells and
+	// self-intersections, but inverted shells are 'subtracted', and inverted faces are a disaster.
+	// Both modes handle internal cavities, neither handles open sheets.
+	enum EInsideModes
+	{
+		CrossingCount = 0,
+		WindingCount = 1
+	};
+	EInsideModes InsideMode = EInsideModes::WindingCount;
+
+	// Implementation computes the triangle closest to each Grid cell, can
+	// return this Grid if desired (only reason not to is avoid hanging onto memory)
+	bool bWantClosestTriGrid = false;
+
+	// Grid of per-cell crossing or parity counts
+	bool bWantIntersectionsGrid = true;
+
+	/** if this function returns true, the calculation will abort */
+	TFunction<bool()> CancelF = []()
+	{
+		return false;
+	};
+
+	// OUTPUTS
+
+	FVector3f GridOrigin;
+	FBlockedGrid3f Grid;
+
+protected:
+	// grid of closest triangle indices; only available if bWantClosestTriGrid is set to true.  Access via GetClosestTriGrid()
+	FBlockedGrid3i ClosestTriGrid;
+	// grid of intersection counts; only available if bWantIntersectionsGrid is set to true.  Access via GetIntersectionsGrid()
+	FBlockedGrid3i IntersectionsGrid;
+
+
+protected:
+	// structs used in a scatter / gather of triangles into the grid blocks they overlap.
+	
+	// Grid that holds one atomic<int> in each block
+	struct FScatterCounter : public TBlockedGrid3Layout<FBlockedGrid3f::BlockSize>
+	{
+		typedef TBlockedGrid3Layout<FBlockedGrid3f::BlockSize> MyLayout;
+		FScatterCounter(const FVector3i& Dims, bool bInitParallel)
+			: MyLayout(Dims)
+		{
+			const FVector3i BlockDims = MyLayout::GetBlockDimensions();
+			const int32 NumBlocks = BlockDims[0] * BlockDims[1] * BlockDims[2];
+			NumObjectsPerBlock.SetNum(NumBlocks);
+			std::atomic<int>* PerBlock = NumObjectsPerBlock.GetData();
+			ParallelFor(NumBlocks, [PerBlock](int32 id)
+			{
+				PerBlock[id] = 0;
+			}, !bInitParallel);
+		}
+		
+		// increment the counter for the indicated block
+		void AtomicIncrement(const int32 BlockIndex)
+		{
+			std::atomic<int>& NumObjects = NumObjectsPerBlock.GetData()[BlockIndex];
+			NumObjects.fetch_add(1, std::memory_order_relaxed);
+		}
+		
+		TArray<std::atomic<int>> NumObjectsPerBlock;
+	};
+
+	// sparse grid that stores array of TriIDs each blocks.
+	struct FTriIDBlockGrid : public TBlockedGrid3Layout<FBlockedGrid3f::BlockSize>
+	{
+		typedef TBlockedGrid3Layout<FBlockedGrid3f::BlockSize> MyLayout;
+
+		FTriIDBlockGrid(const FScatterCounter& TriCounter, bool bInitParallel)
+			: MyLayout(TriCounter.GetDimensions())
+		{
+			const FVector3i DimAsBlocks = MyLayout::GetBlockDimensions();
+			const int32 NumBlocks = DimAsBlocks.X * DimAsBlocks.Y * DimAsBlocks.Z;
+			TrisInBlock.SetNum(NumBlocks);
+			PerBlockNum.SetNum(NumBlocks);
+			const std::atomic<int>* PerBlockCounted = TriCounter.NumObjectsPerBlock.GetData();
+			std::atomic<int>* SizePerBlock = PerBlockNum.GetData();
+			ParallelFor(NumBlocks, [this, SizePerBlock, PerBlockCounted, &TriCounter](int32 id)
+			{
+				TArray<int32>& Tris = TrisInBlock[id];
+				const int32 n = PerBlockCounted[id];
+				Tris.AddUninitialized(n);
+				SizePerBlock[id] = 0;
+			}, !bInitParallel);
+		}
+
+
+
+		void AddTriID(const int32& BlockIndex, const int32& TriId)
+		{
+			TArray<int32>& Tris = TrisInBlock.GetData()[BlockIndex];
+			std::atomic<int>& Num = PerBlockNum.GetData()[BlockIndex];
+
+			int32 i = Num.fetch_add(1, std::memory_order_relaxed);;
+			Tris.GetData()[i] = TriId;
+		}
+		TArray<int32> GetOccupiedBlockIDs() const 
+		{
+			const int32 NumBlocks = TrisInBlock.Num();
+			int32 NumBlocksWithTris = 0;
+			for (int32 i = 0; i < NumBlocks; ++i)
+			{
+				if (TrisInBlock[i].Num() != 0)
+				{
+					NumBlocksWithTris++;
+				}
+			}
+			
+			TArray<int32> TmpArray;
+			TmpArray.Reserve(NumBlocksWithTris);
+			for (int32 i = 0; i < NumBlocks; ++i)
+			{
+				if (TrisInBlock[i].Num() != 0)
+				{
+					TmpArray.Add(i);
+				}
+			}
+			return TmpArray;
+		};
+
+		TArray<std::atomic<int>>  PerBlockNum;
+		TArray<TArray<int32>>     TrisInBlock;
+	};
+
+	struct FTriBBox
+	{
+		FVector3i BoxMin = FVector3i(INT32_MAX, INT32_MAX, INT32_MAX);
+		FVector3i BoxMax = FVector3i(INT32_MIN, INT32_MIN, INT32_MIN);
+
+		bool IsValid() const
+		{
+			return ((BoxMax[0] >= BoxMin[0]) && (BoxMax[1] >= BoxMin[1]) && (BoxMax[2] >= BoxMin[2]));
+		}
+
+		// reduce this box by clamping with specified min/max
+		void Clamp(const FVector3i& OtherMin, const FVector3i& OtherMax)
+		{	
+			BoxMin[0] = FMath::Clamp(BoxMin[0], OtherMin[0], OtherMax[0]);
+			BoxMin[1] = FMath::Clamp(BoxMin[1], OtherMin[1], OtherMax[1]);
+			BoxMin[2] = FMath::Clamp(BoxMin[2], OtherMin[2], OtherMax[2]);
+
+			BoxMax[0] = FMath::Clamp(BoxMax[0], OtherMin[0], OtherMax[0]);
+			BoxMax[1] = FMath::Clamp(BoxMax[1], OtherMin[1], OtherMax[1]);
+			BoxMax[2] = FMath::Clamp(BoxMax[2], OtherMin[2], OtherMax[2]);
+		}
+	};
+
+public:
+
+	/**
+	 * @param Mesh Triangle mesh to build an SDF around
+	 * @param CellSize Spacing between Grid points
+	 * @param Spatial Optional AABB tree; note it *must* be provided if ComputeMode is set to NarrowBand_SpatialFloodFill
+	 */
+	TSparseNarrowBandMeshSDF(
+		const TriangleMeshType* Mesh = nullptr, 
+		double CellSize = 1, 
+		TMeshAABBTree3<TriangleMeshType>* Spatial = nullptr ) 
+		: Mesh(Mesh), Spatial(Spatial), CellSize(CellSize)
+	{
+	}
+
+	TTriLinearGridInterpolant<TSparseNarrowBandMeshSDF> MakeInterpolant()
+	{
+		return TTriLinearGridInterpolant<TSparseNarrowBandMeshSDF>(this, (FVector3d)GridOrigin, CellSize, Dimensions());
+	}
+
+	/**
+	 * Encodes heuristic for deciding whether it will be faster to use
+	 *  - NarrowBand_SpatialFloodFill
+	 *  - Or NarrowBandOnly
+	 * The heuristic is to use the Spatial method unless you have relatively long edges relative to the width
+	 * of the narrow band.  For a 1-cell wide band, or one where few triangles interact w/ each cell in the band
+	 * the spatial overhead will probably not be worthwhile.
+	 * TODO: test and tune this ...
+	 */
+	static bool ShouldUseSpatial(int32 ExactCells, double CellSize, double AvgEdgeLen)
+	{
+		double EdgesPerCell = CellSize / AvgEdgeLen;
+		double EdgesPerBand = ExactCells * EdgesPerCell;
+		return ExactCells > 1 && EdgesPerBand >= .25;
+	}
+
+	bool Validate(const FAxisAlignedBox3d& Bounds)
+	{
+		if (Bounds.IsEmpty() || !FMath::IsFinite(Bounds.MaxDim()))
+		{
+			return false;
+		}
+		if (CellSize <= 0 || !FMath::IsFinite(CellSize))
+		{
+			return false;
+		}
+		if (ApproxMaxCellsPerDimension < 5)
+		{
+			return false;
+		}
+		return true;
+	}
+
+	/**
+	 * Compute the SDF
+	 * @param Bounds Bounding box for the mesh data (passed as param it is usually already available, depending on TriangleMeshType and whether an AABB tree was provided)
+	 * @return false if cancelled or failed (e.g. due to invalid arguments); true otherwise
+	 */
+	bool Compute(FAxisAlignedBox3d Bounds)
+	{
+		if (!ensure(Validate(Bounds)))
+		{
+			return false;
+		}
+
+		if (!ensureMsgf(ExactBandWidth*2 < ApproxMaxCellsPerDimension, TEXT("Cannot request band wider than half the max cells per dimension")))
+		{
+			ExactBandWidth = FMath::Max(ApproxMaxCellsPerDimension/2-1, 1);
+		}
+
+		double MaxDim = MaxElement(Bounds.Max - Bounds.Min + ExpandBounds * 2.0);
+		if (!ensureMsgf(MaxDim / CellSize <= ApproxMaxCellsPerDimension - 2 * ExactBandWidth, TEXT("SDF resolution clamped to avoid excessive memory use")))
+		{
+			CellSize = MaxDim / (ApproxMaxCellsPerDimension - 2 * ExactBandWidth);
+			if (!ensure(CellSize > 0 && FMath::IsFinite(CellSize)))
+			{
+				return false;
+			}
+		}
+
+		float fBufferWidth = ExactBandWidth * CellSize;
+		if (ComputeMode == EComputeModes::NarrowBand_SpatialFloodFill)
+		{
+			fBufferWidth = (float)FMath::Max(fBufferWidth, float(NarrowBandMaxDistance));
+		}
+		GridOrigin = (FVector3f)Bounds.Min - fBufferWidth * FVector3f::One() - (FVector3f)ExpandBounds;
+		FVector3f max = (FVector3f)Bounds.Max + fBufferWidth * FVector3f::One() + (FVector3f)ExpandBounds;
+		int32 NI = (int32)((max.X - GridOrigin.X) / CellSize) + 1;
+		int32 NJ = (int32)((max.Y - GridOrigin.Y) / CellSize) + 1;
+		int32 NK = (int32)((max.Z - GridOrigin.Z) / CellSize) + 1;
+
+		
+		if (ComputeMode == EComputeModes::NarrowBand_SpatialFloodFill)
+		{
+			if (!ensureMsgf(Spatial && NarrowBandMaxDistance != 0, TEXT("SweepingMeshSDF::Compute: must set Spatial data structure and band max distance")))
+			{
+				return false;
+			}
+			make_level_set3_parallel_floodfill(GridOrigin, CellSize, NI, NJ, NK, Grid, NarrowBandMaxDistance);
+		}
+		else // In general this path is faster than NarrowBand_SpatialFloodFill as it is lockless.
+		{
+			if (Spatial != nullptr) 
+			{
+				make_level_set3_parallel_spatial(GridOrigin, CellSize, NI, NJ, NK, Grid, ExactBandWidth);
+			}
+			else
+			{
+				make_level_set3_parallel(GridOrigin, CellSize, NI, NJ, NK, Grid, ExactBandWidth);
+			}
+		}
+
+		if (bComputeSigns)
+		{
+			update_signs(GridOrigin, CellSize, NI, NJ, NK, Grid, IntersectionsGrid);
+		}
+		
+
+		CleanupUnwanted();
+
+		return !CancelF();
+	}
+
+
+	/** @return the number of cells in each direction*/
+	FVector3i Dimensions() const
+	{
+		return Grid.GetDimensions();
+	}
+
+	/** @return a grid that can be sampled within the narrow band to find the ID of the closest triangle.  Note: this is populated during Compute()*/
+	const FBlockedGrid3i& GetClosestTriGrid() const
+	{
+		checkf(bWantClosestTriGrid, TEXT("Set bWantClosestTriGrid=true to return this value"));
+		return ClosestTriGrid;
+	}
+
+	/** @return a grid that can be sampled within the narrow band to find cells that actually contain mesh triangles.  Note: this is populated during Compute()*/
+	const FBlockedGrid3i& GetIntersectionsGrid() const
+	{
+		checkf(bWantIntersectionsGrid, TEXT("Set bWantIntersectionsGrid=true to return this value"));
+		return IntersectionsGrid;
+	}
+
+	/** @return the distance value at the specified grid cell.*/
+	float At(int32 I, int32 J, int32 K) const
+	{
+		return Grid.GetValue(I, J, K);
+	}
+
+	/** @return the distance value at the specified grid cell.*/
+	float operator[](const FVector3i& ijk) const
+	{
+		return Grid.GetValue(ijk);
+	}
+	
+	/** @return the distance value at the specified grid cell.*/
+	float GetValue(const FVector3i& ijk) const // TTriLinearGridInterpolant interface 
+	{
+		return Grid.GetValue(ijk);
+	}
+
+	/** @return the physical space location of specified grid point.*/
+	FVector3f CellCenter(int32 I, int32 J, int32 K) const
+	{
+		return cell_center(FVector3i(I, J, K));
+	}
+
+
+
+private:
+	
+	// iterate over a bounded grid region applying a functor at each cell
+	template <typename FunctorType>
+	void inclusive_loop_3d(const FVector3i& start, const FVector3i& inclusive_end, FunctorType func)
+	{
+		FVector3i ijk;
+		for (ijk[2] = start[2]; ijk[2] <= inclusive_end[2]; ++ijk[2])
+		{
+			for (ijk[1] = start[1]; ijk[1] <= inclusive_end[1]; ++ijk[1])
+			{
+				for (ijk[0] = start[0]; ijk[0] <= inclusive_end[0]; ++ijk[0])
+				{
+					func(ijk);
+				}
+			}
+		}
+	}
+
+	FVector3f cell_center(const FVector3i& IJK) const
+	{
+		return FVector3f((float)IJK.X * CellSize + GridOrigin[0],
+						 (float)IJK.Y * CellSize + GridOrigin[1],
+						 (float)IJK.Z * CellSize + GridOrigin[2]);
+	}
+
+	float upper_bound(const FVector3i& GridDimensions, float DX) const
+	{
+		return (float)((GridDimensions[0] + GridDimensions[1] + GridDimensions[2]) * DX);
+	}
+
+	// if provided, uses IsInsideFunction to determine sign, otherwise uses a counting method.
+	void update_signs(const FVector3f Origin, float DX, int32 NI, int32 NJ, int32 NK, FBlockedGrid3f& Distances, FBlockedGrid3i& IntersectionCountGrid)
+	{
+	
+		typedef FBlockedGrid3f::BlockData3Type    FloatBlockData3Type;
+
+
+		if (IsInsideFunction)  // Compute signs using inside/outside oracle.
+		{
+			
+			TArray<FloatBlockData3Type*> DistanceBlocks = Distances.GetAllocatedBlocks();
+			const int32 NumBlocksAllocated = DistanceBlocks.Num();
+
+			// evaluate the inside/outside oracle at each cell within the allocated blocks.
+			ParallelFor(NumBlocksAllocated, [this, &DistanceBlocks, &Distances, Origin, DX](int32 AllocatedBlockIdx)
+			{
+
+				FloatBlockData3Type& DistanceBlock = *DistanceBlocks[AllocatedBlockIdx];
+
+				const FVector3i BlockIJK = Distances.BlockIndexToBlockIJK(DistanceBlock.Id);
+				const FVector3i BlockOrigin = FloatBlockData3Type::BlockSize * BlockIJK;
+
+				const int32 ElemCount = FloatBlockData3Type::ElemCount;
+				for (int32 LocalIndex = 0; LocalIndex < ElemCount; ++LocalIndex)
+				{
+					const FVector3i GridCoord = DistanceBlock.ToLocalIJK(LocalIndex) + BlockOrigin;
+					const FVector3d Pos((float)GridCoord.X * DX + Origin[0], (float)GridCoord.Y * DX + Origin[1], (float)GridCoord.Z * DX + Origin[2]);
+					if (this->IsInsideFunction(Pos))
+					{
+						float& dist = DistanceBlock.At(LocalIndex);
+						dist *= -1;
+					}
+				}
+			}, !this->bUseParallel);
+			
+			if (CancelF())
+			{
+				return;
+			}
+
+			// update default value for unallocated blocks: evaluate the inside/outside oracle at the center of each.
+
+			const int32 NumBlocks = Distances.GetNumBlocks();
+			const FVector3d BlockCenter(0.5 * DX * (double)FloatBlockData3Type::BlockSize);
+			ParallelFor(NumBlocks, [this, &Distances, Origin, DX, BlockCenter](int32 BlockIndex)
+				{
+					const FVector3i BlockIJK = Distances.BlockIndexToBlockIJK(BlockIndex);
+					if (!Distances.IsBlockAllocated(BlockIJK))
+					{
+						const FVector3i BlockOrigin = FloatBlockData3Type::BlockSize * BlockIJK;
+						const FVector3d BlockOriginPos((float)BlockOrigin.X * DX + Origin.X, (float)BlockOrigin.Y * DX + Origin.Y, (float)BlockOrigin.Z * DX + Origin.Z);
+						const FVector3d Pos = BlockOriginPos + BlockCenter;
+						if (this->IsInsideFunction(Pos))
+						{
+							Distances.ProcessBlockDefaultValue(BlockIJK, [](float& value) {value = -value; });
+						}
+					}
+
+				}, !this->bUseParallel);
+
+		}
+		else // use intersection counts and the specified InsideMode to determine inside/outside regions.
+		{
+			// intersection_count(I,J,K) is # of tri intersections in (I-1,I]x{J}x{K}
+			
+			compute_intersections(Origin, DX, NI, NJ, NK, Distances, IntersectionCountGrid);
+			if (CancelF())
+			{
+				return;
+			}
+
+			// then figure out signs (inside/outside) from intersection counts
+			compute_signs(NI, NJ, NK, IntersectionCountGrid, Distances);
+		}
+	} // update_signs
+
+	// compute AABBs for the tris in the mesh.  the resulting TArray is indexed by the same TID as the mesh
+	// ( note, if the mesh doesn't have compact TIDs, the boxes that correspond to gaps in the TIDs will be in a state IsInvalid() == true.)
+	TArray<FTriBBox> compute_tri_bboxes(FVector3f Origin, float DX, int32 NI, int32 NJ, int32 NK, int32 ExactBand)
+	{
+		const double invdx = 1.0 / DX;
+		const FVector3d org((double)Origin[0], (double)Origin[1], (double)Origin[2]);
+
+		TArray<FTriBBox> TriBBoxes; TriBBoxes.SetNum(Mesh->MaxTriangleID());
+
+		auto ComputeTriBoundingBox = [this, ExactBand, NI, NJ, NK, invdx, org](const int32 TID, FVector3i& BBMin, FVector3i& BBMax)
+									{
+										FVector3d xp = FVector3d::Zero(), xq = FVector3d::Zero(), xr = FVector3d::Zero();
+										Mesh->GetTriVertices(TID, xp, xq, xr);
+
+										// real IJK coordinates of xp/xq/xr
+
+										const FVector3d fp = (xp - org) * invdx;
+										const FVector3d fq = (xq - org) * invdx;
+										const FVector3d fr = (xr - org) * invdx;
+
+										// clamped integer bounding box of triangle plus/minus exact-band
+										BBMin[0] = FMath::Clamp(((int)FMath::Min3(fp[0], fq[0], fr[0])) - ExactBand, 0, NI - 1);
+										BBMin[1] = FMath::Clamp(((int)FMath::Min3(fp[1], fq[1], fr[1])) - ExactBand, 0, NJ - 1);
+										BBMin[2] = FMath::Clamp(((int)FMath::Min3(fp[2], fq[2], fr[2])) - ExactBand, 0, NK - 1);
+
+										BBMax[0] = FMath::Clamp(((int)FMath::Max3(fp[0], fq[0], fr[0])) + ExactBand + 1, 0, NI - 1);
+										BBMax[1] = FMath::Clamp(((int)FMath::Max3(fp[1], fq[1], fr[1])) + ExactBand + 1, 0, NJ - 1);
+										BBMax[2] = FMath::Clamp(((int)FMath::Max3(fp[2], fq[2], fr[2])) + ExactBand + 1, 0, NK - 1);
+									};
+
+		ParallelFor(Mesh->MaxTriangleID(), [this, &ComputeTriBoundingBox, &TriBBoxes](int TID)
+		{
+			if (!Mesh->IsTriangle(TID)) return; // skipped BBoxes will be left in 'bbox.IsInvalid() == true' state
+
+			FVector3i& BoxMin = TriBBoxes.GetData()[TID].BoxMin;
+			FVector3i& BoxMax = TriBBoxes.GetData()[TID].BoxMax;
+			ComputeTriBoundingBox(TID, BoxMin, BoxMax);
+		});
+
+		return TriBBoxes;
+	}
+
+	// scatter box ids (based on the overlap of bbox with grid block) to return a grid that holds IDs in block-level TArrays. 
+	FTriIDBlockGrid scatter_tris(const TArray<FTriBBox>& TriBBoxes, int32 NI, int32 NJ, int32 NK)
+	{
+		bool bCancel = false;
+
+		// count number of tris that will scatter to each block.
+		FScatterCounter TriCounter(FVector3i(NI, NJ, NK), this->bUseParallel);
+		{
+			const int32 NumBBoxes = TriBBoxes.Num();
+			ParallelFor(NumBBoxes, [this, &TriBBoxes, &TriCounter, &bCancel](const  int32 TID)
+			{
+				constexpr int32 BlockSize = FBlockedGrid3f::BlockSize;
+
+				const FTriBBox& TriBBox = TriBBoxes.GetData()[TID];
+				if (!TriBBox.IsValid())
+				{
+					return;
+				}
+				if (TID % 100 == 0)
+				{
+					bCancel = this->CancelF();
+				}
+				if (bCancel)
+				{
+					return;
+				}
+
+				// convert bbox to block-level ijk
+				FVector3i box_min(TriBBox.BoxMin[0] / BlockSize, TriBBox.BoxMin[1] / BlockSize, TriBBox.BoxMin[2] / BlockSize);
+				FVector3i box_max(TriBBox.BoxMax[0] / BlockSize, TriBBox.BoxMax[1] / BlockSize, TriBBox.BoxMax[2] / BlockSize);
+
+				// increment the count for each block this tri scatters to.
+				inclusive_loop_3d(box_min, box_max, 
+								[&TriCounter](const FVector3i& block_ijk) 
+								{
+									const int32 BlockIndex = TriCounter.BlockIJKToBlockIndex(block_ijk);
+									TriCounter.AtomicIncrement(BlockIndex); 
+								});
+				
+			}, !this->bUseParallel);
+		}
+
+		// scatter the tris into the TriIDBlockedGrid
+		FTriIDBlockGrid TriIDBlockGrid(TriCounter, this->bUseParallel);
+		{
+			const int32 NumBBoxes = TriBBoxes.Num();
+			ParallelFor(NumBBoxes, [this, &TriBBoxes, &TriIDBlockGrid, &bCancel](const  int32 TID)
+			{
+				constexpr int32 BlockSize = FBlockedGrid3f::BlockSize;
+
+				const FTriBBox& TriBBox = TriBBoxes.GetData()[TID];
+				if (!TriBBox.IsValid())
+				{
+					return;
+				}
+				if (TID % 100 == 0)
+				{
+					bCancel = this->CancelF();
+				}
+				if (bCancel)
+				{
+					return;
+				}
+
+				// convert bbox to block-level ijk
+				FVector3i box_min(TriBBox.BoxMin[0] / BlockSize, TriBBox.BoxMin[1] / BlockSize, TriBBox.BoxMin[2] / BlockSize);
+				FVector3i box_max(TriBBox.BoxMax[0] / BlockSize, TriBBox.BoxMax[1] / BlockSize, TriBBox.BoxMax[2] / BlockSize);
+
+				// increment the count for each block this tri scatters to.
+				inclusive_loop_3d(box_min, box_max, 
+								[&TriIDBlockGrid, TID](const FVector3i& block_ijk) 
+								{
+										const int32 BlockIndex = TriIDBlockGrid.BlockIJKToBlockIndex(block_ijk);
+										TriIDBlockGrid.AddTriID(BlockIndex, TID); // blocks have already been allocated, and this is a threadsafe add to an array
+								});
+			
+			}, !this->bUseParallel);
+		}
+
+		return TriIDBlockGrid;
+	}
+
+	// lockless method - uses a scatter/ gather pattern to implement a painters algorithm (rasterizing all triangles but retaining only the min). 
+	void make_level_set3_parallel(const FVector3f Origin, float DX, int32 NI, int32 NJ, int32 NK, FBlockedGrid3f& DistanceGrid, int32 ExactBand)
+	{
+		const float upper_bound_value = upper_bound(FVector3i(NI, NJ, NK), DX);
+
+		DistanceGrid.Reset(NI, NJ, NK, upper_bound_value);
+		ClosestTriGrid.Reset(NI, NJ, NK, -1);
+
+		// bboxes for each tri.
+		const TArray<FTriBBox> TriBBoxes = compute_tri_bboxes(Origin, DX, NI, NJ, NK, ExactBand);
+
+		// blocked grid of tri ids - identifies the tris that overlap with each grid block.
+		const FTriIDBlockGrid TIDBlockGrid = scatter_tris(TriBBoxes, NI, NJ, NK);
+
+		// identify the blocks that are occupied by tris
+		const TArray<int32> OccupiedBlocks = TIDBlockGrid.GetOccupiedBlockIDs();
+
+		if (this->CancelF())
+		{
+			return;
+		}
+
+
+		const double invdx = 1.0 / DX;
+		const FVector3d org((double)Origin[0], (double)Origin[1], (double)Origin[2]);
+
+		// parallelize over occupied blocks.
+		// within each block, loop over the associated triangles, rasterizing distance (retaining the min).
+		{
+
+			// actually raster the tris into the distance and closest tri grid
+			const int32 NumOccupied = OccupiedBlocks.Num();
+			ParallelFor(NumOccupied, [this, &TriBBoxes, &OccupiedBlocks, &TIDBlockGrid, &DistanceGrid, ExactBand, org, invdx, DX](const int32& OccupiedId)
+			{
+				const FVector3i GridDimensions = DistanceGrid.GetDimensions();
+				const int32 BlockIndex = OccupiedBlocks.GetData()[OccupiedId];
+				const TArray<int32>& TrisToRasterize = TIDBlockGrid.TrisInBlock[BlockIndex];
+
+				const FVector3i BlockIJK = DistanceGrid.BlockIndexToBlockIJK(BlockIndex);
+				// allocates block on demand.  doing so here assumes the parallel memory allocator is up to it..
+				FBlockedGrid3f::BlockData3Type& DistanceBlock   = DistanceGrid.TouchBlockData(BlockIJK);
+				FBlockedGrid3i::BlockData3Type& ClosestTriBlock = ClosestTriGrid.TouchBlockData(BlockIJK);
+			
+				// bounding box for this block [inclusive, inclusive]
+				
+				const FVector3i BlockMin = FBlockedGrid3f::BlockSize * BlockIJK;
+				const FVector3i BlockMax( FMath::Min(BlockMin[0] + FBlockedGrid3f::BlockSize, GridDimensions[0]) -1,
+										  FMath::Min(BlockMin[1] + FBlockedGrid3f::BlockSize, GridDimensions[1]) -1,
+										  FMath::Min(BlockMin[2] + FBlockedGrid3f::BlockSize, GridDimensions[2]) -1);
+
+				for (const int32 TID : TrisToRasterize)
+				{ 
+
+					FVector3d xp = FVector3d::Zero(), xq = FVector3d::Zero(), xr = FVector3d::Zero();
+					Mesh->GetTriVertices(TID, xp, xq, xr);
+
+					// real IJK coordinates of xp/xq/xr
+
+					const FVector3d fp = (xp - org) * invdx;
+					const FVector3d fq = (xq - org) * invdx;
+					const FVector3d fr = (xr - org) * invdx;
+
+					// clamped integer bounding box of triangle plus exact-band
+					FTriBBox tri_bbox = TriBBoxes.GetData()[TID];
+
+					// clamp against the block
+					tri_bbox.Clamp(BlockMin, BlockMax);
+
+					// iter over the bbox region, min-compositing the distance from this tri into the grid. 
+					inclusive_loop_3d(tri_bbox.BoxMin, tri_bbox.BoxMax, 
+									[&](const FVector3i& ijk)
+									{
+										const FVector3i local_coords = ijk - BlockMin;
+										const int32 local_index = DistanceBlock.ToLinear(local_coords[0], local_coords[1], local_coords[2]);
+										const FVector3d cellpos((double)ijk[0], (double)ijk[1], (double)ijk[2]);
+										float d = DX * PointTriangleDistance(cellpos, fp, fq, fr);
+
+										float& grid_dist = DistanceBlock.At(local_index);
+										if (d < grid_dist)
+										{
+											grid_dist = d;
+											ClosestTriBlock.At(local_index) = TID;
+										}
+									});
+				}
+
+			}, !this->bUseParallel);
+		};
+	}   // make_level_set3_parallel
+
+	// lockless method - uses a scatter/ gather pattern along with a spatial acceleration structure to find the closest triangle (relative to each narrow band cell)
+	void make_level_set3_parallel_spatial(FVector3f Origin, float DX, int32 NI, int32 NJ, int32 NK, FBlockedGrid3f& DistanceGrid, int32 ExactBand)
+	{
+		const float upper_bound_value = upper_bound(FVector3i(NI, NJ, NK), DX);
+
+		DistanceGrid.Reset(NI, NJ, NK, upper_bound_value);
+		ClosestTriGrid.Reset(NI, NJ, NK, -1);
+
+		// bboxes for each tri.
+		const TArray<FTriBBox> TriBBoxes = compute_tri_bboxes(Origin, DX, NI, NJ, NK, ExactBand);
+
+		// blocked grid of tri IDs - identifies the tris that overlap with each grid block.
+		const FTriIDBlockGrid TIDBlockGrid = scatter_tris(TriBBoxes, NI, NJ, NK);
+
+		// identify the blocks that are occupied by tris
+		const TArray<int32> OccupiedBlocks = TIDBlockGrid.GetOccupiedBlockIDs();
+
+
+		if (this->CancelF())
+		{
+			return;
+		}
+
+		// parallelize over occupied blocks.  
+		// within the block, 
+		//   1st, mark overlapped cells by scattering the tri bboxes identified with block
+		//   2nd, gather the closest distance value into each overlapped cell using the spatial lookup. 
+		
+		{
+			const double max_dist = ExactBand * (DX * FMathd::Sqrt2);
+		
+			// actually raster the tris bbox
+			const int32 NumOccupied = OccupiedBlocks.Num();
+			ParallelFor(NumOccupied, [this, &Origin, &OccupiedBlocks, &TIDBlockGrid, &DistanceGrid, &TriBBoxes, DX, upper_bound_value,  max_dist](const int32& OccupiedId)
+				{
+					const FVector3i GridDimensions       = DistanceGrid.GetDimensions();
+					const int32 BlockIndex               = OccupiedBlocks.GetData()[OccupiedId];
+					const TArray<int32>& TrisToRasterize = TIDBlockGrid.TrisInBlock[BlockIndex];
+
+					const FVector3i BlockIJK    = DistanceGrid.BlockIndexToBlockIJK(BlockIndex);
+					const FVector3i BlockOrigin = FBlockedGrid3f::BlockData3Type::BlockSize * BlockIJK;
+					// allocates block on demand.  doing so here assumes the parallel memory allocator is up to it..
+					FBlockedGrid3f::BlockData3Type& DistanceBlock   = DistanceGrid.TouchBlockData(BlockIJK);
+					FBlockedGrid3i::BlockData3Type& ClosestTriBlock = ClosestTriGrid.TouchBlockData(BlockIJK);
+
+					FBlockedGrid3b::BlockData3Type BitBlock(false, -1);
+					// bounding box for this block [inclusive, inclusive]
+					const FVector3i BlockMin = FBlockedGrid3f::BlockSize * BlockIJK;
+					const FVector3i BlockMax( FMath::Min(BlockMin[0] + FBlockedGrid3f::BlockSize, GridDimensions[0]) - 1,
+						                      FMath::Min(BlockMin[1] + FBlockedGrid3f::BlockSize, GridDimensions[1]) - 1,
+						                      FMath::Min(BlockMin[2] + FBlockedGrid3f::BlockSize, GridDimensions[2]) - 1);
+
+					for (const int32 TID : TrisToRasterize)
+					{
+						FVector3d xp = FVector3d::Zero(), xq = FVector3d::Zero(), xr = FVector3d::Zero();
+						Mesh->GetTriVertices(TID, xp, xq, xr);
+
+						// clamped integer bounding box of triangle plus exact-band
+						FTriBBox tri_bbox = TriBBoxes.GetData()[TID];
+
+						// clamp against the block
+						tri_bbox.Clamp(BlockMin, BlockMax);
+						
+						inclusive_loop_3d(tri_bbox.BoxMin, tri_bbox.BoxMax, 
+										[&](const FVector3i& ijk)
+										{
+											const FVector3i local_coords = ijk - BlockMin;
+											const int32 local_index = DistanceBlock.ToLinear(local_coords[0], local_coords[1], local_coords[2]);
+											BitBlock.At(local_index) = true;
+										});
+					}
+
+					
+
+					for (int32 local_index = 0; local_index < FBlockedGrid3f::BlockData3Type::ElemCount; ++local_index)
+					{
+						if (BitBlock.At(local_index) == true)
+						{
+							const FVector3i GridCoord = DistanceBlock.ToLocalIJK(local_index) + BlockOrigin;
+
+							const FVector3d Pos((float)GridCoord.X* DX + Origin[0], (float)GridCoord.Y* DX + Origin[1], (float)GridCoord.Z* DX + Origin[2]);
+							double dsqr;
+							int near_tid = this->Spatial->FindNearestTriangle(Pos, dsqr, max_dist);
+							if (near_tid == IndexConstants::InvalidID)
+							{
+								DistanceBlock.At(local_index) = upper_bound_value; 
+							}
+							else
+							{
+								const float dist = (float)FMath::Sqrt(dsqr);
+								DistanceBlock.At(local_index) = dist;  
+								ClosestTriBlock.At(local_index) = near_tid;
+							}
+						}
+						
+					}
+
+				}, !this->bUseParallel);
+		}
+	}
+
+	// use the mesh vertex points as seeds, and flood fill distance until reaching the exact band width. this uses some thread-locking, but can be fast if the width is not large compared to block size
+	void make_level_set3_parallel_floodfill(FVector3f Origin, float DX, int32 NI, int32 NJ, int32 NK, FBlockedGrid3f& DistanceGrid, double NarrowBandWidth)
+	{
+	
+		typedef FBlockedGrid3f::BlockData3Type   FloatBlockData3Type;
+		typedef FBlockedGrid3b::BlockData3Type   BoolBlockData3Type;
+		typedef FBlockedGrid3i::BlockData3Type   IntBlockData3Type;
+		
+		typedef BoolBlockData3Type::BitArrayConstIterator  BitArrayConstIter;
+
+		const double ox = (double)Origin[0], oy = (double)Origin[1], oz = (double)Origin[2];
+		const double invdx = 1.0 / DX;
+		const float upper_bound_value = upper_bound(FVector3i(NI, NJ, NK), DX);
+
+		DistanceGrid.Reset(NI, NJ, NK, upper_bound_value);
+		ClosestTriGrid.Reset(NI, NJ, NK, -1);
+		
+		TArray<FCriticalSection> BlockLocks;  BlockLocks.SetNum(DistanceGrid.GetNumBlocks());
+		auto BlockLockProvider = [&BlockLocks](int32 block_id)->FCriticalSection&
+		{
+			return BlockLocks.GetData()[block_id];
+		};
+
+		// unsigned distance populate the distance grid and the closest tri grid using a limited flood fill, starting with
+		// the locations of the mesh verts.
+		{
+			FBlockedGrid3b VisitedGrid(NI, NJ, NK, false);
+			int32 CurCandidateID = 0;
+			FBlockedGrid3b CandidateGrids[2] = {FBlockedGrid3b(NI, NJ, NK, false), FBlockedGrid3b(NI, NJ, NK, false) };
+		
+			
+
+			// populate the CandidateGrid with all the mesh verts
+			{
+				
+				FBlockedGrid3b& CandidateGrid = CandidateGrids[CurCandidateID];
+				
+			
+				bool bCancel = false;
+				ParallelFor(Mesh->MaxVertexID(), [this, &bCancel, &CandidateGrid, &BlockLockProvider, NI, NJ, NK, Origin, DX](const int32 VID)
+				{
+					const double ox = (double)Origin[0], oy = (double)Origin[1], oz = (double)Origin[2];
+					const double invdx = 1.0 / DX;
+					
+					if (!Mesh->IsVertex(VID))
+					{
+						return;
+					}
+					if (VID % 1000 == 0)
+					{
+						bCancel = this->CancelF();
+					}
+					if (bCancel)
+					{
+						return;
+					}
+					const FVector3d v = Mesh->GetVertex(VID);
+					
+					// real IJK coordinates of v
+					const double  fi = (v.X - ox) * invdx, fj = (v.Y - oy) * invdx, fk = (v.Z - oz) * invdx;
+					
+					const FVector3i Idx( FMath::Clamp((int32)fi, 0, NI - 1),
+										 FMath::Clamp((int32)fj, 0, NJ - 1),
+										 FMath::Clamp((int32)fk, 0, NK - 1));
+
+					CandidateGrid.SetValueWithLock(Idx.X, Idx.Y, Idx.Z, true, BlockLockProvider);
+
+				}, !this->bUseParallel);
+
+				if (bCancel)
+				{
+					return;
+				}
+			}
+
+			const double max_dist = NarrowBandWidth;
+			const double max_query_dist = max_dist + (2.0 * DX * FMathd::Sqrt2);
+
+			int32 NumCandidateBlocks = CandidateGrids[CurCandidateID].GetAllocatedBlocks().Num();
+			while (NumCandidateBlocks > 0)
+			{
+				const FBlockedGrid3b& CandidateGrid = CandidateGrids[CurCandidateID];
+				FBlockedGrid3b& NextCandidateGrid = CandidateGrids[1-CurCandidateID];
+
+				VisitedGrid.PreAllocateFromSourceGrid(CandidateGrid);
+				DistanceGrid.PreAllocateFromSourceGrid(CandidateGrid);
+				ClosestTriGrid.PreAllocateFromSourceGrid(CandidateGrid);
+				NextCandidateGrid.Reset(NI, NJ, NK, false);
+
+				// for each block with candidate locations, generate seed points and flood fill within the block and capture fill that exits the block
+				// in the NextCandidateGrid.
+
+				TArray<const BoolBlockData3Type*> CandidateBlocks = CandidateGrid.GetAllocatedConstBlocks();
+
+				ParallelFor(NumCandidateBlocks, [this, max_query_dist, &BlockLockProvider, &VisitedGrid, &DistanceGrid, &NextCandidateGrid, &CandidateBlocks](const int32 cb)
+				{
+					const BoolBlockData3Type& CandidateBlock = *CandidateBlocks[cb];
+					const int32 BlockIdx = CandidateBlock.Id;
+
+					FloatBlockData3Type& DistanceBlock = *DistanceGrid.GetBlockData(BlockIdx);
+					IntBlockData3Type& ClosestTriBlock = *ClosestTriGrid.GetBlockData(BlockIdx);
+					BoolBlockData3Type& VistedBlock    = *VisitedGrid.GetBlockData(BlockIdx);
+
+					const FVector3i BlockIJK    = DistanceGrid.BlockIndexToBlockIJK(BlockIdx);
+					const FVector3i BlockOrigin = FloatBlockData3Type::BlockSize * BlockIJK;
+
+					// the candidates were scattered from neighboring blocks, but if this block has already been processed some 
+					// of them may have been visited.
+
+					// make a mask that is the subtraction of the visited cells from the candidate cells.
+					BoolBlockData3Type::BlockDataBitMask Mask = CandidateBlock.BitArray;
+					Mask.CombineWithBitwiseXOR(VistedBlock.BitArray, EBitwiseOperatorFlags::MaintainSize);
+					Mask.CombineWithBitwiseAND(CandidateBlock.BitArray, EBitwiseOperatorFlags::MaintainSize);
+
+					// inter over unvisited candidate cells, updating result grids and creating a list of seed points. 
+					TArray<FVector3i> SeedLocalCoords;
+					for (BitArrayConstIter CIter(Mask); CIter; ++CIter)
+					{
+						const int32 LocalIndex = CIter.GetIndex();
+						VistedBlock.BitArray[LocalIndex] = true;
+
+						const FVector3i LocalCoords = DistanceBlock.ToLocalIJK(LocalIndex);
+						const FVector3i GridCoords = LocalCoords + BlockOrigin;
+						const FVector3d Pos(cell_center(GridCoords));
+
+						double dsqr;
+						const int32 near_tid = this->Spatial->FindNearestTriangle(Pos, dsqr, max_query_dist);
+						if (near_tid != IndexConstants::InvalidID)
+						{
+							const float dist = (float)FMath::Sqrt(dsqr);
+							DistanceBlock.At(LocalIndex) = dist;
+							ClosestTriBlock.At(LocalIndex) = near_tid;
+
+							SeedLocalCoords.Add(LocalCoords);
+						}
+					}
+	
+					auto IsInBlock = [](const FVector3i& LocalCoords)
+										{
+											bool bInBlock = true;
+											bInBlock = LocalCoords.X > -1 && LocalCoords.Y > -1 && LocalCoords.Z > -1;
+											bInBlock = bInBlock && LocalCoords.X < FloatBlockData3Type::BlockSize&& LocalCoords.Y < FloatBlockData3Type::BlockSize && LocalCoords.Z < FloatBlockData3Type::BlockSize;
+											return bInBlock;
+										};
+
+					// flood fill from the seeds within the block.  if flood fill exits the block, mark the locations in the NextCanidateGrid for the next pass.
+					while(SeedLocalCoords.Num() > 0)
+					{
+						const FVector3i Seed = SeedLocalCoords.Pop(false);
+						for (const FVector3i& idx_offset : IndexUtil::GridOffsets26)
+						{
+							const FVector3i LocalCoords = Seed + idx_offset;
+							const FVector3i GridCoords = LocalCoords + BlockOrigin;
+							if (IsInBlock(LocalCoords))
+							{
+								const int32 local_index = FloatBlockData3Type::ToLinear(LocalCoords.X, LocalCoords.Y, LocalCoords.Z);
+								if (VistedBlock.BitArray[local_index] == false)
+								{
+									VistedBlock.BitArray[local_index] = true;
+
+									const FVector3d Pos(cell_center(GridCoords));
+								
+									double dsqr;
+									const int32 near_tid = this->Spatial->FindNearestTriangle(Pos, dsqr, max_query_dist);
+									if (near_tid != IndexConstants::InvalidID)
+									{
+										const float dist = (float)FMath::Sqrt(dsqr);
+										DistanceBlock.At(local_index) = dist; 
+										ClosestTriBlock.At(local_index) = near_tid; 
+
+										SeedLocalCoords.Add(LocalCoords);
+									}
+								
+								}
+							}
+							else // outside this block - mark the location as a candidate for the next pass.
+							{
+								if (NextCandidateGrid.IsValidIJK(GridCoords))
+								{
+									const bool bIsCanidate = true;
+									NextCandidateGrid.SetValueWithLock(GridCoords.X, GridCoords.Y, GridCoords.Z, bIsCanidate, BlockLockProvider);
+								}
+							}
+						}
+					}
+
+
+				},
+				!this->bUseParallel);
+
+				CurCandidateID = 1-CurCandidateID;
+				NumCandidateBlocks = CandidateGrids[CurCandidateID].GetAllocatedBlocks().Num();
+
+				if (CancelF())
+				{
+					return;
+				}
+			}
+
+		}
+	}   // end make_level_set_3
+
+
+	void CleanupUnwanted()
+	{
+		if (!bWantClosestTriGrid)
+		{
+			ClosestTriGrid.Reset();
+		}
+
+		if (!bWantIntersectionsGrid)
+		{
+			IntersectionsGrid.Reset();
+		}
+	}
+
+
+
+	// fill the intersection Grid w/ number of intersections in each cell
+	void compute_intersections(FVector3f Origin, float DX, int32 NI, int32 NJ, int32 NK, const FBlockedGrid3f& DistanceGrid, FBlockedGrid3i& IntersectionCountGrid)
+	{
+		
+		// grid of atomic<int>, pre-allocated to match the topology of the distance grid. 
+		struct FAtomicGrid3i : public TBlockedGrid3Layout<FBlockedGrid3f::BlockSize>
+		{
+			typedef TStaticArray<std::atomic<int32>, FBlockedGrid3f::BlockElemCount>  BlockAtomicData;
+			typedef TBlockedGrid3Layout<FBlockedGrid3f::BlockSize>  MyLayout;
+			
+			FAtomicGrid3i(const FBlockedGrid3f& Grid3f, const bool bThreadedBuild)
+			 : MyLayout(Grid3f.GetDimensions())
+			{
+				const int32 NumBlocks = Grid3f.GetNumBlocks();
+				Blocks.SetNum(NumBlocks);
+				// we use a parallel memory allocator, so we allow allocation to occur inside this parallel loop.
+				ParallelFor(NumBlocks, [this, &Grid3f](int32 i)
+				{
+					if (Grid3f.GetBlockData(i).IsValid())
+					{
+						Blocks[i] = MakeUnique<BlockAtomicData>();
+						BlockAtomicData& Data= *Blocks[i];
+						for (int32 j = 0; j < FBlockedGrid3i::BlockElemCount; ++j)
+						{
+							Data[j] = 0;
+						}
+					}
+				}, !bThreadedBuild);
+			}
+			TArray<TUniquePtr<BlockAtomicData>> Blocks;
+		} AtomicGrid3i(DistanceGrid, this->bUseParallel);
+
+		bool bCancel= false;
+
+		// this is what we will do for each triangle. There are no Grid-reads, only Grid-writes, 
+		// since we use atomic_increment, it is always thread-safe
+		ParallelFor(Mesh->MaxTriangleID(),
+			[this, Origin, DX, NI, NJ, NK, &IntersectionCountGrid, &AtomicGrid3i, &bCancel](int32 TID)
+			{
+				double ox = (double)Origin[0], oy = (double)Origin[1], oz = (double)Origin[2];
+				double invdx = 1.0 / DX;
+
+				if (!Mesh->IsTriangle(TID))
+				{
+					return;
+				}
+				if (TID % 100 == 0 && CancelF() == true)
+				{
+					bCancel = true;
+				}
+				if (bCancel)
+				{
+					return;
+				}
+
+				FVector3d xp = FVector3d::Zero(), xq = FVector3d::Zero(), xr = FVector3d::Zero();
+				Mesh->GetTriVertices(TID, xp, xq, xr);
+
+
+				bool neg_x = false;
+				if (InsideMode == EInsideModes::WindingCount)
+				{
+					FVector3d n = VectorUtil::NormalDirection(xp, xq, xr);
+					neg_x = n.X > 0;
+				}
+
+				// real IJK coordinates of xp/xq/xr
+				double fip = (xp[0] - ox) * invdx, fjp = (xp[1] - oy) * invdx, fkp = (xp[2] - oz) * invdx;
+				double fiq = (xq[0] - ox) * invdx, fjq = (xq[1] - oy) * invdx, fkq = (xq[2] - oz) * invdx;
+				double fir = (xr[0] - ox) * invdx, fjr = (xr[1] - oy) * invdx, fkr = (xr[2] - oz) * invdx;
+
+				// recompute J/K integer bounds of triangle w/o exact band
+				int32 j0 = FMath::Clamp(FMath::CeilToInt(FMath::Min3(fjp, fjq, fjr)),  0, NJ - 1);
+				int32 j1 = FMath::Clamp(FMath::FloorToInt(FMath::Max3(fjp, fjq, fjr)), 0, NJ - 1);
+				int32 k0 = FMath::Clamp(FMath::CeilToInt(FMath::Min3(fkp, fkq, fkr)),  0, NK - 1);
+				int32 k1 = FMath::Clamp(FMath::FloorToInt(FMath::Max3(fkp, fkq, fkr)), 0, NK - 1);
+
+				
+				auto AtomicIncDec = [&AtomicGrid3i, neg_x](const int32 i, const int32 j, const int32 k)
+										{
+											int32 BlockIndex, LocalIndex;
+											AtomicGrid3i.GetBlockAndLocalIndex(i, j, k, BlockIndex, LocalIndex);
+											
+											// note the grid has been pre-allocated, so the blocks will exist
+											auto& ABlock = *AtomicGrid3i.Blocks[BlockIndex];
+											std::atomic<int32>& value = ABlock[LocalIndex];
+											if (neg_x)
+											{
+												value.fetch_sub(1, std::memory_order_relaxed);
+											}
+											else
+											{
+												value.fetch_add(1, std::memory_order_relaxed);
+											}
+										};
+
+				// and do intersection counts
+				for (int K = k0; K <= k1; ++K)
+				{
+					for (int J = j0; J <= j1; ++J)
+					{
+						double A, B, C;
+						if (PointInTriangle2d(J, K, fjp, fkp, fjq, fkq, fjr, fkr, A, B, C))
+						{
+							double fi = A * fip + B * fiq + C * fir; // intersection I coordinate
+							int i_interval = FMath::CeilToInt(fi); // intersection is in (i_interval-1,i_interval]
+							if (i_interval < 0)
+							{
+								AtomicIncDec(0, J, K);
+							}
+							else if (i_interval < NI)
+							{
+								AtomicIncDec(i_interval, J, K);
+							}
+							else
+							{
+								// we ignore intersections that are beyond the +x side of the Grid
+							}
+						}
+					}
+				}
+			}, !this->bUseParallel);
+
+			IntersectionCountGrid.Reset(NI, NJ, NK, 0);
+			// copy result back to the intersection grid
+			const int32 NumBlocks = AtomicGrid3i.Blocks.Num();
+			ParallelFor(NumBlocks, [&IntersectionCountGrid, &AtomicGrid3i](const int32 blockid)
+			{
+				TUniquePtr<FAtomicGrid3i::BlockAtomicData>& AtomicBlockPtr = AtomicGrid3i.Blocks[blockid];
+				
+				if (AtomicBlockPtr.IsValid())
+				{
+					FVector3i BlockIJK = AtomicGrid3i.BlockIndexToBlockIJK(blockid);
+					FBlockedGrid3i::BlockData3Type& BlockData = IntersectionCountGrid.TouchBlockData(BlockIJK);	
+					const FAtomicGrid3i::BlockAtomicData& AtomicBlockData = *AtomicBlockPtr;
+
+					for (int32 j = 0; j < FBlockedGrid3i::BlockElemCount; ++j)
+					{
+						const int32 Count = AtomicBlockData[j];
+						BlockData.At(j) = Count;
+					}
+
+					// delete the atomic block after copying the data.
+					AtomicBlockPtr.Release();
+				}
+			
+			}, !this->bUseParallel);
+	}
+
+
+	// iterate through each x-row of Grid and set unsigned distances to be negative (inside the mesh), based on the IntersectionCount
+	void compute_signs(int32 NI, int32 NJ, int32 NK, const FBlockedGrid3i& IntersectionCountGrid, FBlockedGrid3f& Distances)
+	{
+		if (bUseParallel)
+		{
+			bool bCancel = false;
+
+			const FVector3i BlockDims = Distances.GetBlockDimensions();
+			// can process each x block-row in parallel
+			ParallelFor(BlockDims.Y * BlockDims.Z, [this, BlockDims, &Distances, &IntersectionCountGrid, &bCancel](int32 BlockRowID)
+			{
+				const FVector3i GridDims = Distances.GetDimensions();
+				if (BlockRowID % 10 == 0 && CancelF() == true)
+				{
+					bCancel = true;
+				}
+
+				if (bCancel)
+				{
+					return;
+				}
+
+				constexpr int32 BlockSize = FBlockedGrid3f::BlockSize;
+				TStaticArray<int32, BlockSize*BlockSize> BlockCrossSection;
+				for (int32 i = 0, I = BlockSize * BlockSize; i < I; ++i) 
+				{ 
+					BlockCrossSection[i] = 0;
+				}
+
+				// 2d array linearized as BlockRowID = BlockJ + BlockDims[1]*BlockK
+				const int32 BlockJ = BlockRowID % BlockDims[1];
+				const int32 BlockK = BlockRowID / BlockDims[1];
+				FVector3i BlockIJK(0, BlockJ, BlockK);
+				for (  BlockIJK[0] = 0; BlockIJK[0] < BlockDims[0]; ++BlockIJK[0]) {
+
+					if (Distances.IsBlockAllocated(BlockIJK))
+					{
+						checkSlow(IntersectionCountGrid.IsBlockAllocated(BlockIJK));
+
+						FBlockedGrid3f::BlockData3Type& DistanceBlockData = Distances.TouchBlockData(BlockIJK);
+						const FBlockedGrid3i::BlockData3Type& IntersectionCountBlockData = *IntersectionCountGrid.GetBlockData(BlockIJK);
+
+						// note, some data values stored in the blocks are formally outside the dimensions of the grids.. 
+						// so this loop potentially touches some regions that will never be accessed. 
+						for (int32 k = 0; k < BlockSize; ++k)
+						{
+							for (int32 j = 0; j < BlockSize; ++j)
+							{
+								int32& total_count = BlockCrossSection[j + k * BlockSize];
+
+								for (int32 i = 0; i < BlockSize; ++i)
+								{
+									const int32 LocalIndex = TBlockData3Layout<BlockSize>::ToLinear(i, j, k);
+
+									total_count += IntersectionCountBlockData.At(LocalIndex);
+
+									if (
+										(InsideMode == EInsideModes::WindingCount && total_count > 0) ||
+										(InsideMode == EInsideModes::CrossingCount && total_count % 2 == 1)
+										)
+									{
+										// we are inside the mesh
+										float& d = DistanceBlockData.At(LocalIndex);
+										d = -FMath::Abs(d);
+									}
+
+								}
+							}
+						}
+							
+					}
+					else
+					{
+						checkSlow(!IntersectionCountGrid.IsBlockAllocated(BlockIJK));
+						// the block isn't allocated - so it "should" all share the same state (inside or outside) for water-tight meshes,
+						// but the counting method used for sign assignment may produce odd results for non water-tight meshes where a block
+						// far from the narrow band may have both inside and outside elements..  in this case we just assign the block
+						// the most common sign.
+						int32 NumNeg = 0;
+						for (int32 index = 0, IndexMax = BlockSize * BlockSize; index < IndexMax; ++index)
+						{
+							const int32 total_count = BlockCrossSection[index];
+							if (
+								(InsideMode == EInsideModes::WindingCount && total_count > 0) ||
+								(InsideMode == EInsideModes::CrossingCount && total_count % 2 == 1)
+								)
+							{
+								NumNeg++;
+							}
+						}
+						if ( float(NumNeg)/float(BlockSize * BlockSize) > 0.5f )
+						{
+							// we are inside the mesh
+							Distances.ProcessBlockDefaultValue( BlockIJK, [](float& value){value = -FMath::Abs(value);});
+						}
+					}
+						
+				}
+			}, !this->bUseParallel);
+
+		}
+		else
+		{
+			for (int K = 0; K < NK; ++K)
+			{
+				if (CancelF())
+				{
+					return;
+				}
+
+				for (int J = 0; J < NJ; ++J)
+				{
+					int total_count = 0;
+					for (int I = 0; I < NI; ++I)
+					{
+						total_count += IntersectionCountGrid.GetValue(I, J, K);
+						if (
+							(InsideMode == EInsideModes::WindingCount && total_count > 0) ||
+							(InsideMode == EInsideModes::CrossingCount && total_count % 2 == 1)
+							)
+						{
+							// we are inside the mesh
+							Distances.ProcessValue(I, J, K, [](float& value){value = -value;});
+						}
+					}
+				}
+			}
+		}
+	}
+};
+
+
+} // end namespace UE::Geometry
+} // end namespace UE
\ No newline at end of file
diff --git a/Engine/Source/Runtime/GeometryCore/Public/Implicit/SweepingMeshSDF.h b/Engine/Source/Runtime/GeometryCore/Public/Implicit/SweepingMeshSDF.h
--- a/Engine/Source/Runtime/GeometryCore/Public/Implicit/SweepingMeshSDF.h
+++ b/Engine/Source/Runtime/GeometryCore/Public/Implicit/SweepingMeshSDF.h
@@ -274,12 +274,12 @@ public:
 
 	float At(int I, int J, int K) const
 	{
-		return Grid.At(I, J, K);
+		return Grid.GetValue(I, J, K);
 	}
 
 	float operator[](const FVector3i& Idx) const
 	{
-		return Grid[Idx];
+		return Grid.GetValue(Idx);
 	}
 
 	FVector3f CellCenter(int I, int J, int K) const
@@ -289,7 +289,7 @@ public:
 
 	float GetValue(FVector3i Idx) const // TTriLinearGridInterpolant interface 
 	{
-		return Grid[Idx];
+		return Grid.GetValue(Idx);
 	}
 
 private:
diff --git a/Engine/Source/Runtime/GeometryCore/Public/Spatial/BlockedDenseGrid3.h b/Engine/Source/Runtime/GeometryCore/Public/Spatial/BlockedDenseGrid3.h
--- a/Engine/Source/Runtime/GeometryCore/Public/Spatial/BlockedDenseGrid3.h
+++ b/Engine/Source/Runtime/GeometryCore/Public/Spatial/BlockedDenseGrid3.h
@@ -5,258 +5,907 @@
 #pragma once
 
 #include "CoreMinimal.h"
+#include "BlockedLayout3.h"
 #include "BoxTypes.h"
-#include "IntBoxTypes.h"
+#include "Containers/BitArray.h"
 #include "Containers/StaticArray.h"
+#include "Misc/ScopeLock.h"
+#include "HAL/CriticalSection.h"
+
 
 namespace UE
 {
 namespace Geometry
 {
 
-/**
- * TBlockedDenseGrid represents a dense 3D uniform grid, but the grid is allocated in BlockSize^3 blocks on-demand
- * (BlockSize is a compile-time constant). This allows very large grids to be used without
- * having to pre-allocate all the memory, eg for sparse/narrow-band use cases.
- * 
- * Thread-safe access functions are provided, these are implemented by internal locking on the grid-block level.
- */
-template<typename ElemType>
-class TBlockedDenseGrid3
+// Block Data Buffer, holds linear array of typed data
+template <typename ElemType_, int32 BlockSize_>
+class TBlockData3 : public TBlockData3Layout<BlockSize_>
 {
-protected:
+public:
+	typedef TBlockData3Layout<BlockSize_>                        DataLayout;
+	typedef ElemType_                                            ElemType;
+	typedef TStaticArray<ElemType, DataLayout::ElemCount>        LinearDataStorageType;
 
-	static const int32 BlockSize = 32;
-	using BlockType = TStaticArray<ElemType, BlockSize*BlockSize*BlockSize>;
+	TBlockData3()
+	:Id(-1)
+	{}
 
-	TArray<TUniquePtr<BlockType>> Blocks;
-	TArray<FCriticalSection> BlockLocks;
+	TBlockData3(const ElemType& Value, const int32 ID)
+	:Id(ID)
+	{
+		Reset(Value);
+	}
 
-	/** dimensions per axis */
-	FVector3i BlockDimensions;
+	/** initialize all the data elements to the specified value*/
+	void Reset(const ElemType& Value)
+	{
+		for (int32 k = 0, K = DataLayout::ElemCount; k < K; ++k)
+		{
+			Data[k] = Value;
+		}
+	}
 
-	/** dimensions per axis */
-	FVector3i Dimensions;
+	/** access the data by index into the underlying linear array*/
+	ElemType& At(const int32 Index)
+	{
+		checkSlow(Index < DataLayout::ElemCount);
+		return Data[Index];
+	}
+	/** access the data by index into the underlying linear array*/
+	ElemType At(const int32 Index) const 
+	{
+		checkSlow(Index < DataLayout::ElemCount);
+		return Data[Index];
+	}
 
-	/** The value to use when no value is known - ie for Get() queries on unallocated cells, or to initialize new blocks */
-	ElemType ConstantValue = (ElemType)0;
+	template <typename FuncType>
+	void ModifyValue(const int32 Index, FuncType Func)
+	{
+		ElemType& Value = Data[Index];
+		Func(Value);
+	}
 
-protected:
+	int32                        Id;          // an id associated with this block data.
+	LinearDataStorageType        Data;        // payload
 
-	void Init(BlockType& Block)
+};
+
+// Specialize the Block Data Buffer for bools by using a bitarray. 
+template <int32 BlockSize_>
+class TBlockData3<bool, BlockSize_> : public TBlockData3Layout<BlockSize_>
+{
+public:
+	typedef TBlockData3Layout<BlockSize_>                       DataLayout;
+	typedef bool                                                ElemType;
+	typedef TInlineAllocator<DataLayout::ElemCount>             AllocatorType;
+	typedef TBitArray<AllocatorType>                            BlockDataBitMask;
+	typedef TConstSetBitIterator<AllocatorType>                 BitArrayConstIterator;
+
+	TBlockData3(const ElemType& Value, const int32 ID)
+		:Id(ID)
 	{
-		for (int32 k = 0; k < BlockSize * BlockSize * BlockSize; ++k)
-		{
-			Block[k] = ConstantValue;
-		}
+		Reset(Value);
 	}
 
-	void GetBlockCoords(int32 X, int32 Y, int32 Z, int32& BlockIndexOut, int32& LocalIndexOut) const
+	void Reset(const ElemType& Value)
 	{
-		int32 BlockX = X / BlockSize;
-		int32 BlockY = Y / BlockSize;
-		int32 BlockZ = Z / BlockSize;
-		BlockIndexOut = BlockX + BlockDimensions.X * (BlockY + BlockDimensions.Y * BlockZ);
+		BitArray.Init(Value, DataLayout::ElemCount);
+	}
 
-		int32 LocalX = X % BlockSize;
-		int32 LocalY = Y % BlockSize;
-		int32 LocalZ = Z % BlockSize;
-		LocalIndexOut = LocalX + BlockSize * (LocalY + BlockSize * LocalZ);
+	FBitReference At(int32 LocalIndex)
+	{
+		checkSlow(LocalIndex < DataLayout::ElemCount);
+		return BitArray[LocalIndex];
+	}
+
+	bool At(int32 LocalIndex) const
+	{
+		checkSlow(LocalIndex < DataLayout::ElemCount);
+		return BitArray[LocalIndex];
 	}
 
+	template <typename FuncType>
+	void ModifyValue(const int32 Index, FuncType Func)
+	{
+		bool Value = BitArray[Index];
+		Func(Value);
+		BitArray[Index] = Value;
+	}
+
+	// bit field operations can correspond to topological operations
+	void TopologyUnion(const TBlockData3<bool, BlockSize_>& OtherBlockData)
+	{
+		const auto& OtherBitArray = OtherBlockData.BitArray;
+		BitArray.CombineWithBitwiseOR(OtherBitArray, EBitwiseOperatorFlags::MaintainSize);
+	}
+
+	int32                        Id;
+	BlockDataBitMask             BitArray;    
+};
+
+
+/**
+* TBasicBlockedDenseGrid3 represents a 3D blocked uniform grid.
+*
+* The grid is allocated in BlockSize^3 blocks on-demand (BlockSize is a compile-time constant) allowing
+* very large grids to be used without having to pre-allocate all the memory, eg for sparse/narrow-band use cases.
+* 
+* Block-level default values are returned when querying a grid cell in an unallocated block.
+*
+* For multi-threaded applications consider using the derived classes TBlockedDenseGrid3 which manages internal locks,
+* or TBlockedGrid3 with your own lock provider.
+*/
+template<typename ElemType, int32 BlockSize_>
+class TBasicBlockedDenseGrid3 : public TBlockedGrid3Layout<BlockSize_>
+{
+public:
+	typedef TBlockedGrid3Layout<BlockSize_>                                   BlockedGrid3LayoutType;
+	typedef TBlockData3<ElemType, BlockSize_>                                 BlockData3Type;
+
+
+	static constexpr int32 BlockSize      = BlockSize_;                         // length of side of a block
+	static constexpr int32 BlockElemCount = BlockSize * BlockSize * BlockSize;  // number of cells in each block
+protected:
+
+	/** The value to use to initialize new blocks */
+	ElemType DefaultValue = (ElemType)0;
+
+
+	// a default value and a pointer to a block that will be allocated on first write.
+	struct FBlock3
+	{
+		// release any allocated block data and reset the default value.
+		void Reset(const ElemType& DefaultIn = (ElemType)0)
+		{
+			BlockDataPtr.Reset();
+			DefaultValue = DefaultIn;
+		}
+
+		/** @return the default value if the block has not been allocated, otherwise the indexed value in the block*/
+		ElemType  GetValue(const int32& LocalIndex) const 
+		{
+			return (BlockDataPtr.IsValid()) ? BlockDataPtr->At(LocalIndex) : DefaultValue;
+		}
+
+		TUniquePtr<BlockData3Type>  BlockDataPtr;     // pointer to 'heavy' data
+		ElemType                    DefaultValue;     // value to be used when BlockData is not allocated
+	};
+
+
+	TArray<FBlock3>          Blocks;                 // array of blocks that span the dimensions of the gird.
+	
+
+protected:
+
+	// declare other templated instantiations as friends so that PreAllocateFromSourceGrid() can work
+	template<class OtherElemType, int32 OtherBlockSize> friend class TBasicBlockedDenseGrid3;
+
+	//friend class TBlockedDenseBitGrid3<BlockSize>;
+
 #if UE_BUILD_DEBUG
-	TUniquePtr<BlockType>& GetBlock(int32 Index) { return Blocks[Index]; }
-	const TUniquePtr<BlockType>& GetBlock(int32 Index) const { return Blocks[Index]; }
-	FCriticalSection* GetBlockLock(int32 Index) { return &BlockLocks[Index]; }
+	FBlock3& GetBlock(int32 Index) { return Blocks[Index]; }
+	const FBlock3& GetBlock(int32 Index) const { return Blocks[Index]; }
 #else
 	// skip range checks in non-debug builds
-	TUniquePtr<BlockType>& GetBlock(int32 Index) { return Blocks.GetData()[Index]; }
-	const TUniquePtr<BlockType>& GetBlock(int32 Index) const { return Blocks.GetData()[Index]; }
-	FCriticalSection* GetBlockLock(int32 Index) { return &BlockLocks.GetData()[Index]; }
+	FBlock3& GetBlock(int32 Index) { return Blocks.GetData()[Index]; }
+	const FBlock3& GetBlock(int32 Index) const { return Blocks.GetData()[Index]; }
 #endif
 
+	/** @return true if Block has an allocated block data buffer*/
+	bool IsBlockAllocated(const FBlock3& Block) const
+	{
+		return Block.BlockDataPtr.IsValid();
+	}
+	/** @return true if the specified Block has an allocated block data buffer*/
+	bool IsBlockAllocated(const int32 BlockIndex) const
+	{
+		const FBlock3& Block = GetBlock(BlockIndex);
+		return IsBlockAllocated(Block);
+	}
+
+	/** return a reference to the block data buffer associated with the specified block, allocating if needed */
+	BlockData3Type& TouchBlockData(const int32 BlockIndex)
+	{
+		FBlock3& Block = GetBlock(BlockIndex);
+		if (!IsBlockAllocated(Block))
+		{
+			const ElemType& Value = Block.DefaultValue;
+			Block.BlockDataPtr = MakeUnique<BlockData3Type>(Value, BlockIndex);
+		}
+		return *Block.BlockDataPtr;
+	}
+
+	// --- write methods
+	template <typename FuncType>
+	void WriteValue(const int32& BlockIndex, const int32& LocalIndex, FuncType Func)
+	{
+		BlockData3Type& BlockData = TouchBlockData(BlockIndex);
+		BlockData.ModifyValue(LocalIndex, Func);
+	}
+
 	template<typename FuncType>
-	void WriteValueThreadSafe(int32 X, int32 Y, int32 Z, FuncType Func)
+	void WriteValue(int32 I, int32 J, int32 K, FuncType Func)
 	{
 		int32 BlockIndex, LocalIndex;
-		GetBlockCoords(X, Y, Z, BlockIndex, LocalIndex);
+		BlockedGrid3LayoutType::GetBlockAndLocalIndex(I, J, K, BlockIndex, LocalIndex);
+		WriteValue(BlockIndex, LocalIndex, Func);
+	}
+	
+	template <typename FuncType>
+	void WriteValueWithLock(const int32& BlockIndex, const int32& LocalIndex,  FuncType Func, FCriticalSection* CriticalSection)
+	{
+		FScopeLock Lock(CriticalSection);
+		WriteValue(BlockIndex, LocalIndex, Func);
+	}
+
+	// --- read methods
+
+	ElemType ReadValue(const int32& BlockIndex, const int32& LocalIndex) const 
+	{
+		return GetBlock(BlockIndex).GetValue(LocalIndex);
+	}
+
+	ElemType ReadValue(int32 I, int32 J, int32 K) const
+	{
+		int32 BlockIndex, LocalIndex;
+		BlockedGrid3LayoutType::GetBlockAndLocalIndex(I, J, K, BlockIndex, LocalIndex);
+		return ReadValue(BlockIndex, LocalIndex);
+	}
+
+	ElemType ReadValueWithLock(const int32& BlockIndex, const int32& LocalIndex, FCriticalSection* CriticalSection)
+	{
+		FScopeLock Lock(CriticalSection);
+		return ReadValue(BlockIndex, LocalIndex);
+	}
+
+
+
+public:
+	/**
+	* Create empty grid
+	*/
+	TBasicBlockedDenseGrid3() : BlockedGrid3LayoutType(0,0,0)
+	{
+	}
+
+	/**
+	* Create grid with specified domain.
+	*/
+	TBasicBlockedDenseGrid3(int32 DimI, int32 DimJ, int32 DimK, ElemType InitialValue)
+	{
+		DefaultValue = InitialValue;
+		Resize(DimI, DimJ, DimK);
+	}
 
-		FScopeLock Lock(GetBlockLock(BlockIndex));
+	/**
+	* Reconfigure the grid to have the target dimensions. This clears all the 
+	* existing grid memory and resets the block-level default value.
+	*/
+	void Resize(int32 DimI, int32 DimJ, int32 DimK)
+	{		
+		BlockedGrid3LayoutType::Resize(DimI, DimJ, DimK);
+		const int32 NumBlocks = GetNumBlocks();
 
-		if (GetBlock(BlockIndex).IsValid() == false)
+		Blocks.Reset();
+		Blocks.SetNum(NumBlocks);
+		for (int32 k = 0; k < NumBlocks; ++k)
 		{
-			Blocks[BlockIndex] = MakeUnique<BlockType>();
-			Init(*Blocks[BlockIndex]);
+			Blocks[k].Reset(DefaultValue);
 		}
-		TUniquePtr<BlockType>& Block = GetBlock(BlockIndex);
-		ElemType& Value = (*Block)[LocalIndex];
-		Func(Value);
 	}
 
+	/**
+	* Discard all allocated blocks but retain grid dimensions. This also resets the block-level default values.
+	*/
+	void Reset()
+	{
+		Reset(BlockedGrid3LayoutType::Dimensions.X, BlockedGrid3LayoutType::Dimensions.Y, BlockedGrid3LayoutType::Dimensions.Z, DefaultValue);
+	}
 
-	template<typename FuncType>
-	void WriteValue(int32 X, int32 Y, int32 Z, FuncType Func)
+	/**
+	* Reconfigure the grid to have the target dimensions and default value. 
+	* This clears all the existing grid memory, and resets default values to the specified InitialValue.
+	*/
+	void Reset(int32 DimI, int32 DimJ, int32 DimK, ElemType InitialValue)
+	{
+		DefaultValue = InitialValue;
+		Resize(DimI, DimJ, DimK);
+	}
+
+	/**
+	* @return the grid value at (I,J,K)
+	*/
+	ElemType GetValue(int32 I, int32 J, int32 K) const
+	{
+		return ReadValue(I, J, K);
+	}
+
+	/**
+	* @return the grid value at CellIJK.  Note this is safe to call from multiple read threads (but not if simultaneously writing!)
+	*/
+	ElemType GetValue(const FVector3i& CellIJK) const
+	{
+		return ReadValue(CellIJK.X, CellIJK.Y, CellIJK.Z);
+	}
+
+
+	/**
+	* Set the grid value at (I,J,K)
+	*/
+	void SetValue(int32 I, int32 J, int32 K, ElemType NewValue)
+	{
+		WriteValue(I, J, K, [NewValue](ElemType& GridValueInOut) { GridValueInOut = NewValue; });
+	}
+
+	/**
+	* Set the grid value at CellIJK
+	*/
+	void SetValue(const FVector3i& CellIJK, ElemType NewValue)
+	{
+		WriteValue(CellIJK.X, CellIJK.Y, CellIJK.Z, [NewValue](ElemType& GridValueInOut) { GridValueInOut = NewValue; });
+	}
+
+
+	/**
+	* Call an external lambda with a reference to the grid value at (I,J,K).
+	* Called as Func(ElemType&), so the caller can both read and write the grid cell.  
+	* 
+	* Note: if doing pure reads prefer the GetValue() methods because the ProcessValue method will allocate the underlying data block. 
+	*/
+	template<typename ProcessFunc>
+	void ProcessValue(int32 I, int32 J, int32 K, ProcessFunc Func)
+	{
+		WriteValue(I, J, K, Func);
+	}
+
+	/**
+	* Call an external lambda with a reference to the grid value at CellIJK.
+	* Called as Func(ElemType&), so the caller can both read and write the grid cell
+	* 
+	* Note: if doing pure reads prefer the GetValue() methods because the ProcessValue method will allocate the underlying data block. 
+	*/
+	template<typename ProcessFunc>
+	void ProcessValue(const FVector3i& CellIJK, ProcessFunc Func)
+	{
+		WriteValue(CellIJK.X, CellIJK.Y, CellIJK.Z, Func);
+	}
+
+
+	/**
+	* Pre-allocate the grid blocks that also exist in SourceGrid of the same dimensions.
+	* This is useful in grid-processing algorithms where parallel grid operations can be
+	* done without locking if the blocks are already allocated.
+	* 
+	* @return false on failure (e.g. different sized grids), otherwise true.
+	*/
+	template <typename OtherElemType>
+	bool PreAllocateFromSourceGrid(const TBasicBlockedDenseGrid3<OtherElemType, BlockSize>& SourceGrid)
+	{
+		return AllocateTopologyUnionImpl(SourceGrid);
+	}
+
+	/** @return the total number of block units form the spatial decomposition of the grid.  Note these block may not all be allocated*/
+	int32 GetNumBlocks() const
+	{
+		const FVector3i& BlockDims = BlockedGrid3LayoutType::GetBlockDimensions();
+		return BlockDims.X * BlockDims.Y * BlockDims.Z;
+	}
+
+protected:
+	template <typename OtherGridType >
+	bool AllocateTopologyUnionImpl(const OtherGridType& OtherGrid)
 	{
-		int32 BlockIndex, LocalIndex;
-		GetBlockCoords(X, Y, Z, BlockIndex, LocalIndex);
 
-		if (GetBlock(BlockIndex).IsValid() == false)
+		// bound checking
+		if ((BlockedGrid3LayoutType::GetDimensions() != OtherGrid.GetDimensions()) || (BlockedGrid3LayoutType::GetBlockDimensions() != OtherGrid.GetBlockDimensions()))
 		{
-			Blocks[BlockIndex] = MakeUnique<BlockType>();
-			Init(*Blocks[BlockIndex]);
+			check(false);		// this is almost certainly a development-time error in calling code
+			return false;
 		}
-		TUniquePtr<BlockType>& Block = GetBlock(BlockIndex);
-		ElemType& Value = (*Block)[LocalIndex];
-		Func(Value);
+
+		const int32 NumBlocks = GetNumBlocks();
+		for (int32 k = 0; k < NumBlocks; ++k)
+		{
+			if (OtherGrid.IsBlockAllocated(k))
+			{
+				TouchBlockData(k);
+			}
+		}
+		return true;
 	}
+};
+
+/**
+* TBlockedGrid3 represents a 3D uniform grid.
+*
+* The grid is allocated in BlockSize^3 blocks on-demand (BlockSize is a compile-time constant) allowing
+* very large grids to be used without having to pre-allocate all the memory, eg for sparse/narrow-band use cases.
+*
+* Block-level default values are returned when querying a grid cell in an unallocated block., 
+* 
+* Block-level access is provided to facilitate parallel algorithms.  Additionally, 'Set' methods are provided that utilize 
+* critical sections provided by the caller.  
+*/
+template<typename ElemType, int32 BlockSize_>
+class TBlockedGrid3 : public TBasicBlockedDenseGrid3<ElemType, BlockSize_>
+{
+public:
+	typedef TBasicBlockedDenseGrid3<ElemType, BlockSize_>                     BlockedDenseGridType;
+	typedef TBlockedGrid3Layout<BlockSize_>                                   BlockedGrid3LayoutType;
+	typedef TBlockData3<ElemType, BlockSize_>                                 BlockData3Type;
+
+
+	static constexpr int32 BlockSize      = BlockSize_;                         // length of side of a block
+	static constexpr int32 BlockElemCount = BlockSize * BlockSize * BlockSize;  // number of cells in each block
+
+
+protected:
+
+	typedef typename TBasicBlockedDenseGrid3<ElemType, BlockSize_>::FBlock3           FBlock3;
 
 
-	ElemType ReadValueThreadSafe(int32 X, int32 Y, int32 Z)
+	template<typename FuncType>
+	void WriteBlockDefaultValue(const int32 BlockIndex, FuncType Func, bool bDeallocateBock = true)
 	{
-		int32 BlockIndex, LocalIndex;
-		GetBlockCoords(X, Y, Z, BlockIndex, LocalIndex);
+		FBlock3& Block = BlockedDenseGridType::GetBlock(BlockIndex);
+		Func(Block.DefaultValue);
+		if (bDeallocateBock)
+		{
+			Block.BlockDataPtr.Reset();
+		}
+	}
 
-		FScopeLock Lock(GetBlockLock(BlockIndex));
-		const TUniquePtr<BlockType>& Block = GetBlock(BlockIndex);
-		return (Block.IsValid()) ? (*Block)[LocalIndex] : ConstantValue;
+	template <typename FuncType>
+	void WriteBlockDefaultValue(const FVector3i& BlockIJK, FuncType Func, bool bDeallocateBock = true)
+	{
+		const int32 BlockIndex = BlockedGrid3LayoutType::BlockIJKToBlockIndex(BlockIJK);
+		WriteBlockDefaultValue(BlockIndex, Func, bDeallocateBock);
 	}
 
-	ElemType ReadValue(int32 X, int32 Y, int32 Z) const
+	template<typename FuncType>
+	void WriteBlockDefaultValueWithLock(const int32 BlockIndex, FuncType Func, FCriticalSection* CriticalSection, bool bDeallocateBock = true)
 	{
-		int32 BlockIndex, LocalIndex;
-		GetBlockCoords(X, Y, Z, BlockIndex, LocalIndex);
+		FScopeLock Lock(CriticalSection);
+		WriteBlockDefaultValue(BlockIndex, Func, bDeallocateBock);
+	}
+
+
+	// -- block-level default value read methods.
 
-		const TUniquePtr<BlockType>& Block = GetBlock(BlockIndex);
-		return (Block.IsValid()) ? (*Block)[LocalIndex] : ConstantValue;
+	ElemType ReadBlockDefaultValue(const int32& BlockIndex) const
+	{
+		return BlockedDenseGridType::GetBlock(BlockIndex).DefaultValue;
 	}
 
+	ElemType ReadBlockDefaultValue(const FVector3i& BlockIJK) const
+	{
+		const int32 BlockIndex = BlockedGrid3LayoutType::BlockIJKToBlockIndex(BlockIJK);
+		return ReadBlockDefaultValue(BlockIndex);
+	}
 
 public:
 	/**
-	 * Create empty grid
-	 */
-	TBlockedDenseGrid3() : Dimensions(0,0,0)
+	* Create empty grid
+	*/
+	TBlockedGrid3() :  BlockedDenseGridType()
 	{
 	}
 
-	TBlockedDenseGrid3(int32 DimX, int32 DimY, int32 DimZ, ElemType InitialValue)
+	/**
+	* Create grid with specified domain.
+	*/
+	TBlockedGrid3(int32 DimI, int32 DimJ, int32 DimK, ElemType InitialValue)
+	: BlockedDenseGridType(DimI, DimJ, DimK, InitialValue)
 	{
-		Resize(DimX, DimY, DimZ);
-		ConstantValue = InitialValue;
 	}
 
-	int64 Size() const
+
+	// -- methods to set values when using locks.
+
+	/**
+	* Set the grid value at (I,J,K), with user supplied locking in the form 
+	* of a function of the form
+	* 
+	*		FCriticalSection& CriticalSectionProvider(id)
+	* 
+	* that provides a critical section for each block id.
+	*/
+	template <typename CriticalSectionProviderType>
+	void SetValueWithLock(int32 I, int32 J, int32 K, ElemType NewValue, CriticalSectionProviderType& CriticalSectionProvider)
 	{
-		return (int64)Dimensions.X * (int64)Dimensions.Y * (int64)Dimensions.Z;
+		int32 BlockIndex, LocalIndex;
+		BlockedGrid3LayoutType::GetBlockAndLocalIndex(I, J, K, BlockIndex, LocalIndex);
+		FCriticalSection& CriticalSection = CriticalSectionProvider(BlockIndex);
+		BlockedDenseGridType::WriteValueWithLock(BlockIndex, LocalIndex, [&NewValue](ElemType& Value){Value = NewValue;}, &CriticalSection);
 	}
 
-	bool IsValidIndex(const FVector3i& Index) const
+	/**
+	* Set the grid value at CellIJK,  with user supplied locking in the form
+	* of a function of the form
+	*
+	*		FCriticalSection& CriticalSectionProvider(id)
+	*
+	* that provides a critical section for each block id.
+	*/
+	template <typename CriticalSectionProviderType>
+	void SetValueWithLock(const FVector3i& CellIJK, ElemType NewValue,  CriticalSectionProviderType CriticalSectionProvider)
 	{
-		return Index.X >= 0 && Index.Y >= 0 && Index.Z >= 0
-			&& Index.X < Dimensions.X && Index.Y < Dimensions.Y && Index.Z < Dimensions.Z;
+		SetValueWithLock(CellIJK.X, CellIJK.Y, CellIJK.Z, NewValue, CriticalSectionProvider);
 	}
 
-	const FVector3i& GetDimensions() const
+	/**
+	* Call an external lambda with a reference to the grid value at (I,J,K).
+	* Called as Func(ElemType&), so the caller can both read and write the grid cell
+	* 
+	* note: if doing pure reads prefer the GetValue() methods because the ProcessValue method will allocate the underlying data block. 
+	*/
+	template<typename ProcessFunc, typename CriticalSectionProviderType>
+	void ProcessValueWithLock(int32 I, int32 J, int32 K, ProcessFunc Func, CriticalSectionProviderType CriticalSectionProvider)
 	{
-		return Dimensions;
+		int32 BlockIndex, LocalIndex;
+		BlockedGrid3LayoutType::GetBlockAndLocalIndex(I, J, K, BlockIndex, LocalIndex);
+		FCriticalSection& CriticalSection = CriticalSectionProvider(BlockIndex);
+		BlockedDenseGridType::WriteValueWithLock(BlockIndex, LocalIndex, Func, &CriticalSection);
 	}
 
 	/**
-	 * Reconfigure the grid to have the target dimensions. This clears all the 
-	 * existing grid memory.
-	 */
-	void Resize(int32 DimX, int32 DimY, int32 DimZ)
+	* Call an external lambda with a reference to the grid value at CellIJK.
+	* Called as Func(ElemType&), so the caller can both read and write the grid cell
+	*
+	* note: if doing pure reads prefer the GetValue() methods because the ProcessValue method will allocate the underlying data block. 
+	*/
+	template<typename ProcessFunc, typename CriticalSectionProviderType>
+	void ProcessValueWithLock(const FVector3i& CellIJK, ProcessFunc Func, CriticalSectionProviderType CriticalSectionProvider)
 	{
-		check((int64)DimX * (int64)DimY * (int64)DimZ < INT_MAX);
-		int32 BlocksX = (DimX / BlockSize) + 1;
-		int32 BlocksY = (DimY / BlockSize) + 1;
-		int32 BlocksZ = (DimZ / BlockSize) + 1;
+		ProcessValueWithLock(CellIJK.X, CellIJK.Y, CellIJK.Z, Func, CriticalSectionProvider);
+	}
 
-		Blocks.Reset();
-		int32 NumBlocks = BlocksX * BlocksY * BlocksZ;
-		Blocks.SetNum(NumBlocks);
-		BlockLocks.SetNum(NumBlocks);
-		for (int32 k = 0; k < NumBlocks; ++k)
+	// -- block access methods ---/
+
+	/** @return the total number of block units that covers the spatial dimensions of the grid.*/
+	int32 GetNumBlocks() const
+	{
+		const FVector3i& BlockDims = BlockedGrid3LayoutType::GetBlockDimensions();
+		return BlockDims.X * BlockDims.Y * BlockDims.Z;
+	}
+
+	/** @return array of pointers to the allocated blocks.  Note: the grid owns the blocks, so the calling code should not delete them. */
+	TArray<const BlockData3Type*> GetAllocatedConstBlocks() const
+	{
+		TArray<const BlockData3Type*> Result;
+		const int32 NumBlocks = BlockedDenseGridType::GetNumBlocks();
+		for (int32 i = 0; i < NumBlocks; ++i)
 		{
-			Blocks[k] = TUniquePtr<BlockType>();
+			if (BlockedDenseGridType::IsBlockAllocated(i))
+			{
+				Result.Add(BlockedDenseGridType::Blocks[i].BlockDataPtr.Get());
+			}
 		}
+		return Result;
+	}
+
+	/** @return array of pointers to the allocated blocks.  Note: the grid owns the blocks, so the calling code should not delete them.*/
+	TArray<BlockData3Type*> GetAllocatedBlocks()
+	{
+		TArray<BlockData3Type*> Result;
+		const int32 NumBlocks = BlockedDenseGridType::GetNumBlocks();
+		for (int32 i = 0; i < NumBlocks; ++i)
+		{
+			if (BlockedDenseGridType::IsBlockAllocated(i))
+			{
+				Result.Add(BlockedDenseGridType::Blocks[i].BlockDataPtr.Get());
+			}
+		}
+		return Result;
+	}
+
+	/** @return array of pointers to the allocated blocks.  Note: the grid owns the blocks, so the calling code should not delete them.*/
+	TArray<const BlockData3Type*> GetAllocatedBlocks() const
+	{
+		return GetAllocatedConstBlocks();
+	}
 
-		BlockDimensions = FVector3i(BlocksX, BlocksY, BlocksZ);
-		Dimensions = FVector3i(DimX, DimY, DimZ);
+	/** 
+	* @return true if the specified block has been allocated.
+	* @param BlockIJK  - the lattice address of the block in terms of block units ( e.g. the location in a grid of blocks).
+	*/
+	bool IsBlockAllocated(const FVector3i& BlockIJK) const
+	{
+		return GetBlockData(BlockIJK).IsValid();
+	}
+
+	/** 
+	* @return a references to the specified block, this will allocate the block if it doesn't already exists. 
+	* @param BlockIJK  - the lattice address of the block in terms of block units ( e.g. the location in a grid of blocks).
+	*/
+	BlockData3Type& TouchBlockData(const FVector3i& BlockIJK)
+	{
+		const int32 BlockIndex = BlockedGrid3LayoutType::BlockIJKToBlockIndex(BlockIJK);
+		return BlockedDenseGridType::TouchBlockData(BlockIndex);
+	}
+
+	/**
+	* @return reference to a unique pointer that points to a block that may or may not be allocated
+	* @param BlockIJK  - the lattice address of the block in terms of block units ( e.g. the location in a grid of blocks).
+	*/
+	TUniquePtr<BlockData3Type>& GetBlockData(const FVector3i& BlockIJK)
+	{
+		const int32 BlockIndex = BlockedGrid3LayoutType::BlockIJKToBlockIndex(BlockIJK);
+		return GetBlockData(BlockIndex);
 	}
 
 	/**
-	 * @return the grid value at (X,Y,Z)
-	 */
-	ElemType GetValue(int32 X, int32 Y, int32 Z) const
+	* @return reference to a unique pointer that points to a block that may or may not be allocated
+	* @param BlockIJK  - the lattice address of the block in terms of block units ( e.g. the location in a grid of blocks).
+	*/
+	const TUniquePtr<BlockData3Type>& GetBlockData(const FVector3i& BlockIJK) const
 	{
-		return ReadValue(X, Y, Z);
+		const int32 BlockIndex = BlockedGrid3LayoutType::BlockIJKToBlockIndex(BlockIJK);
+		return GetBlockData(BlockIndex);
 	}
 
 	/**
-	 * @return the grid value at (X,Y,Z), with internal locking, so it is safe to call this from multiple read & write threads
-	 */
-	ElemType GetValueThreadSafe(int32 X0, int32 Y, int32 Z)
+	* @return reference to a unique pointer that points to a block that may or may not be allocated
+	* @param BlockIndex  - the index of the block in an internal linearization of the blocks.
+	* Note:   BlockIJK and BlockIndex  can be converted by means of ::BlockIJKToBlockIndex() and ::BlockIndexToBlockIJK()
+	*/
+	TUniquePtr<BlockData3Type>& GetBlockData(int32 BlockIndex)
 	{
-		return ReadValueThreadSafe(X0, Y, Z);
+		return BlockedDenseGridType::GetBlock(BlockIndex).BlockDataPtr;
 	}
 
 	/**
-	 * Set the grid value at (X,Y,Z)
-	 */
-	void SetValue(int32 X, int32 Y, int32 Z, ElemType NewValue)
+	* @return reference to a unique pointer that points to a block that may or may not be allocated
+	* @param BlockIndex  - the index of the block in an internal linearization of the blocks.
+	* Note:   BlockIJK and BlockIndex  can be converted by means of ::BlockIJKToBlockIndex() and ::BlockIndexToBlockIJK()
+	*/
+	const TUniquePtr<BlockData3Type>& GetBlockData(int32 BlockIndex) const
 	{
-		WriteValue(X, Y, Z, [NewValue](ElemType& GridValueInOut) { GridValueInOut = NewValue; });
+		return BlockedDenseGridType::GetBlock(BlockIndex).BlockDataPtr;
 	}
 
+	// --- methods on the block's default values.  
+	// ---  The block 'default value' is returned when querying a cell in an unallocated block, and is used to initialize the block on allocation.
+
 	/**
-	 * Set the grid value at (X,Y,Z), with internal locking, so it is safe to call this from multiple read & write threads
-	 */
-	void SetValueThreadSafe(int32 X, int32 Y, int32 Z, ElemType NewValue)
+	* @return the default value for the specified block.
+	* @param BlockIJK    - the lattice address of the block in terms of block units ( e.g. the location in a grid of blocks).
+	* 
+	*/ 
+	ElemType GetBlockDefaultValue(const FVector3i& BlockIJK) const
 	{
-		WriteValueThreadSafe(X, Y, Z, [NewValue](ElemType& GridValueInOut) { GridValueInOut = NewValue; });
+		return ReadBlockDefaultValue(BlockIJK);
 	}
 
 	/**
-	 * Call an external lambda with a reference to the grid value at (X,Y,Z).
-	 * Called as Func(ElemType&), so the caller can both read and write the grid cell
-	 */
+	* Update the default value for the specified block.
+	* @param BlockIJK    - the lattice address of the block in terms of block units ( e.g. the location in a grid of blocks).
+	* @param Value       - the new default value for this block
+	* @param bPruneBlock - if true, any internal block allocation will be freed 
+	*                      and the resulting deallocated block will just hold a single default value.
+	*                      if false, the allocation of the block will no change.
+	* 
+	*/
+	void SetBlockDefaultValue(const FVector3i& BlockIJK, const ElemType& Value, bool bPruneBlock = true)
+	{
+		WriteBlockDefaultValue(BlockIJK, [&Value](ElemType& oldValue){ oldValue = Value;}, bPruneBlock);
+	}
+
+	/**
+	* Update the default value for the specified block using a functor
+	* @param BlockIJK    - the lattice address of the block in terms of block units ( e.g. the location in a grid of blocks).
+	* @param Func        - Func(ElemType& value){..} is called on the default value.
+	* @param bPruneBlock - if true, any internal block allocation will be freed 
+	*                      and the resulting deallocated block will just hold a single default value.
+	*                      if false, the allocation of the block will no change.
+	* 
+	*/
 	template<typename ProcessFunc>
-	void ProcessValueThreadSafe(int32 X, int32 Y, int32 Z, ProcessFunc Func)
+	void ProcessBlockDefaultValue(const FVector3i& BlockIJK, ProcessFunc Func, bool bPruneBlock = true)
 	{
-		WriteValueThreadSafe(X, Y, Z, Func);
+		WriteBlockDefaultValue(BlockIJK, Func, bPruneBlock);
 	}
+};
+
+// Note on the block size: 8 seems to be the sweet spot with vastly better performance in SDF generation
+// perhaps because a float block can fit in the L1 cache with room to spare.
+typedef TBlockedGrid3<float, 8>  FBlockedGrid3f;
+typedef TBlockedGrid3<double, 8> FBlockedGrid3d;
+typedef TBlockedGrid3<int, 8>    FBlockedGrid3i;
 
-	FAxisAlignedBox3i Bounds() const
+// note the bool grid uses a bitfield specialization
+typedef TBlockedGrid3<bool, 8>   FBlockedGrid3b;
+
+/**
+* TBlockedDenseGrid adds thread-safe access functions to the 3D blocked uniform grid, TBasicBlockedDenseGrid3.
+* Blocked grid is extended by adding internally owned locks for each data block.
+* 
+* The grid is allocated in BlockSize^3 blocks on-demand (BlockSize is a compile-time constant) allowing 
+* very large grids to be used without having to pre-allocate all the memory, eg for sparse/narrow-band use cases.
+*
+* For multi-threaded applications where memory is a premium, consider using TBlockedGrid3 (above) as it gives better 
+* control over the individual blocks and allows the caller to manage the lifetime of any locks used.
+*/
+template<typename ElemType>
+class TBlockedDenseGrid3 : public TBasicBlockedDenseGrid3<ElemType, 32>
+{
+public:
+	
+	static constexpr int32 BlockSize = 32;
+	static constexpr int32 BlockElemCount = BlockSize * BlockSize * BlockSize;
+	
+	typedef TBasicBlockedDenseGrid3<ElemType, BlockSize>                     BlockedDenseGridType;
+	typedef TBlockedGrid3Layout<BlockSize>                                   BlockedGrid3LayoutType;
+	typedef TBlockData3<ElemType, BlockSize>                                 BlockData3Type;
+
+
+	
+
+protected:
+
+	TArray<FCriticalSection> BlockLocks;             // per-block level lock used for thread-safe methods
+
+#if UE_BUILD_DEBUG
+
+	FCriticalSection* GetBlockLock(int32 Index) { return &BlockLocks[Index]; }
+#else
+	// skip range checks in non-debug builds
+	FCriticalSection* GetBlockLock(int32 Index) { return &BlockLocks.GetData()[Index]; }
+#endif
+
+	template<typename FuncType>
+	void WriteValueThreadSafe(int32 I, int32 J, int32 K, FuncType Func)
+	{
+		int32 BlockIndex, LocalIndex;
+		BlockedGrid3LayoutType::GetBlockAndLocalIndex(I, J, K, BlockIndex, LocalIndex);
+		BlockedDenseGridType::WriteValueWithLock(BlockIndex, LocalIndex, Func, GetBlockLock(BlockIndex));
+	}
+
+
+	ElemType ReadValueThreadSafe(int32 I, int32 J, int32 K)
+	{
+		int32 BlockIndex, LocalIndex;
+		BlockedGrid3LayoutType::GetBlockAndLocalIndex(I, J, K, BlockIndex, LocalIndex);
+		return BlockedDenseGridType::ReadValueWithLock(BlockIndex, LocalIndex, GetBlockLock(BlockIndex));
+	}
+
+public:
+
+	/**
+	* Create empty grid
+	*/
+	TBlockedDenseGrid3() : BlockedDenseGridType()
+	{
+	}
+
+	/**
+	* Create grid with specified domain.
+	*/
+	TBlockedDenseGrid3(int32 DimI, int32 DimJ, int32 DimK, ElemType InitialValue)
+	{
+		BlockedDenseGridType::DefaultValue = InitialValue;
+		Resize(DimI, DimJ, DimK);
+	}
+
+	/**
+	* Reconfigure the grid to have the target dimensions. This clears all the 
+	* existing grid memory and resets the block-level default value.
+	*/
+	void Resize(int32 DimI, int32 DimJ, int32 DimK)
+	{
+		BlockedDenseGridType::Resize(DimI, DimJ, DimK);
+
+		const int32 NumBlocks = BlockedDenseGridType::GetNumBlocks();
+		BlockLocks.Reset();
+		BlockLocks.SetNum(NumBlocks);
+	}
+
+	/**
+	* Discard all allocated blocks but retain grid dimensions
+	*/
+	void Reset()
+	{
+		Reset(BlockedGrid3LayoutType::Dimensions.X, BlockedGrid3LayoutType::Dimensions.Y, BlockedGrid3LayoutType::Dimensions.Z, BlockedDenseGridType::DefaultValue);
+	}
+
+	/**
+	* Reconfigure the grid to have the target dimensions and default value. 
+	* This clears all the existing grid memory.
+	*/
+	void Reset(int32 DimI, int32 DimJ, int32 DimK, ElemType InitialValue)
+	{
+		BlockedDenseGridType::DefaultValue = InitialValue;
+		Resize(DimI, DimJ, DimK);
+	}
+
+	/**
+	* @return true if the specified cell IJK is contained within the grid bounds.
+	*/ 
+	bool IsValidIndex(FVector3i CellIJK) const
+	{
+		return BlockedGrid3LayoutType::IsValidIJK(CellIJK);
+	}
+	/**
+	* @return the grid value at (I,J,K), with internal locking, so it is safe to call this from multiple read & write threads
+	*/
+	ElemType GetValueThreadSafe(int32 I, int32 J, int32 K)
+	{
+		return ReadValueThreadSafe(I, J, K);
+	}
+
+	/**
+	* @return the grid value at CellIJK, with internal locking, so it is safe to call this from multiple read & write threads
+	*/
+	ElemType GetValueThreadSafe(const FVector3i& CellIJK)
 	{
-		return FAxisAlignedBox3i({0, 0, 0},{Dimensions.X, Dimensions.Y, Dimensions.Z});
+		return ReadValueThreadSafe(CellIJK.X, CellIJK.Y, CellIJK.Z);
 	}
-	FAxisAlignedBox3i BoundsInclusive() const
+
+	/**
+	* Set the grid value at (I,J,K), with internal locking, so it is safe to call this from multiple read & write threads
+	*/
+	void SetValueThreadSafe(int32 I, int32 J, int32 K, ElemType NewValue)
+	{
+		WriteValueThreadSafe(I, J, K, [NewValue](ElemType& GridValueInOut) { GridValueInOut = NewValue; });
+	}
+
+	/**
+	* Set the grid value at CellIJK, with internal locking, so it is safe to call this from multiple read & write threads
+	*/
+	void SetValueThreadSafe(const FVector3i& CellIJK, ElemType NewValue)
 	{
-		return FAxisAlignedBox3i({0, 0, 0},{Dimensions.X-1, Dimensions.Y-1, Dimensions.Z-1});
+		WriteValueThreadSafe(CellIJK.X, CellIJK.Y, CellIJK.Z, [NewValue](ElemType& GridValueInOut) { GridValueInOut = NewValue; });
 	}
 
+	/**
+	* Call an external lambda with a reference to the grid value at (I,J,K).
+	* Called as Func(ElemType&), so the caller can both read and write the grid cell
+	*
+	* note: if doing pure reads prefer the GetValue() methods 
+	* because the ProcessValueThreadSafe method will allocate the underlying data block. 
+	*/
+	template<typename ProcessFunc>
+	void ProcessValueThreadSafe(int32 I, int32 J, int32 K, ProcessFunc Func)
+	{
+		WriteValueThreadSafe(I, J, K, Func);
+	}
+
+	/**
+	* Call an external lambda with a reference to the grid value at CellIJK.
+	* Called as Func(ElemType&), so the caller can both read and write the grid cell
+	*/
+	template<typename ProcessFunc>
+	void ProcessValueThreadSafe(const FVector3i& CellIJK, ProcessFunc Func)
+	{
+		WriteValueThreadSafe(CellIJK.X, CellIJK.Y, CellIJK.Z, Func);
+	}
+
+	/** Convert a grid cell id to CellIJK coordinates*/
 	FVector3i ToIndex(int64 LinearIndex) const
 	{
-		int32 x = (int32)(LinearIndex % (int64)Dimensions.X);
-		int32 y = (int32)((LinearIndex / (int64)Dimensions.X) % (int64)Dimensions.Y);
-		int32 z = (int32)(LinearIndex / ((int64)Dimensions.X * (int64)Dimensions.Y));
+		int32 x = (int32)(LinearIndex % (int64)BlockedGrid3LayoutType::Dimensions.X);
+		int32 y = (int32)((LinearIndex / (int64)BlockedGrid3LayoutType::Dimensions.X) % (int64)BlockedGrid3LayoutType::Dimensions.Y);
+		int32 z = (int32)(LinearIndex / ((int64)BlockedGrid3LayoutType::Dimensions.X * (int64)BlockedGrid3LayoutType::Dimensions.Y));
 		return FVector3i(x, y, z);
 	}
+	/** Convert grid CellIJK coordinates to a linear ID, note this id is independent of the actual internal data storage.*/
 	int64 ToLinear(int32 X, int32 Y, int32 Z) const
 	{
-		return (int64)X + (int64)Dimensions.X * ((int64)Y + (int64)Dimensions.Y * (int64)Z);
+		return (int64)X + (int64)BlockedGrid3LayoutType::Dimensions.X * ((int64)Y + (int64)BlockedGrid3LayoutType::Dimensions.Y * (int64)Z);
 	}
+	/** Convert grid CellIJK coordinates to a linear ID, note this id is independent of the actual internal data storage.*/
 	int64 ToLinear(const FVector3i& IJK) const
 	{
-		return (int64)IJK.X + (int64)Dimensions.X * ((int64)IJK.Y + (int64)Dimensions.Y * (int64)IJK.Z);
+		return (int64)IJK.X + (int64)BlockedGrid3LayoutType::Dimensions.X * ((int64)IJK.Y + (int64)BlockedGrid3LayoutType::Dimensions.Y * (int64)IJK.Z);
 	}
-};
 
 
-typedef TBlockedDenseGrid3<float> FBlockedDenseGrid3f;
+};
+
+typedef TBlockedDenseGrid3<float>  FBlockedDenseGrid3f;
 typedef TBlockedDenseGrid3<double> FBlockedDenseGrid3d;
-typedef TBlockedDenseGrid3<int> FBlockedDenseGrid3i;
+typedef TBlockedDenseGrid3<int>    FBlockedDenseGrid3i;
 
 
 } // end namespace UE::Geometry
diff --git a/Engine/Source/Runtime/GeometryCore/Public/Spatial/BlockedLayout3.h b/Engine/Source/Runtime/GeometryCore/Public/Spatial/BlockedLayout3.h
new file mode 100644
--- /dev/null
+++ b/Engine/Source/Runtime/GeometryCore/Public/Spatial/BlockedLayout3.h
@@ -0,0 +1,201 @@
+// Copyright Epic Games, Inc. All Rights Reserved.
+
+
+#pragma once
+#include "IntBoxTypes.h"
+
+
+namespace UE
+{
+namespace Geometry
+{
+// These are 3d grid space layout behaviors shared by the various blocked grids 
+// ( c.f. BlockedDenseGrid3.h and other usages in SparseNarrowBandSDFMesh.h )
+
+
+// Consistent translation between a linear array of size BlockSize to a 3d layout. 
+template <int32 BlockSize_>
+struct TBlockData3Layout
+{
+	static constexpr int32 BlockSize = BlockSize_;
+	static constexpr int32 ElemCount = BlockSize_ * BlockSize_ * BlockSize_;
+
+	/**
+	* @return the 3d cell associated with the specified linear index.
+	*/ 
+	static FVector3i ToLocalIJK(int32 LocalLinear)
+	{
+		//inverse of  ToLinear()
+		const int32 i = LocalLinear % BlockSize;
+		const int32 j = (LocalLinear / BlockSize) % BlockSize;
+		const int32 k = LocalLinear / (BlockSize * BlockSize);
+		return FVector3i(i, j, k);
+	}
+
+	/**
+	* @return linear index associated with the specified 3d cell.
+	*/ 
+	static inline int32 ToLinear(const int32& LocalI, const int32& LocalJ, const int32& LocalK)
+	{
+		//inverse of ToLocalIJK
+		return LocalI + BlockSize * (LocalJ + BlockSize * LocalK);
+	}
+};
+
+// Spatial decomposition of a region of 3d ijk space into 3d cubes of size BlockSize (per side).
+// This class gives a consistent translation between locations in 3d lattice space  and blocks of a spatial decomposition;
+// it is independent of any sort of data that you might want to associate with those locations.  See BlockedDenseGrid3.h for usage.
+template <int32 BlockSize_>
+class TBlockedGrid3Layout
+{
+public:
+	static constexpr int32 BlockSize      = BlockSize_;                        // block is a cube in ijk space of length BlockSize           
+	static constexpr int32 BlockElemCount = BlockSize * BlockSize * BlockSize; // no. of ijk coordinates (i.e. lattice points) in each block
+
+	typedef TBlockData3Layout<BlockSize_>  FBlockData3Layout;
+
+	TBlockedGrid3Layout()
+		: BlockDimensions(FVector3i::Zero())
+		, Dimensions(FVector3i::Zero())
+	{}
+
+	TBlockedGrid3Layout(int32 DimI, int32 DimJ, int32 DimK)
+	{
+		Resize(DimI, DimJ, DimK);
+	}
+
+	TBlockedGrid3Layout(FVector3i Dims)
+	{
+		Resize(Dims[0], Dims[1], Dims[2]);
+	}
+
+	/**
+	* Recompute the block dimensions required to cover the requested cell dimensions.
+	*/
+	void Resize(int32 DimI, int32 DimJ, int32 DimK)
+	{
+		check((int64)DimI * (int64)DimJ * (int64)DimK < INT_MAX); // [todo] want a larger check here.. right now 1290^3 is near the upper lim 
+		const int32 BlocksI = (DimI / BlockSize) + ( (DimI % BlockSize != 0) ? 1 : 0 );
+		const int32 BlocksJ = (DimJ / BlockSize) + ( (DimJ % BlockSize != 0) ? 1 : 0 );
+		const int32 BlocksK = (DimK / BlockSize) + ( (DimK % BlockSize != 0) ? 1 : 0 );
+
+		BlockDimensions = FVector3i(BlocksI, BlocksJ, BlocksK);
+		Dimensions = FVector3i(DimI, DimJ, DimK);
+	}
+
+	/**
+	* @return the dimensions of this layout in terms of cell ijks.
+	*/ 
+	const FVector3i& GetDimensions() const
+	{
+		return Dimensions;
+	}
+
+	/**
+	* @return the dimensions of this layout in terms of blocks.
+	*/ 
+	const FVector3i& GetBlockDimensions() const
+	{
+		return BlockDimensions;
+	}
+
+	/*
+	* @return an axis aligned box that defines the bounds of this region in terms of cells. This bounds is [inclusive, exclusive) 
+	*/
+	FAxisAlignedBox3i Bounds() const
+	{
+		return FAxisAlignedBox3i({ 0, 0, 0 }, { Dimensions[0], Dimensions[1], Dimensions[2] });
+	}
+
+	/*
+	* @return an axis aligned box that defines the bounds of this region in terms of cells. This bounds is [inclusive, inclusive] 
+	*/
+	FAxisAlignedBox3i BoundsInclusive() const
+	{
+		return FAxisAlignedBox3i({ 0, 0, 0 }, { Dimensions[0] - 1, Dimensions[1] - 1, Dimensions[2]- 1 });
+	}
+
+	/**
+	* @return the number of ijk cells in contained within this grid layout.
+	* note: because blocks are 3d pages of fixed size so the number of Blocks X Cells in Block
+	*       might actually be a larger.
+	*/ 
+	int64 Size() const
+	{
+		return (int64)Dimensions[0] * (int64)Dimensions[1] * (int64)Dimensions[2];
+	}
+
+	/**
+	* @return true if the specified cell IJK is within the bounds of the dimensions.
+	*/
+	bool IsValidIJK(const FVector3i& IJK) const
+	{
+		return IJK[0] >= 0 && IJK[1] >= 0 && IJK[2] >= 0
+			&& IJK[0] < Dimensions[0] && IJK[1] < Dimensions[1] && IJK[2] < Dimensions[2];
+	}
+
+	/**
+	* @return true if the specified block IJK is within the bounds of the block dimensions.
+	*/
+	bool IsValidBlockIJK(const FVector3i& BlockIJK)
+	{
+		return BlockIJK[0] >= 0 && BlockIJK[1] >= 0 && BlockIJK[2] >= 0
+			&& BlockIJK[0] < BlockDimensions[0] && BlockIJK[1] < BlockDimensions[1] && BlockIJK[2] < BlockDimensions[2];
+	}
+
+	/**
+	* @return the index of the block that holds the cell indicated by IJK.
+	*/ 
+	int32 IJKtoBlockIndex(const FVector3i& IJK) const
+	{
+		const FVector3i BlockIJK(IJK[0] / BlockSize, IJK[1]/BlockSize, IJK[2]/BlockSize);
+		return BlockIJKToBlockIndex(BlockIJK);
+	}
+
+	/**
+	* return the block index and local index of the cell at (I,J,K).
+	*/ 
+	void GetBlockAndLocalIndex(const FVector3i& IJK, int32& BlockIndexOut, int32& LocalIndexOut) const
+	{
+		BlockIndexOut = IJKtoBlockIndex(IJK);
+
+		const int32 LocalI = IJK.X % BlockSize;
+		const int32 LocalJ = IJK.Y % BlockSize;
+		const int32 LocalK = IJK.Z % BlockSize;
+		LocalIndexOut = FBlockData3Layout::ToLinear(LocalI, LocalJ, LocalK);
+	}	
+	void GetBlockAndLocalIndex(int32 I, int32 J, int32 K, int32& BlockIndexOut, int32& LocalIndexOut) const
+	{
+		return GetBlockAndLocalIndex(FVector3i(I, J, K), BlockIndexOut, LocalIndexOut);
+	}
+
+	/**
+	* @return the 3-d coordinates of the block specified by an index in a linearization of the blocks.
+	* BlockIJK is the lattice address of the block in terms of block units ( e.g. the location in a grid of blocks).
+	*/ 
+	FVector3i BlockIndexToBlockIJK(const int32 BlockIndex) const
+	{
+
+		const int32 BlockI = BlockIndex % BlockDimensions.X;
+		const int32 BlockJ = (BlockIndex / BlockDimensions.X) % BlockDimensions.Y;
+		const int32 BlockK = BlockIndex / (BlockDimensions.X * BlockDimensions.Y);
+
+		return FVector3i(BlockI, BlockJ, BlockK);
+	}
+
+	/**
+	* @return the block index for the specified block (BlockIJK) in a linearization of the 3d block space.
+	*/ 
+	int32 BlockIJKToBlockIndex(const FVector3i& BlockIJK) const 
+	{
+		return BlockIJK[0] + BlockDimensions.X * (BlockIJK[1] + BlockDimensions.Y * BlockIJK[2]);
+	}
+
+protected:
+
+	FVector3i BlockDimensions; // dimensions in blocks per axis
+	FVector3i Dimensions;      // dimensions in cells per axis
+};
+
+} // end namespace UE::Geometry
+} // end namespace UE
diff --git a/Engine/Source/Runtime/GeometryCore/Public/Spatial/DenseGrid3.h b/Engine/Source/Runtime/GeometryCore/Public/Spatial/DenseGrid3.h
--- a/Engine/Source/Runtime/GeometryCore/Public/Spatial/DenseGrid3.h
+++ b/Engine/Source/Runtime/GeometryCore/Public/Spatial/DenseGrid3.h
@@ -117,6 +117,65 @@ public:
 	}
 
 
+
+	/**
+	* @return the grid value at (X,Y,Z)
+	*/
+	ElemType GetValue(int32 X, int32 Y, int32 Z) const
+	{
+		return Buffer[X + Dimensions.X * (Y + Dimensions.Y * Z)];
+	}
+
+	/**
+	* @return the grid value at (X,Y,Z)
+	*/
+	ElemType GetValue(const FVector3i& CellCoords) const
+	{
+		return Buffer[CellCoords.X + Dimensions.X * (CellCoords.Y + Dimensions.Y * CellCoords.Z)];
+	}
+
+	/**
+	* Set the grid value at (X,Y,Z)
+	*/
+	void SetValue(int32 X, int32 Y, int32 Z, ElemType NewValue)
+	{
+		Buffer[X + Dimensions.X * (Y + Dimensions.Y * Z)] = NewValue;
+	}
+
+	/**
+	* Set the grid value at (X,Y,Z)
+	*/
+	void SetValue(const FVector3i& CellCoords, ElemType NewValue)
+	{
+		Buffer[CellCoords.X + Dimensions.X * (CellCoords.Y + Dimensions.Y * CellCoords.Z)] = NewValue;
+	}
+
+	/**
+	* Call an external lambda with a reference to the grid value at (X,Y,Z).
+	* Called as Func(ElemType&), so the caller can both read and write the grid cell
+	*/
+	template<typename ProcessFunc>
+	void ProcessValue(int32 X, int32 Y, int32 Z, ProcessFunc Func)
+	{
+		int32 Index = X + Dimensions.X * (Y + Dimensions.Y * Z);
+		ElemType& CurValue = Buffer[Index];
+		Func(CurValue);
+	}
+
+	/**
+	* Call an external lambda with a reference to the grid value at (X,Y,Z).
+	* Called as Func(ElemType&), so the caller can both read and write the grid cell
+	*/
+	template<typename ProcessFunc>
+	void ProcessValue(const FVector3i& CellCoords, ProcessFunc Func)
+	{
+		int32 Index = CellCoords.X + Dimensions.X * (CellCoords.Y + Dimensions.Y * CellCoords.Z);
+		ElemType& CurValue = Buffer[Index];
+		Func(CurValue);
+	}
+
+
+
 	void GetXPair(int X0, int Y, int Z, ElemType& AOut, ElemType& BOut) const
 	{
 		int Offset = Dimensions.X * (Y + Dimensions.Y * Z);
