Fix instability in initializing quaternions almost aligned to -Z axis, or generally rotating almost 180 degrees w/ a TQuaternion::FromTo() call

diff --git a/Engine/Source/Runtime/GeometryCore/Public/Quaternion.h b/Engine/Source/Runtime/GeometryCore/Public/Quaternion.h
--- a/Engine/Source/Runtime/GeometryCore/Public/Quaternion.h
+++ b/Engine/Source/Runtime/GeometryCore/Public/Quaternion.h
@@ -390,7 +390,7 @@ void TQuaternion<RealType>::SetFromTo(const TVector<RealType>& From, const TVect
 	//   output TQuaternion is not normalized and this causes problems,
 	//   eg like drift if we do repeated SetFromTo()
 	TVector<RealType> from = UE::Geometry::Normalized(From), to = UE::Geometry::Normalized(To);
-	TVector<RealType> bisector = UE::Geometry::Normalized(from + to);
+	TVector<RealType> bisector = UE::Geometry::Normalized(from + to, TMathUtil<RealType>::ZeroTolerance);
 	W = from.Dot(bisector);
 	if (W != 0) 
 	{
