LineSegmentGenerators - actually use the given box center when computing box lines

diff --git a/Engine/Source/Runtime/GeometryCore/Public/Generators/LineSegmentGenerators.h b/Engine/Source/Runtime/GeometryCore/Public/Generators/LineSegmentGenerators.h
--- a/Engine/Source/Runtime/GeometryCore/Public/Generators/LineSegmentGenerators.h
+++ b/Engine/Source/Runtime/GeometryCore/Public/Generators/LineSegmentGenerators.h
@@ -91,8 +91,8 @@ void UE::Geometry::GenerateBoxSegments(
 {
 	// B is box max/min, P and Q are used to store the start and endpoints of the segments we create
 	TVector<RealType>	B[2], P, Q;
-	B[0] = HalfDimensions; // max
-	B[1] = -B[0]; // min
+	B[0] = Center + HalfDimensions; // max
+	B[1] = Center - HalfDimensions; // min
 
 	// Iterate across the four corners of top/side/front and create segments along the Z axis
 	for (int32 i = 0; i < 2; i++)
