automatically fix mesh orientation when geometry script functions transform meshes with negative scale

diff --git a/Engine/Plugins/Experimental/GeometryScripting/Source/GeometryScriptingCore/Private/MeshBakeFunctions.cpp b/Engine/Plugins/Experimental/GeometryScripting/Source/GeometryScriptingCore/Private/MeshBakeFunctions.cpp
--- a/Engine/Plugins/Experimental/GeometryScripting/Source/GeometryScriptingCore/Private/MeshBakeFunctions.cpp
+++ b/Engine/Plugins/Experimental/GeometryScripting/Source/GeometryScriptingCore/Private/MeshBakeFunctions.cpp
@@ -497,7 +497,7 @@ namespace GeometryScriptBakeLocals
 			const FTransformSRT3d SourceToWorld = SourceTransform;
 			MeshTransforms::ApplyTransform(SourceMeshCopy, SourceToWorld);
 			const FTransformSRT3d TargetToWorld = TargetTransform;
-			MeshTransforms::ApplyTransformInverse(SourceMeshCopy, TargetToWorld);
+			MeshTransforms::ApplyTransformInverse(SourceMeshCopy, TargetToWorld, true);
 			SourceMeshToUse = &SourceMeshCopy;
 		}
 
@@ -660,9 +660,9 @@ namespace GeometryScriptBakeLocals
 			// TODO: Remove this once we have support for transforming rays in the core bake loop
 			SourceMeshCopy = *SourceMeshOriginal;
 			const FTransformSRT3d SourceToWorld = SourceTransform;
-			MeshTransforms::ApplyTransform(SourceMeshCopy, SourceToWorld);
+			MeshTransforms::ApplyTransform(SourceMeshCopy, SourceToWorld, true);
 			const FTransformSRT3d TargetToWorld = TargetTransform;
-			MeshTransforms::ApplyTransformInverse(SourceMeshCopy, TargetToWorld);
+			MeshTransforms::ApplyTransformInverse(SourceMeshCopy, TargetToWorld, true);
 			SourceMeshToUse = &SourceMeshCopy;
 		}
 
diff --git a/Engine/Plugins/Experimental/GeometryScripting/Source/GeometryScriptingCore/Private/MeshBasicEditFunctions.cpp b/Engine/Plugins/Experimental/GeometryScripting/Source/GeometryScriptingCore/Private/MeshBasicEditFunctions.cpp
--- a/Engine/Plugins/Experimental/GeometryScripting/Source/GeometryScriptingCore/Private/MeshBasicEditFunctions.cpp
+++ b/Engine/Plugins/Experimental/GeometryScripting/Source/GeometryScriptingCore/Private/MeshBasicEditFunctions.cpp
@@ -369,7 +369,7 @@ UDynamicMesh* UGeometryScriptLibrary_MeshBasicEditFunctions::AppendMeshRepeated(
 		AppendMesh->ProcessMesh([&](const FDynamicMesh3& OtherMesh) { TmpMesh.Copy(OtherMesh); });
 		if (bApplyTransformToFirstInstance)
 		{
-			MeshTransforms::ApplyTransform(TmpMesh, XForm);
+			MeshTransforms::ApplyTransform(TmpMesh, XForm, true);
 		}
 		TargetMesh->EditMesh([&](FDynamicMesh3& AppendToMesh)
 		{
@@ -380,7 +380,7 @@ UDynamicMesh* UGeometryScriptLibrary_MeshBasicEditFunctions::AppendMeshRepeated(
 				Editor.AppendMesh(&TmpMesh, TmpMappings);
 				if (k < RepeatCount)
 				{
-					MeshTransforms::ApplyTransform(TmpMesh, XForm);
+					MeshTransforms::ApplyTransform(TmpMesh, XForm, true);
 					TmpMappings.Reset();
 				}
 			}
diff --git a/Engine/Plugins/Experimental/GeometryScripting/Source/GeometryScriptingCore/Private/MeshBooleanFunctions.cpp b/Engine/Plugins/Experimental/GeometryScripting/Source/GeometryScriptingCore/Private/MeshBooleanFunctions.cpp
--- a/Engine/Plugins/Experimental/GeometryScripting/Source/GeometryScriptingCore/Private/MeshBooleanFunctions.cpp
+++ b/Engine/Plugins/Experimental/GeometryScripting/Source/GeometryScriptingCore/Private/MeshBooleanFunctions.cpp
@@ -85,7 +85,7 @@ UDynamicMesh* UGeometryScriptLibrary_MeshBooleanFunctions::ApplyMeshBoolean(
 	}
 
 	// Boolean result is in the space of TargetTransform, so invert that
-	MeshTransforms::ApplyTransformInverse(NewResultMesh, (FTransformSRT3d)TargetTransform);
+	MeshTransforms::ApplyTransformInverse(NewResultMesh, (FTransformSRT3d)TargetTransform, true);
 
 	if (NewBoundaryEdges.Num() > 0 && Options.bFillHoles)
 	{
diff --git a/Engine/Plugins/Experimental/GeometryScripting/Source/GeometryScriptingCore/Private/MeshPrimitiveFunctions.cpp b/Engine/Plugins/Experimental/GeometryScripting/Source/GeometryScriptingCore/Private/MeshPrimitiveFunctions.cpp
--- a/Engine/Plugins/Experimental/GeometryScripting/Source/GeometryScriptingCore/Private/MeshPrimitiveFunctions.cpp
+++ b/Engine/Plugins/Experimental/GeometryScripting/Source/GeometryScriptingCore/Private/MeshPrimitiveFunctions.cpp
@@ -40,7 +40,7 @@ static void AppendPrimitive(
 			MeshTransforms::Translate(Mesh, PreTranslate);
 		}
 
-		MeshTransforms::ApplyTransform(Mesh, (FTransformSRT3d)Transform);
+		MeshTransforms::ApplyTransform(Mesh, (FTransformSRT3d)Transform, true);
 		if (PrimitiveOptions.PolygroupMode == EGeometryScriptPrimitivePolygroupMode::SingleGroup)
 		{
 			for (int32 tid : Mesh.TriangleIndicesItr())
@@ -100,7 +100,7 @@ static void AppendPrimitiveMesh(
 			MeshTransforms::Translate(Mesh, PreTranslate);
 		}
 
-		MeshTransforms::ApplyTransform(Mesh, (FTransformSRT3d)Transform);
+		MeshTransforms::ApplyTransform(Mesh, (FTransformSRT3d)Transform, true);
 		if (PrimitiveOptions.PolygroupMode == EGeometryScriptPrimitivePolygroupMode::SingleGroup)
 		{
 			for (int32 tid : Mesh.TriangleIndicesItr())
diff --git a/Engine/Plugins/Experimental/GeometryScripting/Source/GeometryScriptingCore/Private/MeshTransformFunctions.cpp b/Engine/Plugins/Experimental/GeometryScripting/Source/GeometryScriptingCore/Private/MeshTransformFunctions.cpp
--- a/Engine/Plugins/Experimental/GeometryScripting/Source/GeometryScriptingCore/Private/MeshTransformFunctions.cpp
+++ b/Engine/Plugins/Experimental/GeometryScripting/Source/GeometryScriptingCore/Private/MeshTransformFunctions.cpp
@@ -14,6 +14,7 @@ using namespace UE::Geometry;
 UDynamicMesh* UGeometryScriptLibrary_MeshTransformFunctions::TransformMesh(
 	UDynamicMesh* TargetMesh,
 	FTransform Transform,
+	bool bFixOrientationForNegativeScale,
 	UGeometryScriptDebug* Debug)
 {
 	if (TargetMesh == nullptr)
@@ -25,7 +26,7 @@ UDynamicMesh* UGeometryScriptLibrary_MeshTransformFunctions::TransformMesh(
 	// todo: publish correct change types
 	TargetMesh->EditMesh([&](FDynamicMesh3& EditMesh) 
 	{
-		MeshTransforms::ApplyTransform(EditMesh, (FTransformSRT3d)Transform);
+		MeshTransforms::ApplyTransform(EditMesh, (FTransformSRT3d)Transform, bFixOrientationForNegativeScale);
 
 	}, EDynamicMeshChangeType::GeneralEdit, EDynamicMeshAttributeChangeFlags::Unknown, false);
 
@@ -59,6 +60,7 @@ UDynamicMesh* UGeometryScriptLibrary_MeshTransformFunctions::TranslateMesh(
 UDynamicMesh* UGeometryScriptLibrary_MeshTransformFunctions::ScaleMesh(
 	UDynamicMesh* TargetMesh,
 	FVector Scale,
+	bool bFixOrientationForNegativeScale,
 	UGeometryScriptDebug* Debug)
 {
 	if (TargetMesh == nullptr)
@@ -70,7 +72,7 @@ UDynamicMesh* UGeometryScriptLibrary_MeshTransformFunctions::ScaleMesh(
 	// todo: publish correct change types
 	TargetMesh->EditMesh([&](FDynamicMesh3& EditMesh) 
 	{
-		MeshTransforms::Scale(EditMesh, (FVector3d)Scale, FVector3d::Zero());
+		MeshTransforms::Scale(EditMesh, (FVector3d)Scale, FVector3d::Zero(), bFixOrientationForNegativeScale);
 
 	}, EDynamicMeshChangeType::GeneralEdit, EDynamicMeshAttributeChangeFlags::Unknown, false);
 
diff --git a/Engine/Plugins/Experimental/GeometryScripting/Source/GeometryScriptingCore/Private/SceneUtilityFunctions.cpp b/Engine/Plugins/Experimental/GeometryScripting/Source/GeometryScriptingCore/Private/SceneUtilityFunctions.cpp
--- a/Engine/Plugins/Experimental/GeometryScripting/Source/GeometryScriptingCore/Private/SceneUtilityFunctions.cpp
+++ b/Engine/Plugins/Experimental/GeometryScripting/Source/GeometryScriptingCore/Private/SceneUtilityFunctions.cpp
@@ -106,7 +106,7 @@ UDynamicMesh* UGeometryScriptLibrary_SceneUtilityFunctions::CopyMeshFromComponen
 	{
 		ToDynamicMesh->EditMesh([&](FDynamicMesh3& EditMesh) 
 		{
-			MeshTransforms::ApplyTransform(EditMesh, (FTransformSRT3d)LocalToWorld);
+			MeshTransforms::ApplyTransform(EditMesh, (FTransformSRT3d)LocalToWorld, true);
 
 		}, EDynamicMeshChangeType::GeneralEdit, EDynamicMeshAttributeChangeFlags::Unknown, false);	
 	}
diff --git a/Engine/Plugins/Experimental/GeometryScripting/Source/GeometryScriptingCore/Public/GeometryScript/MeshTransformFunctions.h b/Engine/Plugins/Experimental/GeometryScripting/Source/GeometryScriptingCore/Public/GeometryScript/MeshTransformFunctions.h
--- a/Engine/Plugins/Experimental/GeometryScripting/Source/GeometryScriptingCore/Public/GeometryScript/MeshTransformFunctions.h
+++ b/Engine/Plugins/Experimental/GeometryScripting/Source/GeometryScriptingCore/Public/GeometryScript/MeshTransformFunctions.h
@@ -16,12 +16,13 @@ class GEOMETRYSCRIPTINGCORE_API UGeometryScriptLibrary_MeshTransformFunctions :
 	GENERATED_BODY()
 public:
 
-	UFUNCTION(BlueprintCallable, Category = "GeometryScript|Transforms", meta=(ScriptMethod))
-	static UPARAM(DisplayName = "Target Mesh") UDynamicMesh* 
+	UFUNCTION(BlueprintCallable, Category = "GeometryScript|Transforms", meta=(ScriptMethod, AdvancedDisplay = "bFixOrientationForNegativeScale"))
+	static UPARAM(DisplayName = "Target Mesh") UDynamicMesh*
 	TransformMesh(
 		UDynamicMesh* TargetMesh,
 		FTransform Transform,
-		UGeometryScriptDebug* Debug = nullptr );
+		bool bFixOrientationForNegativeScale = true,
+		UGeometryScriptDebug* Debug = nullptr);
 
 	UFUNCTION(BlueprintCallable, Category = "GeometryScript|Transforms", meta=(ScriptMethod))
 	static UPARAM(DisplayName = "Target Mesh") UDynamicMesh* 
@@ -30,11 +31,12 @@ public:
 		FVector Translation,
 		UGeometryScriptDebug* Debug = nullptr );
 
-	UFUNCTION(BlueprintCallable, Category = "GeometryScript|Transforms", meta=(ScriptMethod))
+	UFUNCTION(BlueprintCallable, Category = "GeometryScript|Transforms", meta=(ScriptMethod, AdvancedDisplay = "bFixOrientationForNegativeScale"))
 	static UPARAM(DisplayName = "Target Mesh") UDynamicMesh* 
 	ScaleMesh(
 		UDynamicMesh* TargetMesh,
 		FVector Scale = FVector(1,1,1),
-		UGeometryScriptDebug* Debug = nullptr );
+		bool bFixOrientationForNegativeScale = true,
+		UGeometryScriptDebug* Debug = nullptr);
 
 };
\ No newline at end of file
diff --git a/Engine/Source/Runtime/GeometryCore/Private/DynamicMesh/MeshTransforms.cpp b/Engine/Source/Runtime/GeometryCore/Private/DynamicMesh/MeshTransforms.cpp
--- a/Engine/Source/Runtime/GeometryCore/Private/DynamicMesh/MeshTransforms.cpp
+++ b/Engine/Source/Runtime/GeometryCore/Private/DynamicMesh/MeshTransforms.cpp
@@ -21,7 +21,7 @@ void MeshTransforms::Translate(FDynamicMesh3& Mesh, const FVector3d& Translation
 }
 
 
-void MeshTransforms::Scale(FDynamicMesh3& Mesh, const FVector3d& Scale, const FVector3d& Origin)
+void MeshTransforms::Scale(FDynamicMesh3& Mesh, const FVector3d& Scale, const FVector3d& Origin, bool bReverseOrientationIfNeeded)
 {
 	int NumVertices = Mesh.MaxVertexID();
 	ParallelFor(NumVertices, [&](int32 vid)
@@ -31,6 +31,10 @@ void MeshTransforms::Scale(FDynamicMesh3& Mesh, const FVector3d& Scale, const FV
 			Mesh.SetVertex(vid, (Mesh.GetVertex(vid) - Origin) * Scale + Origin);
 		}
 	});
+	if (bReverseOrientationIfNeeded && Scale.X * Scale.Y * Scale.Z < 0)
+	{
+		Mesh.ReverseOrientation(false);
+	}
 }
 
 
@@ -112,7 +116,7 @@ void MeshTransforms::FrameCoordsToWorld(FDynamicMesh3& Mesh, const FFrame3d& Fra
 
 
 
-void MeshTransforms::ApplyTransform(FDynamicMesh3& Mesh, const FTransformSRT3d& Transform)
+void MeshTransforms::ApplyTransform(FDynamicMesh3& Mesh, const FTransformSRT3d& Transform, bool bReverseOrientationIfNeeded)
 {
 	bool bVertexNormals = Mesh.HasVertexNormals();
 
@@ -139,11 +143,16 @@ void MeshTransforms::ApplyTransform(FDynamicMesh3& Mesh, const FTransformSRT3d&
 			Normals->SetElement(elemid, Normalized(Normal));
 		}
 	}
+
+	if (bReverseOrientationIfNeeded && Transform.GetDeterminant() < 0)
+	{
+		Mesh.ReverseOrientation(false);
+	}
 }
 
 
 
-void MeshTransforms::ApplyTransformInverse(FDynamicMesh3& Mesh, const FTransformSRT3d& Transform)
+void MeshTransforms::ApplyTransformInverse(FDynamicMesh3& Mesh, const FTransformSRT3d& Transform, bool bReverseOrientationIfNeeded)
 {
 	bool bVertexNormals = Mesh.HasVertexNormals();
 
@@ -182,6 +191,11 @@ void MeshTransforms::ApplyTransformInverse(FDynamicMesh3& Mesh, const FTransform
 			});
 		}
 	}
+
+	if (bReverseOrientationIfNeeded && Transform.GetDeterminant() < 0)
+	{
+		Mesh.ReverseOrientation(false);
+	}
 }
 
 
diff --git a/Engine/Source/Runtime/GeometryCore/Public/DynamicMesh/MeshTransforms.h b/Engine/Source/Runtime/GeometryCore/Public/DynamicMesh/MeshTransforms.h
--- a/Engine/Source/Runtime/GeometryCore/Public/DynamicMesh/MeshTransforms.h
+++ b/Engine/Source/Runtime/GeometryCore/Public/DynamicMesh/MeshTransforms.h
@@ -22,8 +22,9 @@ namespace MeshTransforms
 
 	/**
 	 * Apply Scale to vertex positions of Mesh, relative to given Origin. Does not modify any other attributes (normals/etc)
+	 * @param bReverseOrientationIfNeeded		If negative scaling inverts the mesh, also invert the triangle orientations
 	 */
-	GEOMETRYCORE_API void Scale(FDynamicMesh3& Mesh, const FVector3d& Scale, const FVector3d& Origin);
+	GEOMETRYCORE_API void Scale(FDynamicMesh3& Mesh, const FVector3d& Scale, const FVector3d& Origin, bool bReverseOrientationIfNeeded = false);
 
 	/**
 	 * Transform Mesh into local coordinates of Frame
@@ -39,15 +40,17 @@ namespace MeshTransforms
 	/**
 	 * Apply given Transform to a Mesh.
 	 * Modifies Vertex Positions and Normals, and any Per-Triangle Normal Overlays
+	 * @param bReverseOrientationIfNeeded		If the Transform inverts the mesh with negative scale, also invert the triangle orientations
 	 */
-	GEOMETRYCORE_API void ApplyTransform(FDynamicMesh3& Mesh, const FTransformSRT3d& Transform);
+	GEOMETRYCORE_API void ApplyTransform(FDynamicMesh3& Mesh, const FTransformSRT3d& Transform, bool bReverseOrientationIfNeeded = false);
 
 
 	/**
 	 * Apply inverse of given Transform to a Mesh.
 	 * Modifies Vertex Positions and Normals, and any Per-Triangle Normal Overlays
+	  @param bReverseOrientationIfNeeded		If the Transform inverts the mesh with negative scale, also invert the triangle orientations
 	 */
-	GEOMETRYCORE_API void ApplyTransformInverse(FDynamicMesh3& Mesh, const FTransformSRT3d& Transform);
+	GEOMETRYCORE_API void ApplyTransformInverse(FDynamicMesh3& Mesh, const FTransformSRT3d& Transform, bool bReverseOrientationIfNeeded = false);
 
 
 	/**
