Fix constrained delaunay triangulation failing on cavity fill by switching to a different (simpler) cavity fill algorithm
Should be faster in most practical cases, though can be slower if an inserted edge must cross a lot of existing edges.

diff --git a/Engine/Source/Runtime/GeometryCore/Private/CompGeom/Delaunay2.cpp b/Engine/Source/Runtime/GeometryCore/Private/CompGeom/Delaunay2.cpp
--- a/Engine/Source/Runtime/GeometryCore/Private/CompGeom/Delaunay2.cpp
+++ b/Engine/Source/Runtime/GeometryCore/Private/CompGeom/Delaunay2.cpp
@@ -442,6 +442,7 @@ struct FDelaunay2Connectivity
 
 	void AddTriangle(const FIndex3i& Tri)
 	{
+		checkSlow(Tri.A != Tri.B && Tri.A != Tri.C && Tri.B != Tri.C);
 		EdgeToVert.Add(FIndex2i(Tri.A, Tri.B), Tri.C);
 		EdgeToVert.Add(FIndex2i(Tri.B, Tri.C), Tri.A);
 		EdgeToVert.Add(FIndex2i(Tri.C, Tri.A), Tri.B);
@@ -1025,91 +1026,49 @@ namespace DelaunayInternal
 		return EdgeToConnect.A;
 	}
 
-	// Helper for FillCavity.  Adds new vertex U to the cavity, trying to attach it initially via triangle UVW
-	template<typename RealType>
-	void CavityInsertVertex(FDelaunay2Connectivity& CavityCDT, TArrayView<const TVector2<RealType>> Vertices, int32 U, FIndex2i VW)
-	{
-		int32 X = CavityCDT.GetVertex(FIndex2i(VW.B, VW.A));
-		// If adding the triangle does not conflict with the existing triangle opp edge VW
-		// we can immediately add the triangle
-		if (
-			X == CavityCDT.InvalidIndex ||
-			(0 < ExactPredicates::Orient2<RealType>(Vertices[U], Vertices[VW.A], Vertices[VW.B]) &&
-			0 >= ExactPredicates::InCircle2<RealType>(Vertices[U], Vertices[VW.A], Vertices[VW.B], Vertices[X]))
-		)
-		{
-			CavityCDT.AddTriangle(FIndex3i(U, VW.A, VW.B)); // already Delaunay
-			return;
-		}
-		// New vertex U conflicts with existing triangle VWX across edge VW,
-		// so we need to delete tri VWX and insert flipped triangles UVX and UXW
-		// then recurse to flip any triangles that are made non-Delaunay by that flip
-		CavityCDT.DeleteTriangle(FIndex3i(VW.B, VW.A, X));
-		CavityInsertVertex<RealType>(CavityCDT, Vertices, U, FIndex2i(VW.A, X));
-		CavityInsertVertex<RealType>(CavityCDT, Vertices, U, FIndex2i(X, VW.B));
-	};
-
-	// Implements the cavity CDT algorithm from "Delaunay Mesh Generation" page 76, 77
+	// Note: previously this implemented the cavity CDT algorithm from "Delaunay Mesh Generation" page 76, 77, but that algorithm appears incorrect; see
+	//  "Fast Segment Insertion and Incremental Construction of Constrained Delaunay Triangulations" by Shewchuk and Brown, for a more complicated,
+	//  corrected version of that algorithm.
+	// For now FillCavity implements a simpler algorithm from Anglada (1997), which the Shewchuk and Brown paper reports as faster in practice for cavities less
+	// than ~30-85 vertices, which should be the vast majority of cases.  They suggest implementing both algorithms and switching based on the input size;
+	// so that is something to consider if FillCavity becomes a bottleneck in the future.
 	template<typename RealType>
 	void FillCavity(FRandomStream& Random, FDelaunay2Connectivity& Connectivity, TArrayView<const TVector2<RealType>> Vertices, const FIndex2i& Edge, const TArray<int32>& Cavity)
 	{
 		check(Cavity.Num() > 2 && Edge.B == Cavity[0] && Edge.A == Cavity.Last());
 		int32 CavityNum = Cavity.Num();
 
-		TArray<int32> Permute = GetShuffledOrder(Random, CavityNum, 1, CavityNum-2); // permutation of the inner vertices of Cavity
-		// doubly-linked list tracking the cavity re-ordering
-		TArray<int32> Next, Prev;
-		Next.SetNumUninitialized(CavityNum);
-		Prev.SetNumUninitialized(CavityNum);
-		for (int32 CurIdx = 0, PrevIdx = CavityNum-1; CurIdx < CavityNum; PrevIdx = CurIdx++)
-		{
-			Next[PrevIdx] = CurIdx;
-			Prev[CurIdx] = PrevIdx;
-		}
 
-		// Precompute Orient2D values that are proportional to distance to the AB line
-		TArray<RealType> ABOrient;
-		ABOrient.Init(0, CavityNum);
-		TVector2<RealType> V0 = Vertices[Cavity[0]];
-		TVector2<RealType> VLast = Vertices[Cavity.Last()];
-		for (int32 OrientIdx = 1; OrientIdx + 1 < CavityNum; OrientIdx++)
+		TArray<FIndex2i> Ranges;
+		Ranges.Reserve(Cavity.Num() - 2);
+		Ranges.Emplace(0, Cavity.Num() - 1);
+		while (!Ranges.IsEmpty())
 		{
-			// TODO: this predicate is only exact vs zero; could use a fast alternative here?
-			ABOrient[OrientIdx] = ExactPredicates::Orient2<RealType>(V0, Vertices[Cavity[OrientIdx]], VLast);
-			checkSlow(ABOrient[OrientIdx] > 0);
-		}
 
-		FDelaunay2Connectivity CavityCDT(false);
-		CavityCDT.Empty(CavityNum - 2); 
-
-		// Create an insertion ordering that tries to avoid concavities between adjacent pairs, and track adjacencies via Next/Prev
-		for (int Idx = CavityNum - 2; Idx >= 2; Idx--)
-		{
-			// Make sure Permute[Idx] is not closer to the AB line than both its neighbors
-			while (
-				ABOrient[Permute[Idx]] < ABOrient[Prev[Permute[Idx]]] &&
-				ABOrient[Permute[Idx]] < ABOrient[Next[Permute[Idx]]])
+			FIndex2i Range = Ranges.Pop(false);
+			int32 Mid = Range.A + 1;
+			if (Range.A + 2 == Range.B)
 			{
-				int32 SwapIdx = Random.RandRange(1, Idx - 1);
-				Swap(Permute[Idx], Permute[SwapIdx]);
+				Connectivity.AddTriangle(FIndex3i(Cavity[Range.A], Cavity[Mid], Cavity[Range.B]));
+				continue;
+			}
+			for (int32 Cand = Mid + 1; Cand < Range.B; Cand++)
+			{
+				if (0 < ExactPredicates::InCircle2<RealType>(Vertices[Cavity[Range.A]], Vertices[Cavity[Mid]], Vertices[Cavity[Range.B]], Vertices[Cavity[Cand]]))
+				{
+					Mid = Cand;
+				}
+			}
+			Connectivity.AddTriangle(FIndex3i(Cavity[Range.A], Cavity[Mid], Cavity[Range.B]));
+			if (Mid - Range.A > 1)
+			{
+				Ranges.Emplace(Range.A, Mid);
+			}
+			if (Range.B - Mid > 1)
+			{
+				Ranges.Emplace(Mid, Range.B);
 			}
-			// make Next/Prev skip over Permute[Idx]
-			Next[Prev[Permute[Idx]]] = Next[Permute[Idx]];
-			Prev[Next[Permute[Idx]]] = Prev[Permute[Idx]];
-		}
-
-		// Add the first triangle of the cavity
-		CavityCDT.AddTriangle(FIndex3i(Cavity[0], Cavity[Permute[1]], Cavity[CavityNum-1]));
-
-		// Progressively add remaining triangles in the permuted ordering (via a recursive function that keeps the cavity triangulation Delaunay)
-		for (int32 Idx = 2; Idx < CavityNum - 1; Idx++)
-		{
-			CavityInsertVertex<RealType>(CavityCDT, Vertices, Cavity[Permute[Idx]],
-				FIndex2i(Cavity[Next[Permute[Idx]]], Cavity[Prev[Permute[Idx]]]));
 		}
-
-		// Insert the cavity triangulation into the overall triangulation
-		Connectivity.Append(CavityCDT);
 	}
 
 	template<typename RealType>
