Fixes to FPriorityOrderPoints:
1. Actually use the SpatialLevels argument
2. Only allow early stopping at level boundaries, so the priority ordering can properly sort points at the same spatial level. Fixes points being spatially clumped when early stopping is enabled.
3. Stop referring to some of the importance weights as "curvature" (we often weight by things other than curvature)
4. Expose OffsetResFactor to allow more control of how clumpy the distribution is (value was hardcoded at 4, but I think 1 or 2 tends to space things better)

diff --git a/Engine/Source/Runtime/GeometryCore/Private/Spatial/PriorityOrderPoints.cpp b/Engine/Source/Runtime/GeometryCore/Private/Spatial/PriorityOrderPoints.cpp
--- a/Engine/Source/Runtime/GeometryCore/Private/Spatial/PriorityOrderPoints.cpp
+++ b/Engine/Source/Runtime/GeometryCore/Private/Spatial/PriorityOrderPoints.cpp
@@ -34,8 +34,10 @@ private:
 // Spatial ordering algorithm
 // Note we support two sets of importance weights because it is sometimes helpful to also sort by an additional metric, such as distance from center
 template<typename RealType>
-void OrderPoints(TArray<int32>& PointOrder, TArrayView<const TVector<RealType>> Points, TArrayView<const TArrayView<const float>> ImportanceWeights, int32 EarlyStop, int32 SpatialLevels)
+void OrderPoints(TArray<int32>& PointOrder, TArrayView<const TVector<RealType>> Points, TArrayView<const TArrayView<const float>> ImportanceWeights, int32 EarlyStop, int32 SpatialLevels, int32 OffsetResFactor)
 {
+	OffsetResFactor = FMath::Clamp(OffsetResFactor, 1, 8);
+
 	const int32 NumPoints = Points.Num();
 	PointOrder.SetNumUninitialized(NumPoints);
 	if (NumPoints == 0)
@@ -137,8 +139,9 @@ void OrderPoints(TArray<int32>& PointOrder, TArrayView<const TVector<RealType>>
 	TMap<int64, FBestPtData> BestPoints;
 	TSet<int64> CoveredCells, OffsetCoveredCells;
 	TArray<int32> MoveToFront;
+	int32 MaxRes = 1 << FMath::Clamp(SpatialLevels, 1, 30);
 	// traverse power of 2 cell resolutions (~octree levels), moving forward "best" points from the cells at each level
-	for (int32 Resolution = 2; MovedPtsOffset < NumPoints && Resolution <= 1024; Resolution *= 2)
+	for (int32 Resolution = 2; MovedPtsOffset < NumPoints && Resolution <= MaxRes; Resolution *= 2)
 	{
 		BestPoints.Reset();
 		CoveredCells.Reset();
@@ -146,7 +149,7 @@ void OrderPoints(TArray<int32>& PointOrder, TArrayView<const TVector<RealType>>
 		const RealType CellSize = MaxBBoxDim / static_cast<RealType>(Resolution);
 
 		// make quarter-sized cells w/ a half-cell offset to help detect points that are close, but across cell boundaries
-		int32 OffsetResolution = Resolution * 4;
+		int32 OffsetResolution = Resolution * OffsetResFactor;
 		const RealType OffsetCellSize = MaxBBoxDim / static_cast<RealType>(OffsetResolution);
 		OffsetResolution += 2; // allow for offset center
 		TVector<RealType> OffsetCenter = -TVector<RealType>(OffsetCellSize / 2);
@@ -218,8 +221,6 @@ void OrderPoints(TArray<int32>& PointOrder, TArrayView<const TVector<RealType>>
 			}
 		};
 
-		bool bShouldStopEarly = false;
-
 		// Decide which points to move into the 'front' section
 		int32 NumToMoveToFront = 0;
 		MoveToFront.Reset();
@@ -231,16 +232,11 @@ void OrderPoints(TArray<int32>& PointOrder, TArrayView<const TVector<RealType>>
 			{
 				NumToMoveToFront++;
 				MoveToFront.Add(BestDataPair.Value.Idx[0]);
-				bShouldStopEarly |= (EarlyStop > -1 && (MovedPtsOffset + NumToMoveToFront) >= EarlyStop);
-				if (bShouldStopEarly)
-				{
-					break;
-				}
 			}
 		}
-		if (NumWeights > 1 && !bShouldStopEarly) // TODO: make this a loop if we support more than two importance weights
+		if (NumWeights > 1) // TODO: make this a loop if we support more than two importance weights
 		{
-			// Add points favored by the curvature metric second
+			// Add points favored by the second metric *after* adding all the first-metric-favored points
 			//  (so they may be skipped where a first-metric-favored point already covered a region)
 			for (const TPair<int64, FBestPtData>& BestDataPair : BestPoints)
 			{
@@ -249,11 +245,6 @@ void OrderPoints(TArray<int32>& PointOrder, TArrayView<const TVector<RealType>>
 				{
 					NumToMoveToFront++;
 					MoveToFront.Add(BestDataPair.Value.Idx[1]);
-					bShouldStopEarly |= (EarlyStop > -1 && (MovedPtsOffset + NumToMoveToFront) >= EarlyStop);
-					if (bShouldStopEarly)
-					{
-						break;
-					}
 				}
 			}
 		}
@@ -292,54 +283,54 @@ void OrderPoints(TArray<int32>& PointOrder, TArrayView<const TVector<RealType>>
 		// Sort the just-added points by first metric (they were in arbitrary order before)
 		if (NumWeights > 0)
 		{
-			DescendingPredicate<float> DescendingCurvaturePred(ImportanceWeights[0]);
-			Sort(&PointOrder[MovedPtsOffset], NumToMoveToFront, DescendingCurvaturePred);
+			DescendingPredicate<float> DescendingImportancePred(ImportanceWeights[0]);
+			Sort(&PointOrder[MovedPtsOffset], NumToMoveToFront, DescendingImportancePred);
 		}
 
 		MovedPtsOffset += NumToMoveToFront;
 
-		if (bShouldStopEarly)
+		if (0 <= EarlyStop && EarlyStop <= MovedPtsOffset)
 		{
 			break;
 		}
 	}
 
 	// Sort remaining non-coincident points by first metric
-	if (NumWeights > 0 && MovedPtsOffset != EarlyStop && MovedPtsOffset < NumPoints)
+	if (NumWeights > 0 && (EarlyStop < 0 || EarlyStop > MovedPtsOffset) && MovedPtsOffset < NumPoints)
 	{
-		DescendingPredicate<float> DescendingCurvaturePred(ImportanceWeights[0]);
-		Sort(&PointOrder[MovedPtsOffset], NumPoints - MovedPtsOffset, DescendingCurvaturePred);
+		DescendingPredicate<float> DescendingImportancePred(ImportanceWeights[0]);
+		Sort(&PointOrder[MovedPtsOffset], NumPoints - MovedPtsOffset, DescendingImportancePred);
 	}
 }
 }
 
 
-void FPriorityOrderPoints::ComputeUniformSpaced(TArrayView<const FVector3d> Points, TArrayView<const float> ImportanceWeights, int32 EarlyStop)
+void FPriorityOrderPoints::ComputeUniformSpaced(TArrayView<const FVector3d> Points, TArrayView<const float> ImportanceWeights, int32 EarlyStop, int32 OffsetResFactor)
 {
 	TArrayView<TArrayView<const float>> ImportanceWeightsWeights(&ImportanceWeights, 1);
-	PriorityOrderPointsLocal::OrderPoints<double>(Order, Points, ImportanceWeightsWeights, EarlyStop, SpatialLevels);
+	PriorityOrderPointsLocal::OrderPoints<double>(Order, Points, ImportanceWeightsWeights, EarlyStop, SpatialLevels, OffsetResFactor);
 }
 
-void FPriorityOrderPoints::ComputeUniformSpaced(TArrayView<const FVector3f> Points, TArrayView<const float> ImportanceWeights, int32 EarlyStop)
+void FPriorityOrderPoints::ComputeUniformSpaced(TArrayView<const FVector3f> Points, TArrayView<const float> ImportanceWeights, int32 EarlyStop, int32 OffsetResFactor)
 {
 	TArrayView<TArrayView<const float>> ImportanceWeightsWeights(&ImportanceWeights, 1);
-	PriorityOrderPointsLocal::OrderPoints<float>(Order, Points, ImportanceWeightsWeights, EarlyStop, SpatialLevels);
+	PriorityOrderPointsLocal::OrderPoints<float>(Order, Points, ImportanceWeightsWeights, EarlyStop, SpatialLevels, OffsetResFactor);
 }
 
-void FPriorityOrderPoints::ComputeUniformSpaced(TArrayView<const FVector3d> Points, TArrayView<const float> ImportanceWeights, TArrayView<const float> SecondImportanceWeights, int32 EarlyStop)
+void FPriorityOrderPoints::ComputeUniformSpaced(TArrayView<const FVector3d> Points, TArrayView<const float> ImportanceWeights, TArrayView<const float> SecondImportanceWeights, int32 EarlyStop, int32 OffsetResFactor)
 {
 	TArray<TArrayView<const float>, TFixedAllocator<2>> MultiImportanceWeights;
 	MultiImportanceWeights.Add(ImportanceWeights);
 	MultiImportanceWeights.Add(SecondImportanceWeights);
-	PriorityOrderPointsLocal::OrderPoints<double>(Order, Points, MultiImportanceWeights, EarlyStop, SpatialLevels);
+	PriorityOrderPointsLocal::OrderPoints<double>(Order, Points, MultiImportanceWeights, EarlyStop, SpatialLevels, OffsetResFactor);
 }
 
-void FPriorityOrderPoints::ComputeUniformSpaced(TArrayView<const FVector3f> Points, TArrayView<const float> ImportanceWeights, TArrayView<const float> SecondImportanceWeights, int32 EarlyStop)
+void FPriorityOrderPoints::ComputeUniformSpaced(TArrayView<const FVector3f> Points, TArrayView<const float> ImportanceWeights, TArrayView<const float> SecondImportanceWeights, int32 EarlyStop, int32 OffsetResFactor)
 {
 	TArray<TArrayView<const float>, TFixedAllocator<2>> MultiImportanceWeights;
 	MultiImportanceWeights.Add(ImportanceWeights);
 	MultiImportanceWeights.Add(SecondImportanceWeights);
-	PriorityOrderPointsLocal::OrderPoints<float>(Order, Points, MultiImportanceWeights, EarlyStop, SpatialLevels);
+	PriorityOrderPointsLocal::OrderPoints<float>(Order, Points, MultiImportanceWeights, EarlyStop, SpatialLevels, OffsetResFactor);
 }
 
 void FPriorityOrderPoints::ComputeDescendingImportance(TArrayView<const float> ImportanceWeights)
@@ -354,6 +345,6 @@ void FPriorityOrderPoints::ComputeDescendingImportance(TArrayView<const float> I
 	{
 		Order[Idx] = Idx;
 	}
-	PriorityOrderPointsLocal::DescendingPredicate<float> DescendingCurvaturePred(ImportanceWeights);
-	Sort(&Order[0], NumPoints, DescendingCurvaturePred);
+	PriorityOrderPointsLocal::DescendingPredicate<float> DescendingImportancePred(ImportanceWeights);
+	Sort(&Order[0], NumPoints, DescendingImportancePred);
 }
\ No newline at end of file
diff --git a/Engine/Source/Runtime/GeometryCore/Public/Spatial/PriorityOrderPoints.h b/Engine/Source/Runtime/GeometryCore/Public/Spatial/PriorityOrderPoints.h
--- a/Engine/Source/Runtime/GeometryCore/Public/Spatial/PriorityOrderPoints.h
+++ b/Engine/Source/Runtime/GeometryCore/Public/Spatial/PriorityOrderPoints.h
@@ -45,9 +45,10 @@ public:
 	 * @param Points				Points to re-order
 	 * @param ImportanceWeights		Points with a higher ImportanceWeight in their local region of space will appear earlier in the ordering
 	 * @param EarlyStop				If >= 0, stop re-ordering once we've sorted this many points; beyond this the ordering will be arbitrary
+	 * @param OffsetResFactor		Factor that controls how closely the points are allowed to space, with higher values -> more clumping allowed
 	 */
-	void ComputeUniformSpaced(TArrayView<const FVector3d> Points, TArrayView<const float> ImportanceWeights, int32 EarlyStop = -1);
-	void ComputeUniformSpaced(TArrayView<const FVector3f> Points, TArrayView<const float> ImportanceWeights, int32 EarlyStop = -1);
+	void ComputeUniformSpaced(TArrayView<const FVector3d> Points, TArrayView<const float> ImportanceWeights, int32 EarlyStop = -1, int32 OffsetResFactor = 4);
+	void ComputeUniformSpaced(TArrayView<const FVector3f> Points, TArrayView<const float> ImportanceWeights, int32 EarlyStop = -1, int32 OffsetResFactor = 4);
 
 	/**
 	 * Compute an ordering for the input points that attempts to keep points 'well spaced' / un-clumped
@@ -56,9 +57,10 @@ public:
 	 * @param ImportanceWeights			Points with a higher ImportanceWeight in their local region of space will appear earlier in the ordering
 	 * @param SecondImportanceWeights	A second set of importance weights; we may pick an additional point from each local region to also sample points by this metric
 	 * @param EarlyStop					If >= 0, stop re-ordering once we've sorted this many points; beyond this the ordering will be arbitrary
+	 * @param OffsetResFactor			Factor that controls how closely the points are allowed to space, with higher values -> more clumping allowed
 	 */
-	void ComputeUniformSpaced(TArrayView<const FVector3d> Points, TArrayView<const float> ImportanceWeights, TArrayView<const float> SecondImportanceWeights, int32 EarlyStop = -1);
-	void ComputeUniformSpaced(TArrayView<const FVector3f> Points, TArrayView<const float> ImportanceWeights, TArrayView<const float> SecondImportanceWeights, int32 EarlyStop = -1);
+	void ComputeUniformSpaced(TArrayView<const FVector3d> Points, TArrayView<const float> ImportanceWeights, TArrayView<const float> SecondImportanceWeights, int32 EarlyStop = -1, int32 OffsetResFactor = 4);
+	void ComputeUniformSpaced(TArrayView<const FVector3f> Points, TArrayView<const float> ImportanceWeights, TArrayView<const float> SecondImportanceWeights, int32 EarlyStop = -1, int32 OffsetResFactor = 4);
 
 	/**
 	 * Compute an ordering that only sorts by descending ImportanceWeights, and does not attempt to spatially distribute the points
