GeometryCore: SparseNarrowBandMeshSDF.   Update the way distance sign is computed for a block far from the narrow band to make it consistent with the existing dense grid version (SweepingMeshSDF.h)

all the voxels in a block far from the mesh surface should share the same state (inside or outside) for water-tight meshes,  but the counting method used for sign assignment may produce odd results for non water-tight meshes where the block may have both inside and outside elements.  this change accounts for that behavior.

diff --git a/Engine/Source/Runtime/GeometryCore/Public/Implicit/SparseNarrowBandMeshSDF.h b/Engine/Source/Runtime/GeometryCore/Public/Implicit/SparseNarrowBandMeshSDF.h
--- a/Engine/Source/Runtime/GeometryCore/Public/Implicit/SparseNarrowBandMeshSDF.h
+++ b/Engine/Source/Runtime/GeometryCore/Public/Implicit/SparseNarrowBandMeshSDF.h
@@ -1234,7 +1234,7 @@ private:
 				const int32 BlockK = BlockRowID / BlockDims[1];
 				FVector3i BlockIJK(0, BlockJ, BlockK);
 				for (  BlockIJK[0] = 0; BlockIJK[0] < BlockDims[0]; ++BlockIJK[0]) {
-
+					
 					if (Distances.IsBlockAllocated(BlockIJK))
 					{
 						checkSlow(IntersectionCountGrid.IsBlockAllocated(BlockIJK));
@@ -1274,26 +1274,73 @@ private:
 					else
 					{
 						checkSlow(!IntersectionCountGrid.IsBlockAllocated(BlockIJK));
-						// the block isn't allocated - so it "should" all share the same state (inside or outside) for water-tight meshes,
+						
+						// the block isn't allocated - so it should all share the same state (inside or outside) for water-tight meshes,
 						// but the counting method used for sign assignment may produce odd results for non water-tight meshes where a block
-						// far from the narrow band may have both inside and outside elements..  in this case we just assign the block
-						// the most common sign.
-						int32 NumNeg = 0;
-						for (int32 index = 0, IndexMax = BlockSize * BlockSize; index < IndexMax; ++index)
+						// far from the narrow band may have both inside and outside elements.  
+			
+						// count the number of "inside" cells on the (already processed) neighbor block face adjacent to BlockIJK
+						int32 NumInsideCells = 0;
 						{
-							const int32 total_count = BlockCrossSection[index];
-							if (
-								(InsideMode == EInsideModes::WindingCount && total_count > 0) ||
-								(InsideMode == EInsideModes::CrossingCount && total_count % 2 == 1)
-								)
+							for (int32 index = 0, IndexMax = BlockSize * BlockSize; index < IndexMax; ++index)
 							{
-								NumNeg++;
+								const int32 total_count = BlockCrossSection[index];
+								if (
+									(InsideMode == EInsideModes::WindingCount && total_count > 0) ||
+									(InsideMode == EInsideModes::CrossingCount && total_count % 2 == 1)
+									)
+								{
+									NumInsideCells++;
+								}
 							}
+					
 						}
-						if ( float(NumNeg)/float(BlockSize * BlockSize) > 0.5f )
+						
+						switch (NumInsideCells)
 						{
-							// we are inside the mesh
-							Distances.ProcessBlockDefaultValue( BlockIJK, [](float& value){value = -FMath::Abs(value);});
+							case  0:
+							{
+								// neighbor is all outside - this unallocated block can be represented by existing positive value ( do nothing)
+								break;
+							}
+							case (BlockSize * BlockSize):
+							{
+								// neighbor is all inside - the unallocated block can be represented by single negative 
+								// we are inside the mesh
+								Distances.ProcessBlockDefaultValue(BlockIJK, [](float& value) {value = -FMath::Abs(value); });
+								break;
+							}
+							default:
+							{
+								// allocate the block (unique blocks can be allocated in parallel)
+								FBlockedGrid3f::BlockData3Type& DistanceBlockData = Distances.TouchBlockData(BlockIJK);
+
+								for (int32 k = 0; k < BlockSize; ++k)
+								{
+									for (int32 j = 0; j < BlockSize; ++j)
+									{
+										const int32 total_count = BlockCrossSection[j + k * BlockSize];
+
+										for (int32 i = 0; i < BlockSize; ++i)
+										{
+											const int32 LocalIndex = TBlockData3Layout<BlockSize>::ToLinear(i, j, k);
+
+											if (
+												(InsideMode == EInsideModes::WindingCount && total_count > 0) ||
+												(InsideMode == EInsideModes::CrossingCount && total_count % 2 == 1)
+												)
+											{
+												// we are inside the mesh
+												float& d = DistanceBlockData.At(LocalIndex);
+												d = -FMath::Abs(d);
+											}
+
+										}
+									}
+								}
+
+								break;
+							}
 						}
 					}
 						
