ContainmentTests: use full level set grid transform, not just translation

diff --git a/Engine/Plugins/Runtime/GeometryProcessing/Source/DynamicMesh/Private/ShapeApproximation/SimpleShapeSet3.cpp b/Engine/Plugins/Runtime/GeometryProcessing/Source/DynamicMesh/Private/ShapeApproximation/SimpleShapeSet3.cpp
--- a/Engine/Plugins/Runtime/GeometryProcessing/Source/DynamicMesh/Private/ShapeApproximation/SimpleShapeSet3.cpp
+++ b/Engine/Plugins/Runtime/GeometryProcessing/Source/DynamicMesh/Private/ShapeApproximation/SimpleShapeSet3.cpp
@@ -219,37 +219,34 @@ static void FilterContained(const FSimpleShapeSet3d& Geometry, const FConvexShap
 
 static void FilterContained(const FSimpleShapeSet3d& Geometry, FLevelSetShape3d& LevelSetShape, const TArray<FSimpleShapeElementKey>& Elements, int32 k, TArray<bool>& RemovedInOut)
 {
-	FOrientedBox3d GridBox = LevelSetShape.GetGridBox();
+	TTriLinearGridInterpolant<FDenseGrid3f> GridInterp(&LevelSetShape.Grid, FVector3d::ZeroVector, LevelSetShape.CellSize, LevelSetShape.Grid.GetDimensions());
 
 	int32 N = Elements.Num();
 	for (int32 j = k + 1; j < N; ++j)
 	{
 		if (RemovedInOut[j] == false)
 		{
-			// TODO: use the full GridTransform, not just translation. JIRA UE-155269
-			TTriLinearGridInterpolant<FDenseGrid3f> GridInterp(&LevelSetShape.Grid, LevelSetShape.GridTransform.GetTranslation(), LevelSetShape.CellSize, LevelSetShape.Grid.GetDimensions());
-
 			bool bContained = false;
 			int32 ElemIdx = Elements[j].Index;
 			if (Elements[j].Type == ESimpleShapeType::Sphere)
 			{
-				bContained = UE::Geometry::IsInside(GridInterp, Geometry.Spheres[ElemIdx].Sphere);
+				bContained = UE::Geometry::IsInside(GridInterp, LevelSetShape.GridTransform, Geometry.Spheres[ElemIdx].Sphere);
 			}
 			else if (Elements[j].Type == ESimpleShapeType::Box)
 			{
-				bContained = UE::Geometry::IsInside(GridInterp, Geometry.Boxes[ElemIdx].Box);
+				bContained = UE::Geometry::IsInside(GridInterp, LevelSetShape.GridTransform, Geometry.Boxes[ElemIdx].Box);
 			}
 			else if (Elements[j].Type == ESimpleShapeType::Capsule)
 			{
-				bContained = UE::Geometry::IsInside(GridInterp, Geometry.Capsules[ElemIdx].Capsule);
+				bContained = UE::Geometry::IsInside(GridInterp, LevelSetShape.GridTransform, Geometry.Capsules[ElemIdx].Capsule);
 			}
 			else if (Elements[j].Type == ESimpleShapeType::Convex)
 			{
-				bContained = UE::Geometry::IsInside<double>(GridInterp, Geometry.Convexes[ElemIdx].Mesh.VerticesItr());
+				bContained = UE::Geometry::IsInside<double>(GridInterp, LevelSetShape.GridTransform, Geometry.Convexes[ElemIdx].Mesh.VerticesItr());
 			}
 			else if (Elements[j].Type == ESimpleShapeType::LevelSet)
 			{
-				bContained = UE::Geometry::IsInside(GridInterp, Geometry.LevelSets[ElemIdx].GetGridBox());
+				bContained = UE::Geometry::IsInside(GridInterp, LevelSetShape.GridTransform, Geometry.LevelSets[ElemIdx].GetGridBox());
 			}
 			else
 			{
diff --git a/Engine/Source/Runtime/GeometryCore/Public/Intersection/ContainmentQueries3.h b/Engine/Source/Runtime/GeometryCore/Public/Intersection/ContainmentQueries3.h
--- a/Engine/Source/Runtime/GeometryCore/Public/Intersection/ContainmentQueries3.h
+++ b/Engine/Source/Runtime/GeometryCore/Public/Intersection/ContainmentQueries3.h
@@ -183,33 +183,34 @@ namespace UE
 		 * Test if the axis-aligned bounding box of InnerSphere is completely inside the negative region of the signed distance field discretized on OuterGrid
 		 */
 		template<typename RealType, typename GridType>
-		bool IsInside(const TTriLinearGridInterpolant<GridType>& OuterGrid, const TSphere3<RealType>& InnerSphere);
+		bool IsInside(const TTriLinearGridInterpolant<GridType>& OuterGrid, const TTransform<RealType>& OuterGridTransform, const TSphere3<RealType>& InnerSphere);
 
 		/**
 		 * Test if the axis-aligned bounding box of InnerCapsule is completely inside the negative region of the signed distance field discretized on OuterGrid
 		 */
 		template<typename RealType, typename GridType>
-		bool IsInside(const TTriLinearGridInterpolant<GridType>& OuterGrid, const TCapsule3<RealType>& InnerCapsule);
+		bool IsInside(const TTriLinearGridInterpolant<GridType>& OuterGrid, const TTransform<RealType>& OuterGridTransform, const TCapsule3<RealType>& InnerCapsule);
 
 		/**
 		 * Test if the axis-aligned bounding box of InnerBox is completely inside the negative region of the signed distance field discretized on OuterGrid
 		 */
 		template<typename RealType, typename GridType>
-		bool IsInside(const TTriLinearGridInterpolant<GridType>& OuterGrid, const TOrientedBox3<RealType>& InnerBox);
+		bool IsInside(const TTriLinearGridInterpolant<GridType>& OuterGrid, const TTransform<RealType>& OuterGridTransform, const TOrientedBox3<RealType>& InnerBox);
 
 		/**
 		 * Test if all enumerable points are inside the negative region of the signed distance field discretized on OuterGrid
 		 */
 		template<typename RealType, typename GridType, typename EnumerablePointsType, typename E = decltype(DeclVal<EnumerablePointsType>().begin())>
-		bool IsInside(const TTriLinearGridInterpolant<GridType>& OuterGrid, EnumerablePointsType EnumerablePts);
+		bool IsInside(const TTriLinearGridInterpolant<GridType>& OuterGrid, const TTransform<RealType>& OuterGridTransform, EnumerablePointsType EnumerablePts);
 
 		/**
 		 * Test if the axis-aligned bounding box of InnerGrid is completely inside the negative region of the signed distance field discretized on OuterGrid
 		 */
 		template<typename RealType, typename GridType1, typename GridType2>
-		bool IsInside(const TTriLinearGridInterpolant<GridType1>& OuterGrid, const TTriLinearGridInterpolant<GridType2>& InnerGrid)
+		bool IsInside(const TTriLinearGridInterpolant<GridType1>& OuterGrid, const TTransform<RealType>& OuterGridTransform, 
+					  const TTriLinearGridInterpolant<GridType2>& InnerGrid, const TTransform<RealType>& InnerGridTransform )
 		{
-			return IsInside(OuterGrid, TOrientedBox3<RealType>(InnerGrid.Bounds()));
+			return IsInside(OuterGrid, OuterGridTransform * InnerGridTransform.Inverse(), TOrientedBox3<RealType>(InnerGrid.Bounds()));
 		}
 
 	}
@@ -325,43 +326,52 @@ namespace UE
 		}
 
 		template<typename RealType, typename GridType>
-		bool IsInside(const TTriLinearGridInterpolant<GridType>& OuterGrid, const TSphere3<RealType>& InnerSphere)
+		bool IsInside(const TTriLinearGridInterpolant<GridType>& OuterGrid, const TTransform<RealType>& OuterGridTransform, const TSphere3<RealType>& InnerSphere)
 		{
-			TAxisAlignedBox3<RealType> InnerAABB(InnerSphere.Center - InnerSphere.Radius * TVector<RealType>::OneVector,
+			const TAxisAlignedBox3<RealType> InnerAABB(InnerSphere.Center - InnerSphere.Radius * TVector<RealType>::OneVector,
 				InnerSphere.Center + InnerSphere.Radius * TVector<RealType>::OneVector);
 
-			// TODO: TSphere3 has a SignedDistance function, so we could get a better discretization
-			return IsInside(OuterGrid, InnerAABB);
+			const TAxisAlignedBox3<RealType> GridSpaceInnerAABB(InnerAABB, [&OuterGridTransform](const TVector<RealType>& Corner)
+			{
+				return OuterGridTransform.InverseTransformPosition(Corner);
+			});
+
+			return IsInside(OuterGrid, GridSpaceInnerAABB);
 		}
 
 		template<typename RealType, typename GridType>
-		bool IsInside(const TTriLinearGridInterpolant<GridType>& OuterGrid, const TCapsule3<RealType>& InnerCapsule)
+		bool IsInside(const TTriLinearGridInterpolant<GridType>& OuterGrid, const TTransform<RealType>& OuterGridTransform, const TCapsule3<RealType>& InnerCapsule)
 		{
-			// TODO: TCapsule3 has a SignedDistance function, so we could get a better discretization
-			return IsInside(OuterGrid, InnerCapsule.GetBounds());
+			const TAxisAlignedBox3<RealType> GridSpaceInnerAABB(InnerCapsule.GetBounds(), [&OuterGridTransform](const TVector<RealType>& Corner)
+			{
+				return OuterGridTransform.InverseTransformPosition(Corner);
+			});
+
+			return IsInside(OuterGrid, GridSpaceInnerAABB);
 		}
 
 		template<typename RealType, typename GridType>
-		bool IsInside(const TTriLinearGridInterpolant<GridType>& OuterGrid, const TOrientedBox3<RealType>& InnerBox)
+		bool IsInside(const TTriLinearGridInterpolant<GridType>& OuterGrid, const TTransform<RealType>& OuterGridTransform, const TOrientedBox3<RealType>& InnerBox)
 		{
-			TAxisAlignedBox3<RealType> InnerAABB;
-			InnerBox.EnumerateCorners([&InnerAABB](const FVector3d& CornerPt)
+			TAxisAlignedBox3<RealType> GridSpaceInnerAABB;
+			InnerBox.EnumerateCorners([&OuterGridTransform, &GridSpaceInnerAABB](const FVector3d& CornerPt)
 			{
-				InnerAABB.Contain(CornerPt);
+				GridSpaceInnerAABB.Contain(OuterGridTransform.InverseTransformPosition(CornerPt));
 			});
 
-			return IsInside(OuterGrid, InnerAABB);
+			return IsInside(OuterGrid, GridSpaceInnerAABB);
 		}
 
 		template<typename RealType, typename GridType, typename EnumerablePointsType, typename E>
-		bool IsInside(const TTriLinearGridInterpolant<GridType>& OuterGrid, EnumerablePointsType EnumerablePts)
+		bool IsInside(const TTriLinearGridInterpolant<GridType>& OuterGrid, const TTransform<RealType>& OuterGridTransform, EnumerablePointsType EnumerablePts)
 		{
-			TAxisAlignedBox3<RealType> InnerAABB;
+			TAxisAlignedBox3<RealType> GridSpaceInnerAABB;
 			for (TVector<RealType> Point : EnumerablePts)
 			{
-				InnerAABB.Contain(Point);
+				GridSpaceInnerAABB.Contain(OuterGridTransform.InverseTransformPosition(Point));
 			}
-			return IsInside(OuterGrid, InnerAABB);
+
+			return IsInside(OuterGrid, GridSpaceInnerAABB);
 		}
 
 	}
