GeometryCore:  Morphology.h,  Update to use sparse SDF due to memory savings over the dense SDF.  The AutoLOD tool and the Morphology tool use this class - but the behavior shouldn't change.

diff --git a/Engine/Source/Runtime/GeometryCore/Public/Implicit/Morphology.h b/Engine/Source/Runtime/GeometryCore/Public/Implicit/Morphology.h
--- a/Engine/Source/Runtime/GeometryCore/Public/Implicit/Morphology.h
+++ b/Engine/Source/Runtime/GeometryCore/Public/Implicit/Morphology.h
@@ -8,7 +8,7 @@
 
 #include "Spatial/MeshAABBTree3.h"
 
-#include "Implicit/SweepingMeshSDF.h"
+#include "Implicit/SparseNarrowBandMeshSDF.h"
 #include "Implicit/GridInterpolant.h"
 #include "Implicit/ImplicitFunctions.h"
 
@@ -131,15 +131,21 @@ public:
 	}
 
 protected:
+
+	template<typename MeshType>
+	using TMeshSDF = TSparseNarrowBandMeshSDF<MeshType>; 
+
 	void ComputeFirstPass(double UnsignedOffset, double SignedOffset)
 	{
 		TRACE_CPUPROFILER_EVENT_SCOPE(Geometry_Morphology_FirstPass);
 
-		TSweepingMeshSDF<TriangleMeshType> ComputedSDF;
+		typedef TMeshSDF<TriangleMeshType>   MeshSDFType;
+
+		MeshSDFType ComputedSDF;
 		ComputedSDF.Mesh = Source;
 
 		ComputedSDF.Spatial = SourceSpatial;
-		ComputedSDF.ComputeMode = TSweepingMeshSDF<TriangleMeshType>::EComputeModes::NarrowBand_SpatialFloodFill;
+		ComputedSDF.ComputeMode = MeshSDFType::EComputeModes::NarrowBand_SpatialFloodFill;
 
 
 		ComputedSDF.CellSize = GridCellSize;
@@ -152,7 +158,7 @@ protected:
 		if (!ComputedSDF.ShouldUseSpatial(ComputedSDF.ExactBandWidth, ComputedSDF.CellSize, AvgEdgeLen))
 		{
 			ComputedSDF.Spatial = nullptr;
-			ComputedSDF.ComputeMode = TSweepingMeshSDF<TriangleMeshType>::EComputeModes::NarrowBandOnly;
+			ComputedSDF.ComputeMode = MeshSDFType::EComputeModes::NarrowBandOnly;
 		}
 
 		{
@@ -160,7 +166,7 @@ protected:
 			ComputedSDF.Compute(SourceSpatial->GetBoundingBox());
 		}
 
-		TTriLinearGridInterpolant<TSweepingMeshSDF<TriangleMeshType>> Interpolant = ComputedSDF.MakeInterpolant();
+		TTriLinearGridInterpolant<MeshSDFType> Interpolant = ComputedSDF.MakeInterpolant();
 
 		MarchingCubes.IsoValue = SignedOffset;
 		MarchingCubes.Bounds = SourceSpatial->GetBoundingBox();
@@ -198,6 +204,8 @@ protected:
 	void ComputeSecondPass(double UnsignedOffset, double SignedOffset)
 	{
 		TRACE_CPUPROFILER_EVENT_SCOPE(Geometry_Morphology_SecondPass);
+		
+		typedef TMeshSDF<TIndexVectorMeshArrayAdapter<FIndex3i, double, FVector3d>>   MeshSDFType;
 
 		if (MarchingCubes.Triangles.Num() == 0)
 		{
@@ -208,7 +216,7 @@ protected:
 		TIndexVectorMeshArrayAdapter<FIndex3i, double, FVector3d> MCAdapter(&MarchingCubes.Vertices, &MarchingCubes.Triangles);
 		TMeshAABBTree3<TIndexVectorMeshArrayAdapter<FIndex3i, double, FVector3d>> SecondSpatial(&MCAdapter, false);
 
-		TSweepingMeshSDF<TIndexVectorMeshArrayAdapter<FIndex3i, double, FVector3d>> SecondSDF;
+		MeshSDFType SecondSDF;
 		SecondSDF.Mesh = &MCAdapter;
 		SecondSDF.CellSize = GridCellSize;
 		SecondSDF.Spatial = nullptr;
@@ -224,12 +232,12 @@ protected:
 			TRACE_CPUPROFILER_EVENT_SCOPE(Geometry_Morphology_SecondPass_BuildSpatial);
 			SecondSpatial.Build();
 			SecondSDF.Spatial = &SecondSpatial;
-			SecondSDF.ComputeMode = TSweepingMeshSDF<TIndexVectorMeshArrayAdapter<FIndex3i, double, FVector3d>>::EComputeModes::NarrowBand_SpatialFloodFill;
+			SecondSDF.ComputeMode = MeshSDFType::EComputeModes::NarrowBand_SpatialFloodFill;
 			Bounds = SecondSpatial.GetBoundingBox(); // Use the tighter bounds from the AABB tree since we have it
 		}
 		else
 		{
-			SecondSDF.ComputeMode = TSweepingMeshSDF<TIndexVectorMeshArrayAdapter<FIndex3i, double, FVector3d>>::EComputeModes::NarrowBandOnly;
+			SecondSDF.ComputeMode = MeshSDFType::EComputeModes::NarrowBandOnly;
 		}
 
 
@@ -242,7 +250,7 @@ protected:
 			TRACE_CPUPROFILER_EVENT_SCOPE(Geometry_Morphology_SecondPass_ComputeSDF);
 			SecondSDF.Compute(Bounds);
 		}
-		TTriLinearGridInterpolant<TSweepingMeshSDF<TIndexVectorMeshArrayAdapter<FIndex3i, double, FVector3d>>> Interpolant = SecondSDF.MakeInterpolant();
+		TTriLinearGridInterpolant<MeshSDFType> Interpolant = SecondSDF.MakeInterpolant();
 
 		MarchingCubes.Reset();
 		MarchingCubes.IsoValue = SignedOffset;
diff --git a/Engine/Source/Runtime/GeometryCore/Public/Implicit/SparseNarrowBandMeshSDF.h b/Engine/Source/Runtime/GeometryCore/Public/Implicit/SparseNarrowBandMeshSDF.h
--- a/Engine/Source/Runtime/GeometryCore/Public/Implicit/SparseNarrowBandMeshSDF.h
+++ b/Engine/Source/Runtime/GeometryCore/Public/Implicit/SparseNarrowBandMeshSDF.h
@@ -1178,13 +1178,13 @@ private:
 			const int32 NumBlocks = AtomicGrid3i.Blocks.Num();
 			ParallelFor(NumBlocks, [&IntersectionCountGrid, &AtomicGrid3i](const int32 blockid)
 			{
-				TUniquePtr<FAtomicGrid3i::BlockAtomicData>& AtomicBlockPtr = AtomicGrid3i.Blocks[blockid];
+				TUniquePtr<typename FAtomicGrid3i::BlockAtomicData>& AtomicBlockPtr = AtomicGrid3i.Blocks[blockid];
 				
 				if (AtomicBlockPtr.IsValid())
 				{
 					FVector3i BlockIJK = AtomicGrid3i.BlockIndexToBlockIJK(blockid);
 					FBlockedGrid3i::BlockData3Type& BlockData = IntersectionCountGrid.TouchBlockData(BlockIJK);	
-					const FAtomicGrid3i::BlockAtomicData& AtomicBlockData = *AtomicBlockPtr;
+					const typename FAtomicGrid3i::BlockAtomicData& AtomicBlockData = *AtomicBlockPtr;
 
 					for (int32 j = 0; j < FBlockedGrid3i::BlockElemCount; ++j)
 					{
