DynamicMesh3: Add comments to AppendTriangle functions saying the vertices must exist and the indices must be distinct

diff --git a/Engine/Source/Runtime/GeometryCore/Public/DynamicMesh/DynamicMesh3.h b/Engine/Source/Runtime/GeometryCore/Public/DynamicMesh/DynamicMesh3.h
--- a/Engine/Source/Runtime/GeometryCore/Public/DynamicMesh/DynamicMesh3.h
+++ b/Engine/Source/Runtime/GeometryCore/Public/DynamicMesh/DynamicMesh3.h
@@ -534,8 +534,10 @@ public:
 	/** Copy vertex SourceVertexID from existing SourceMesh, returns new vertex id */
 	int AppendVertex(const FDynamicMesh3& SourceMesh, int SourceVertexID);
 
+	/** TriVertices must be distinct and refer to existing, valid vertices */
 	int AppendTriangle(const FIndex3i& TriVertices, int GroupID = 0);
 
+	/** Vertex0, Vertex1, and Vertex2 must be distinct and refer to existing, valid vertices */
 	inline int AppendTriangle(int Vertex0, int Vertex1, int Vertex2, int GroupID = 0)
 	{
 		return AppendTriangle(FIndex3i(Vertex0, Vertex1, Vertex2), GroupID);
