Add progress tracking to FProgressCancel, with nested scopes and messages
Add FAsyncTaskExecuterWithProgressCancel to allow the progress to be reported back to the calling thread

Implement initial usage in Voronoi fracture, using an FSlowTask to handle the progress bar and cancel button UI

diff --git a/Engine/Plugins/Experimental/ChaosEditor/Source/FractureEditor/Private/FractureToolBackgroundTask.h b/Engine/Plugins/Experimental/ChaosEditor/Source/FractureEditor/Private/FractureToolBackgroundTask.h
--- a/Engine/Plugins/Experimental/ChaosEditor/Source/FractureEditor/Private/FractureToolBackgroundTask.h
+++ b/Engine/Plugins/Experimental/ChaosEditor/Source/FractureEditor/Private/FractureToolBackgroundTask.h
@@ -48,13 +48,14 @@ public:
 template<class GeometryCollectionOpType>
 int RunCancellableGeometryCollectionOp(FGeometryCollection& ToUpdate, TUniquePtr<GeometryCollectionOpType>&& NewOp, FText DefaultMessage, float DialogDelay = .5)
 {
-	FScopedSlowTask SlowTask(1, DefaultMessage);
-	SlowTask.MakeDialogDelayed(DialogDelay, true);
 	using FGeometryCollectionTask = UE::Geometry::TModelingOpTask<GeometryCollectionOpType>;
-	using FExecuter = UE::Geometry::FAsyncTaskExecuterWithAbort<FGeometryCollectionTask>;
+	using FExecuter = UE::Geometry::FAsyncTaskExecuterWithProgressCancel<FGeometryCollectionTask>;
 	TUniquePtr<FExecuter> BackgroundTask = MakeUnique<FExecuter>(MoveTemp(NewOp));
 	BackgroundTask->StartBackgroundTask();
 
+	FScopedSlowTask SlowTask(1, DefaultMessage);
+	SlowTask.MakeDialogDelayed(DialogDelay, true);
+
 	bool bSuccess = false;
 	while (true)
 	{
@@ -70,8 +71,19 @@ int RunCancellableGeometryCollectionOp(FGeometryCollection& ToUpdate, TUniquePtr
 			break;
 		}
 		FPlatformProcess::Sleep(.2); // SlowTask::ShouldCancel will throttle any updates faster than .2 seconds
-		SlowTask.TickProgress(); // Note: EnterProgressFrame will also call this, but we aren't getting progress updates here yet
-		// TODO: Connect up a way to get progress updates from the task back to SlowTask
+		float ProgressFrac;
+		FText ProgressMessage;
+		bool bMadeProgress = BackgroundTask->PollProgress(ProgressFrac, ProgressMessage);
+		if (bMadeProgress)
+		{
+			// SlowTask expects progress to be reported before it happens; we work around this by directly updating the progress amount
+			SlowTask.CompletedWork = ProgressFrac;
+			SlowTask.EnterProgressFrame(0, ProgressMessage);
+		}
+		else
+		{
+			SlowTask.TickProgress(); // Still tick the UI when we don't get new progress frames
+		}
 	}
 
 	if (bSuccess)
diff --git a/Engine/Plugins/Experimental/ChaosEditor/Source/FractureEditor/Private/FractureToolCutter.cpp b/Engine/Plugins/Experimental/ChaosEditor/Source/FractureEditor/Private/FractureToolCutter.cpp
--- a/Engine/Plugins/Experimental/ChaosEditor/Source/FractureEditor/Private/FractureToolCutter.cpp
+++ b/Engine/Plugins/Experimental/ChaosEditor/Source/FractureEditor/Private/FractureToolCutter.cpp
@@ -374,6 +374,10 @@ public:
 	// TGenericDataOperator interface:
 	virtual void CalculateResult(FProgressCancel* Progress) override
 	{
+		// Assume Voronoi diagram only takes ~1% of the total time to compute
+		FProgressCancel::FProgressScope VoronoiDiagramProgress =
+			FProgressCancel::CreateScopeTo(Progress, .01, LOCTEXT("ComputingVoronoiDiagramMessage", "Computing Voronoi Diagram"));
+
 		FVoronoiDiagram Voronoi(Sites, Bounds, .1f);
 
 		FPlanarCells VoronoiPlanarCells = FPlanarCells(Sites, Voronoi);
@@ -384,7 +388,13 @@ public:
 			return;
 		}
 
-		// TODO: Pass the ProgressCancel down so the computation can be cancelled
+		VoronoiDiagramProgress.Done();
+
+		// All remaining work is assigned to the mesh fracture
+		FProgressCancel::FProgressScope FractureMeshProgress =
+			FProgressCancel::CreateScopeTo(Progress, 1, LOCTEXT("FractureMeshMessage", "Fracturing Mesh"));
+
+		// TODO: Pass the ProgressCancel down so the computation can be tracked and cancelled
 		ResultGeometryIndex = CutMultipleWithPlanarCells(VoronoiPlanarCells, *CollectionCopy, Selection, Grout, PointSpacing, Seed, Transform);
 		
 		SetResult(MoveTemp(CollectionCopy));
diff --git a/Engine/Plugins/Runtime/MeshModelingToolset/Source/ModelingOperators/Public/ModelingTaskTypes.h b/Engine/Plugins/Runtime/MeshModelingToolset/Source/ModelingOperators/Public/ModelingTaskTypes.h
--- a/Engine/Plugins/Runtime/MeshModelingToolset/Source/ModelingOperators/Public/ModelingTaskTypes.h
+++ b/Engine/Plugins/Runtime/MeshModelingToolset/Source/ModelingOperators/Public/ModelingTaskTypes.h
@@ -95,6 +95,61 @@ public:
 	}
 };
 
+// An extension of FAsyncTaskExecuterWithAbort that requires the task have a GetProgress() function
+// Intended to be used with FAbortableBackgroundTask
+// Provides a PollProgress method to report out task progress
+template<typename TTask>
+class FAsyncTaskExecuterWithProgressCancel : public FAsyncTaskExecuterWithAbort<TTask>
+{
+protected:
+	FProgressCancel* Progress = nullptr;
+	float LastProgressFrac;
+	FText LastProgressMessage;
+
+public:
+	using FAsyncTaskExecuterWithAbort<TTask>::bAbort;
+
+	template <typename Arg0Type, typename... ArgTypes>
+	FAsyncTaskExecuterWithProgressCancel(Arg0Type&& Arg0, ArgTypes&&... Args)
+		: FAsyncTaskExecuterWithAbort<TTask>(Forward<Arg0Type>(Arg0), Forward<ArgTypes>(Args)...)
+
+	{
+		FAsyncTask<TTask>::GetTask().SetAbortSource(&bAbort);
+		Progress = FAsyncTask<TTask>::GetTask().GetProgress();
+	}
+
+	bool PollProgress(float& OutProgressFrac, FText& OutProgressMessage)
+	{
+		if (!Progress)
+		{
+			return false;
+		}
+		OutProgressFrac = Progress->GetProgress();
+		bool bNewProgress = false;
+		if (OutProgressFrac != LastProgressFrac)
+		{
+			OutProgressMessage = Progress->GetProgressMessage();
+			LastProgressMessage = OutProgressMessage;
+			LastProgressFrac = OutProgressFrac;
+			return true;
+		}
+		else if (OutProgressFrac == 0.0f)
+		{
+			// at progress == 0, allow message to update w/out progress advancing
+			// Note: Not polling the message at progress > 0 allows us to take the message lock less frequently,
+			// but does mean we could miss messages that happen w/out a corresponding progress change.
+			// Because these messages should be short-lived, this trade-off seems ok, but we can revisit this
+			// if the missed progress messages become a problem.
+			OutProgressMessage = Progress->GetProgressMessage();
+			if (!OutProgressMessage.IdenticalTo(LastProgressMessage))
+			{
+				LastProgressMessage = OutProgressMessage;
+				return true;
+			}
+		}
+		return false;
+	}
+};
 
 
 /**
diff --git a/Engine/Source/Runtime/GeometryCore/Private/Util/ProgressCancel.cpp b/Engine/Source/Runtime/GeometryCore/Private/Util/ProgressCancel.cpp
new file mode 100644
--- /dev/null
+++ b/Engine/Source/Runtime/GeometryCore/Private/Util/ProgressCancel.cpp
@@ -0,0 +1,80 @@
+// Copyright Epic Games, Inc. All Rights Reserved.
+
+#include "Util/ProgressCancel.h"
+
+using namespace UE::Geometry;
+
+void FProgressCancel::StartWorkScope(FProgressData& SaveProgressFrameOut, float StepSize, const FText& Message)
+{
+	SaveProgressFrameOut = Progress;
+	Progress.ScopeDepth++;
+	float SavedMax = Progress.CurrentMax;
+	float SavedRange = Progress.Range();
+	// Note: New min/max are clamped to the saved max to prevent going outside the active scope
+	// The clamp on CurrentMin should do nothing as the ProgressFraction should already be prevented from exceeding the range max
+	Progress.CurrentMin = FMathf::Min(ProgressFraction, SavedMax);
+	Progress.CurrentMax = FMathf::Min(ProgressFraction + SavedRange * StepSize, SavedMax);
+	if (!Message.IsEmpty()) // if message is empty, let any parent-scope message persist to this scope
+	{
+		SetProgressMessage(Message);
+	}
+}
+
+void FProgressCancel::EndWorkScope(const FProgressData& SavedProgressFrame)
+{
+	if (!ensureMsgf(Progress.ScopeDepth == SavedProgressFrame.ScopeDepth + 1,
+		TEXT("ScopeDepth (%d) not 1 less than saved ScopeDepth (%d): Likely incorrect use of FProgressScope"),
+		Progress.ScopeDepth, SavedProgressFrame.ScopeDepth))
+	{
+		if (SavedProgressFrame.ScopeDepth + 1 > Progress.ScopeDepth || Progress.CurrentMax > SavedProgressFrame.CurrentMax)
+		{
+			return; // ignore invalid scope rather than potentially allow the active range to go backward
+		}
+	}
+
+	ProgressFraction = Progress.CurrentMax;
+	{
+		bool bWasEmpty = Progress.Message.IsEmpty();
+		Progress = SavedProgressFrame;
+		if (!bWasEmpty || !Progress.Message.IsEmpty())
+		{
+			FScopeLock MessageLock(&MessageCS);
+			ProgressMessage = Progress.Message;
+		}
+	}
+}
+
+FProgressCancel::FProgressScope::FProgressScope(FProgressCancel* ProgressCancel, float ProgressAmount, const FText& Message) : ProgressCancel(ProgressCancel)
+{
+	if (ProgressCancel)
+	{
+		ProgressCancel->StartWorkScope(SavedProgressData, ProgressAmount, Message);
+	}
+}
+
+FProgressCancel::FProgressScope FProgressCancel::CreateScopeTo(FProgressCancel* ProgressCancel, float ProgressTo, const FText& Message)
+{
+	if (ProgressCancel)
+	{
+		float Range = ProgressCancel->Progress.Range();
+		float ProgressBy = 0;
+		if (Range > 0)
+		{
+			float TargetProgress = FMathf::Max(ProgressCancel->ProgressFraction, ProgressCancel->Progress.CurrentMin + Range * ProgressTo);
+			ProgressBy = (TargetProgress - ProgressCancel->ProgressFraction) / Range;
+		}
+		return FProgressScope(ProgressCancel, ProgressBy, Message);
+	}
+	return FProgressScope();
+}
+
+
+
+void FProgressCancel::FProgressScope::Done()
+{
+	if (ensure(!bEnded) && ProgressCancel)
+	{
+		ProgressCancel->EndWorkScope(SavedProgressData);
+	}
+	bEnded = true;
+}
diff --git a/Engine/Source/Runtime/GeometryCore/Public/Util/ProgressCancel.h b/Engine/Source/Runtime/GeometryCore/Public/Util/ProgressCancel.h
--- a/Engine/Source/Runtime/GeometryCore/Public/Util/ProgressCancel.h
+++ b/Engine/Source/Runtime/GeometryCore/Public/Util/ProgressCancel.h
@@ -8,7 +8,10 @@
 #include "Containers/Array.h"
 #include "Internationalization/Text.h"
 #include "Misc/DateTime.h"
+#include "Misc/ScopeLock.h"
+#include "MathUtil.h"
 
+#include <atomic>
 
 class FProgressCancel;
 
@@ -172,19 +175,65 @@ struct FGeometryResult
 
 
 /**
- * FProgressCancel is an obejct that is intended to be passed to long-running
+ * FProgressCancel is intended to be passed to long-running
  * computes to do two things:
- * 1) provide progress info back to caller (not implemented yet)
+ * 1) provide progress info back to caller
  * 2) allow caller to cancel the computation
  */
-class FProgressCancel
+class GEOMETRYCORE_API FProgressCancel
 {
 private:
 	bool WasCancelled = false;  // will be set to true if CancelF() ever returns true
 
+	// Progress tracking data
+	struct FProgressData
+	{
+		// Active range of progress tracking -- progress should not go out of this range
+		float CurrentMin = 0;
+		float CurrentMax = 1;
+		// Number of nested scopes currently affecting the active range
+		int ScopeDepth = 0;
+		// Message describing the current work being done
+		FText Message;
+
+		float Range() const
+		{
+			return CurrentMax - CurrentMin;
+		}
+
+		// Advance progress without going beyond the current range
+		float SafeAdvance(float CurrentFraction, float AdvanceAmount) const
+		{
+			return FMathf::Min(CurrentMax, CurrentFraction + Range() * AdvanceAmount);
+		}
+	};
+	FProgressData Progress;
+	FText ProgressMessage;
+	int MaxMessageDepth = -1; // If set to a positive value, message will not be updated by scopes below this depth
+
+	// Current total progress, in range [0, 1]
+	std::atomic<float> ProgressFraction = 0;
+
+	// critical section for accesses to the Progress Message
+	mutable FCriticalSection MessageCS;
+
+	void StartWorkScope(FProgressData& SaveProgressFrameOut, float StepSize, const FText& Message);
+
+	void EndWorkScope(const FProgressData& SavedProgressFrame);
+
 public:
 	TFunction<bool()> CancelF = []() { return false; };
 
+	void SetMaxMessageDepth(int32 ScopeDepth)
+	{
+		MaxMessageDepth = ScopeDepth;
+	}
+
+	void ClearMaxMessageDepth()
+	{
+		MaxMessageDepth = -1;
+	}
+
 	/**
 	 * @return true if client would like to cancel operation
 	 */
@@ -198,6 +247,184 @@ public:
 		return WasCancelled;
 	}
 
+	friend class FProgressScope;
+
+	// Simple helper to track progress in a local scope on an optional FProgressCancel
+	// Will still work if the ProgressCancel is null (just does nothing in that case)
+	class GEOMETRYCORE_API FProgressScope
+	{
+		FProgressCancel* ProgressCancel;
+		FProgressData SavedProgressData;
+		bool bEnded = false;
+
+	public:
+
+		/**
+		 * @param ProgressCancel		Progress will be tracked on this. If null, the FProgressScope will do nothing.
+		 * @param ProgressAmount		Amount to increase progress w/in this scope (as a fraction of the current outer-scope active progress range)
+		 */
+		FProgressScope(FProgressCancel* ProgressCancel, float ProgressAmount, const FText& Message = FText());
+
+		/**
+		 * Create a dummy/inactive FProgressScope
+		 */
+		FProgressScope() : ProgressCancel(nullptr)
+		{}
+
+		~FProgressScope()
+		{
+			if (!bEnded)
+			{
+				Done();
+			}
+		}
+
+		/**
+		 * Advance to the end of the scope's progress range and close the scope
+		 */
+		void Done();
+
+		/**
+		 * @param Amount	Amount to increase the progress fraction, as a fraction of the current active progress range
+		 */
+		inline void AdvanceProgressBy(float Amount)
+		{
+			if (ProgressCancel)
+			{
+				ProgressCancel->AdvanceCurrentScopeProgressBy(Amount);
+			}
+		}
+
+		/**
+		 * Advance current progress a fraction of the way toward a target value
+		 * For example: if progress is .5, AdvanceProgressToward(1, .5) will take a half step to 1 and set progress to .75
+		 * As with all public progress function, progress is expressed relative to the current active progress range and cannot go backward.
+		 */
+		inline void AdvanceProgressToward(float TargetProgressFrac, float FractionToward)
+		{
+			if (ProgressCancel)
+			{
+				ProgressCancel->AdvanceCurrentScopeProgressToward(TargetProgressFrac, FractionToward);
+			}
+		}
+
+		/**
+		 * @return Amount to progress in the current scope to reach the target (scope-relative) progress value
+		 */
+		inline float GetDistanceTo(float TargetProgressFrac)
+		{
+			if (ProgressCancel)
+			{
+				return ProgressCancel->GetCurrentScopeDistanceTo(TargetProgressFrac);
+			}
+
+			return 0;
+		}
+
+		/**
+		 * @return Progress in the current scope
+		 */
+		inline float GetProgress()
+		{
+			if (ProgressCancel)
+			{
+				return ProgressCancel->GetCurrentScopeProgress();
+			}
+
+			return 1;
+		}
+
+		/**
+		 * Note: This function will leave the current progress unchanged if the target value is less than the current progress --
+		 *  it will not allow the progress value to go backward.
+		 *
+		 * @param ProgressFrac	Value to set current progress to, as a fraction of the current active progress range
+		 */
+		inline void SetProgressTo(float NewProgressFrac)
+		{
+			if (ProgressCancel)
+			{
+				ProgressCancel->SetCurrentScopeProgressTo(NewProgressFrac);
+			}
+		}
+	};
+
+	/**
+	 * @param ProgressCancel		Progress will be tracked on this. If null, the FProgressScope will do nothing.
+	 * @param ProgressTo			Target value to increase progress to w/in this scope (as a fraction of the current outer-scope active progress range)
+	 * @param Message				Optional message describing the work to be done
+	 * @return						A new FProgressScope that covers work from the current progress to the target progress value (relative to the current scope)
+	 */
+	static FProgressScope CreateScopeTo(FProgressCancel* ProgressCancel, float ProgressTo, const FText& Message = FText());
+
+	float GetProgress() const
+	{
+		return FMathf::Clamp(ProgressFraction, 0, 1);
+	}
+
+	FText GetProgressMessage() const
+	{
+		FScopeLock MessageLock(&MessageCS);
+		return ProgressMessage;
+	}
+
+	void SetProgressMessage(const FText& Message)
+	{
+		if (MaxMessageDepth == -1 || Progress.ScopeDepth < MaxMessageDepth)
+		{
+			Progress.Message = Message;
+			FScopeLock MessageLock(&MessageCS);
+			ProgressMessage = Message;
+		}
+	}
+
+	/**
+	 * @param Amount	Amount to increase the progress fraction, as a fraction of the current active progress range
+	 */
+	void AdvanceCurrentScopeProgressBy(float Amount)
+	{
+		ProgressFraction = Progress.SafeAdvance(ProgressFraction, Amount);
+	}
+
+	/**
+	 * Advance current progress a fraction of the way toward a target value
+	 * For example: if progress is .5, AdvanceProgressToward(1, .5) will take a half step to 1 and set progress to .75
+	 * As with all public progress function, progress is expressed relative to the current active progress range and cannot go backward.
+	 */
+	void AdvanceCurrentScopeProgressToward(float TargetProgressFrac, float FractionToward)
+	{
+		float TargetProgressFraction = FMathf::Max(ProgressFraction, Progress.CurrentMin + Progress.Range() * TargetProgressFrac);
+		ProgressFraction = Progress.SafeAdvance(ProgressFraction, FMathf::Lerp(ProgressFraction, TargetProgressFraction, FractionToward));
+	}
+
+	/**
+	 * @return Progress as a fraction of the current scope's active progress range
+	 */
+	float GetCurrentScopeProgress()
+	{
+		float Range = Progress.Range();
+		return Range > 0 ? (ProgressFraction - Progress.CurrentMin) / Range : 1;
+	}
+
+	/**
+	 * @param TargetProgressFrac	Target progress value as a fraction of the current scope's active progress range
+	 * @return Amount to progress in the current scope to reach the target progress value
+	 */
+	float GetCurrentScopeDistanceTo(float TargetProgressFrac)
+	{
+		return TargetProgressFrac - GetCurrentScopeProgress();
+	}
+
+	/**
+	 * Note: This function will leave the current progress unchanged if the target value is less than the current progress --
+	 *  it will not allow the progress value to go backward.
+	 * 
+	 * @param ProgressFrac	Value to set current progress to, as a fraction of the current active progress range
+	 */
+	void SetCurrentScopeProgressTo(float NewProgressFrac)
+	{
+		ProgressFraction = FMathf::Clamp(Progress.CurrentMin + Progress.Range() * NewProgressFrac, ProgressFraction, Progress.CurrentMax);
+	}
 
 public:
 
