fix convex hull's 'extreme points' algorithm not catching exactly collinear points as having dimension==1

diff --git a/Engine/Source/Runtime/GeometryCore/Private/CompGeom/ConvexHull3.cpp b/Engine/Source/Runtime/GeometryCore/Private/CompGeom/ConvexHull3.cpp
--- a/Engine/Source/Runtime/GeometryCore/Private/CompGeom/ConvexHull3.cpp
+++ b/Engine/Source/Runtime/GeometryCore/Private/CompGeom/ConvexHull3.cpp
@@ -118,8 +118,20 @@ void TExtremePoints3<RealType>::Init(int32 NumPoints, TFunctionRef<TVector<RealT
 	Basis[1] = GetPointFunc(Extreme[2]) - Origin;
 	// project Basis[1] to be orthogonal to Basis[0]
 	Basis[1] -= (Basis[0].Dot(Basis[1])) * Basis[0];
-	Normalize(Basis[1]);
+	if (!Basis[1].Normalize(Epsilon)) // points too collinear to form a valid basis
+	{
+		Dimension = 1;
+		Extreme[3] = Extreme[2] = Extreme[1];
+		return;
+	}
 	Basis[2] = Basis[0].Cross(Basis[1]);
+	if (!Basis[2].Normalize(Epsilon)) // points too collinear to form a valid basis
+	{
+		Dimension = 1;
+		Extreme[3] = Extreme[2] = Extreme[1];
+		return;
+	}
+
 
 	{
 		TPlane3<RealType> Plane(Basis[2], Origin);
diff --git a/Engine/Source/Runtime/GeometryCore/Public/CompGeom/ConvexHull3.h b/Engine/Source/Runtime/GeometryCore/Public/CompGeom/ConvexHull3.h
--- a/Engine/Source/Runtime/GeometryCore/Public/CompGeom/ConvexHull3.h
+++ b/Engine/Source/Runtime/GeometryCore/Public/CompGeom/ConvexHull3.h
@@ -33,13 +33,13 @@ struct GEOMETRYCORE_API TExtremePoints3
 	TVector<RealType> Origin{ 0,0,0 };
 	TVector<RealType> Basis[3]{ {0,0,0}, {0,0,0}, {0,0,0} };
 
-	TExtremePoints3(int32 NumPoints, TFunctionRef<TVector<RealType>(int32)> GetPointFunc, TFunctionRef<bool(int32)> FilterFunc = [](int32 Idx) {return true; }, double Epsilon = 0)
+	TExtremePoints3(int32 NumPoints, TFunctionRef<TVector<RealType>(int32)> GetPointFunc, TFunctionRef<bool(int32)> FilterFunc = [](int32 Idx) {return true; }, double Epsilon = TMathUtil<RealType>::Epsilon)
 	{
 		Init(NumPoints, GetPointFunc, FilterFunc, Epsilon);
 	}
 
 private:
-	void Init(int32 NumPoints, TFunctionRef<TVector<RealType>(int32)> GetPointFunc, TFunctionRef<bool(int32)> FilterFunc, double Epsilon = 0);
+	void Init(int32 NumPoints, TFunctionRef<TVector<RealType>(int32)> GetPointFunc, TFunctionRef<bool(int32)> FilterFunc, double Epsilon = TMathUtil<RealType>::Epsilon);
 };
 
 /**
