Improve WeightMap support in Modeling Mode. Duplicate Tool now preserves WeightMaps, conversion between Dynamic and Static meshes properly preserves weight maps and names.
GeometryCore: Add bClearExisting option to FDynamicMeshAttributeSet::EnableMatchingAttributes(). Default true is old behavior. New behavior, when false, is to preserve existing attribute set counts/values. Also now copying group and weight layer names if not set.
ModelingTools: CombineMeshesTool now uses EnableMatchingAttributes instead of only explicitly matching UV layers, so group and weight layers are copied properly to duplicates
MeshConversion: DynamicMeshToMeshDescription now generates unique names in ConvertPolygroupLayers and ConvertWeightLayers. Previously if multiple attribs had the same name (eg NAME_None), output MeshDescription would only have a single attrib w/ that name

diff --git a/Engine/Plugins/Runtime/MeshModelingToolset/Source/MeshModelingTools/Private/CombineMeshesTool.cpp b/Engine/Plugins/Runtime/MeshModelingToolset/Source/MeshModelingTools/Private/CombineMeshesTool.cpp
--- a/Engine/Plugins/Runtime/MeshModelingToolset/Source/MeshModelingTools/Private/CombineMeshesTool.cpp
+++ b/Engine/Plugins/Runtime/MeshModelingToolset/Source/MeshModelingTools/Private/CombineMeshesTool.cpp
@@ -215,9 +215,9 @@ void UCombineMeshesTool::CreateNewAsset()
 			FDynamicMesh3& ComponentDMesh = InputMeshes[ComponentIdx];
 			bNeedColorAttr = bNeedColorAttr || (ComponentDMesh.HasAttributes() && ComponentDMesh.Attributes()->HasPrimaryColors());
 
-			if (ComponentDMesh.HasAttributes() && ComponentDMesh.Attributes()->NumUVLayers() > AccumulateDMesh.Attributes()->NumUVLayers())
+			if (ComponentDMesh.HasAttributes())
 			{
-				AccumulateDMesh.Attributes()->SetNumUVLayers(ComponentDMesh.Attributes()->NumUVLayers());
+				AccumulateDMesh.Attributes()->EnableMatchingAttributes(*ComponentDMesh.Attributes(), false);
 			}
 
 			FTransformSRT3d XF = (UE::ToolTarget::GetLocalToWorldTransform(Targets[ComponentIdx]) * ToAccum);
diff --git a/Engine/Source/Runtime/GeometryCore/Private/DynamicMesh/DynamicMeshAttributeSet.cpp b/Engine/Source/Runtime/GeometryCore/Private/DynamicMesh/DynamicMeshAttributeSet.cpp
--- a/Engine/Source/Runtime/GeometryCore/Private/DynamicMesh/DynamicMeshAttributeSet.cpp
+++ b/Engine/Source/Runtime/GeometryCore/Private/DynamicMesh/DynamicMeshAttributeSet.cpp
@@ -274,48 +274,71 @@ void FDynamicMeshAttributeSet::SplitAllBowties(bool bParallel)
 
 
 
-void FDynamicMeshAttributeSet::EnableMatchingAttributes(const FDynamicMeshAttributeSet& ToMatch)
+void FDynamicMeshAttributeSet::EnableMatchingAttributes(const FDynamicMeshAttributeSet& ToMatch, bool bClearExisting)
 {
-	SetNumUVLayers(ToMatch.NumUVLayers());
-	for (int UVIdx = 0; UVIdx < NumUVLayers(); UVIdx++)
+	int32 ExistingUVLayers = NumUVLayers();
+	int32 RequiredUVLayers = (bClearExisting) ? ToMatch.NumUVLayers() : FMath::Max(ExistingUVLayers, ToMatch.NumUVLayers());
+	SetNumUVLayers(RequiredUVLayers);
+	for (int32 k = bClearExisting ? 0 : ExistingUVLayers; k < NumUVLayers(); k++)
 	{
-		UVLayers[UVIdx].ClearElements();
+		UVLayers[k].ClearElements();
 	}
-	SetNumNormalLayers(ToMatch.NumNormalLayers());
-	for (int NormalLayerIndex = 0; NormalLayerIndex < NumNormalLayers(); NormalLayerIndex++)
+
+	int32 ExistingNormalLayers = NumNormalLayers();
+	int32 RequiredNormalLayers = (bClearExisting) ? ToMatch.NumUVLayers() : FMath::Max(ExistingNormalLayers, ToMatch.NumNormalLayers());
+	SetNumNormalLayers(RequiredNormalLayers);
+	for (int32 k = bClearExisting ? 0 : ExistingNormalLayers; k < NumNormalLayers(); k++)
 	{
-		NormalLayers[NormalLayerIndex].ClearElements();
+		NormalLayers[k].ClearElements();
 	}
-	if (ToMatch.ColorLayer)
+
+	bool bWantColorLayer = ToMatch.HasPrimaryColors() || (this->HasPrimaryColors() && bClearExisting == false);
+	if (bClearExisting || bWantColorLayer == false)
 	{
-		EnablePrimaryColors();
+		DisablePrimaryColors();
 	}
-	else
+	if (bWantColorLayer)
 	{
-		DisablePrimaryColors();
+		EnablePrimaryColors();
 	}
 
-	if (ToMatch.MaterialIDAttrib)
+	bool bWantMaterialID = ToMatch.HasMaterialID() || (this->HasMaterialID() && bClearExisting == false);
+	if (bClearExisting || bWantMaterialID == false)
 	{
-		EnableMaterialID();
+		DisableMaterialID();
 	}
-	else
+	if (bWantMaterialID)
 	{
-		DisableMaterialID();
+		EnableMaterialID();
 	}
 
-	SetNumPolygroupLayers(ToMatch.NumPolygroupLayers());
-	for (int GroupIdx = 0; GroupIdx < NumPolygroupLayers(); ++GroupIdx)
+
+	int32 ExistingPolygroupLayers = NumPolygroupLayers();
+	int32 RequiredPolygroupLayers = (bClearExisting) ? ToMatch.NumPolygroupLayers() : FMath::Max(ExistingPolygroupLayers, ToMatch.NumPolygroupLayers());
+	SetNumPolygroupLayers(RequiredPolygroupLayers);
+	for (int32 k = bClearExisting ? 0 : ExistingPolygroupLayers; k < NumPolygroupLayers(); k++)
 	{
-		PolygroupLayers[GroupIdx].Initialize((int32)0);
+		PolygroupLayers[k].Initialize((int32)0);
+		if (PolygroupLayers[k].GetName() == NAME_None && k < ToMatch.NumPolygroupLayers())
+		{
+			PolygroupLayers[k].SetName( ToMatch.GetPolygroupLayer(k)->GetName() );
+		}
 	}
 
-	SetNumWeightLayers(ToMatch.NumWeightLayers());
-	for (int WeightLayerIdx = 0; WeightLayerIdx < NumWeightLayers(); ++WeightLayerIdx)
+	int32 ExistingWeightLayers = NumWeightLayers();
+	int32 RequiredWeightLayers = (bClearExisting) ? ToMatch.NumWeightLayers() : FMath::Max(ExistingWeightLayers, ToMatch.NumWeightLayers());
+	SetNumWeightLayers(RequiredWeightLayers);
+	for (int32 k = bClearExisting ? 0 : ExistingWeightLayers; k < NumWeightLayers(); k++)
 	{
-		WeightLayers[WeightLayerIdx].Initialize(0.0f);
+		WeightLayers[k].Initialize(0.0f);
+		if (WeightLayers[k].GetName() == NAME_None && k < ToMatch.NumWeightLayers())
+		{
+			WeightLayers[k].SetName( ToMatch.GetWeightLayer(k)->GetName() );
+		}
 	}
 
+	// currently SkinWeightAttributes and generic attributes do not respect bClearExisting, ie they are always cleared
+
 	ResetRegisteredAttributes();
 	
 	SkinWeightAttributes.Reset();
diff --git a/Engine/Source/Runtime/GeometryCore/Public/DynamicMesh/DynamicMeshAttributeSet.h b/Engine/Source/Runtime/GeometryCore/Public/DynamicMesh/DynamicMeshAttributeSet.h
--- a/Engine/Source/Runtime/GeometryCore/Public/DynamicMesh/DynamicMeshAttributeSet.h
+++ b/Engine/Source/Runtime/GeometryCore/Public/DynamicMesh/DynamicMeshAttributeSet.h
@@ -85,9 +85,13 @@ public:
 
 
 	/**
-	 * Enable the matching attributes and overlay layers as the reference Copy set, but do not copy any data across
+	 * Enable the matching attributes and overlay layers as the reference Copy set, but do not copy any data across.
+	 * By default, clears existing attributes, so that there will be an exact match
+	 * If bClearExisting is passed as false, attribute sets will grow if necessary, and new sets are cleared,
+	 * but existing attributes are not removed/cleared.
+	 * @note Currently SkinWeightAttributes and GenericAttributes are always fully cleared regardless of bClearExisting value
 	 */
-	void EnableMatchingAttributes(const FDynamicMeshAttributeSet& ToMatch);
+	void EnableMatchingAttributes(const FDynamicMeshAttributeSet& ToMatch, bool bClearExisting = true);
 
 	/** @return the parent mesh for this overlay */
 	const FDynamicMesh3* GetParentMesh() const { return ParentMesh; }
diff --git a/Engine/Source/Runtime/MeshConversion/Private/DynamicMeshToMeshDescription.cpp b/Engine/Source/Runtime/MeshConversion/Private/DynamicMeshToMeshDescription.cpp
--- a/Engine/Source/Runtime/MeshConversion/Private/DynamicMeshToMeshDescription.cpp
+++ b/Engine/Source/Runtime/MeshConversion/Private/DynamicMeshToMeshDescription.cpp
@@ -985,11 +985,19 @@ void FDynamicMeshToMeshDescription::ConvertPolygroupLayers(const FDynamicMesh3*
 
 	TAttributesSet<FTriangleID>& TriAttribsSet = MeshOut.TriangleAttributes();
 
+	TSet<FName> UniqueNames;
+	int32 UniqueIDGenerator = 0;
 	for (int32 li = 0; li < MeshIn->Attributes()->NumPolygroupLayers(); ++li)
 	{
 		const FDynamicMeshPolygroupAttribute* Polygroups = MeshIn->Attributes()->GetPolygroupLayer(li);
 		FName LayerName = Polygroups->GetName();
-	
+		while (LayerName == NAME_None || UniqueNames.Contains(LayerName))
+		{
+			FString BaseName = (Polygroups->GetName() == NAME_None) ? TEXT("Groups") : LayerName.ToString();
+			LayerName = FName( FString::Printf(TEXT("%s_%d"), *BaseName, UniqueIDGenerator++) );
+		}
+		UniqueNames.Add(LayerName);	
+
 		// Find existing attribute with the same name. If not found, create a new one.
 		TTriangleAttributesRef<int32> Attribute;
 		if (TriAttribsSet.HasAttribute(LayerName))
@@ -1024,10 +1032,18 @@ void FDynamicMeshToMeshDescription::ConvertWeightLayers(const FDynamicMesh3* Mes
 
 	TAttributesSet<FVertexID>& VertexAttribsSet = MeshOut.VertexAttributes();
 
+	TSet<FName> UniqueNames;
+	int32 UniqueIDGenerator = 0;
 	for (int32 LayerIndex = 0; LayerIndex < MeshIn->Attributes()->NumWeightLayers(); ++LayerIndex)
 	{
 		const FDynamicMeshWeightAttribute* Weights = MeshIn->Attributes()->GetWeightLayer(LayerIndex);
 		FName LayerName = Weights->GetName();
+		while (LayerName == NAME_None || UniqueNames.Contains(LayerName))
+		{
+			FString BaseName = (Weights->GetName() == NAME_None) ? TEXT("Weights") : LayerName.ToString();
+			LayerName = FName( FString::Printf(TEXT("%s_%d"), *BaseName, UniqueIDGenerator++) );
+		}
+		UniqueNames.Add(LayerName);
 
 		// Find existing attribute with the same name. If not found, create a new one.
 		TVertexAttributesRef<float> Attribute;
