Deprecate unused Timestamp and GetTimestamp function on TPolygon2 and TPolyline3

diff --git a/Engine/Source/Runtime/GeometryCore/Public/Polygon2.h b/Engine/Source/Runtime/GeometryCore/Public/Polygon2.h
--- a/Engine/Source/Runtime/GeometryCore/Public/Polygon2.h
+++ b/Engine/Source/Runtime/GeometryCore/Public/Polygon2.h
@@ -12,7 +12,6 @@
 #include "LineTypes.h"
 #include "MathUtil.h"
 #include "Intersection/IntrSegment2Segment2.h"
-#include "Util/DynamicVector.h"
 
 namespace UE
 {
@@ -23,41 +22,44 @@ using namespace UE::Math;
 
 /**
  * TPolygon2 is a 2D polygon represented as a list of Vertices.
- * 
- * @todo move operators
  */
 template<typename T>
 class TPolygon2
 {
+
 protected:
 	/** The list of vertices/corners of the polygon */
 	TArray<TVector2<T>> Vertices;
 
-	/** A counter that is incremented every time the polygon vertices are modified */
-	int Timestamp;
+	/** A counter that is incremented every time the polygon vertices are modified. */
+	UE_DEPRECATED(5.1, "Timestamps for TPolygon2 were not being used and will be removed in the future")
+	int Timestamp = 0;
 
 public:
 
-	TPolygon2() : Timestamp(0)
+	TPolygon2()
 	{
 	}
 
-	/**
-	 * Construct polygon that is a copy of another polygon
-	 */
-	TPolygon2(const TPolygon2& Copy) : Vertices(Copy.Vertices), Timestamp(Copy.Timestamp)
-	{
-	}
+	// Note: We need all 5 of these explicitly defaulted just because of the deprecated Timestamp member
+	// Once Timestamp is deleted, we can delete these as well
+	PRAGMA_DISABLE_DEPRECATION_WARNINGS
+	~TPolygon2() = default;
+	TPolygon2(const TPolygon2& Other) = default;
+	TPolygon2(TPolygon2&& Other) noexcept = default;
+	TPolygon2& operator=(const TPolygon2& Other) = default;
+	TPolygon2& operator=(TPolygon2&& Other) noexcept = default;
+	PRAGMA_ENABLE_DEPRECATION_WARNINGS
 
 	/**
 	 * Construct polygon with given list of vertices
 	 */
-	TPolygon2(const TArray<TVector2<T>>& VertexList) : Vertices(VertexList), Timestamp(0)
+	TPolygon2(const TArray<TVector2<T>>& VertexList) : Vertices(VertexList)
 	{
 	}
 
 	template<typename OtherVertexType>
-	TPolygon2(const TArray<OtherVertexType>& VertexList) : Timestamp(0)
+	TPolygon2(const TArray<OtherVertexType>& VertexList)
 	{
 		Vertices.Reserve(VertexList.Num());
 		for (const OtherVertexType& OtherVtx : VertexList)
@@ -70,7 +72,7 @@ public:
 	/**
 	 * Construct polygon with given indices into a vertex array
 	 */
-	TPolygon2(TArrayView<const TVector2<T>> VertexArray, TArrayView<const int32> VertexIndices) : Timestamp(0)
+	TPolygon2(TArrayView<const TVector2<T>> VertexArray, TArrayView<const int32> VertexIndices)
 	{
 		Vertices.SetNum(VertexIndices.Num());
 		for (int32 Idx = 0; Idx < VertexIndices.Num(); Idx++)
@@ -80,9 +82,12 @@ public:
 	}
 
 	/** @return the Timestamp for the polygon, which is updated every time the polygon is modified */
+	UE_DEPRECATED(5.1, "Timestamps for TPolygon2 were not being used and will be removed in the future")
 	int GetTimestamp() const 
 	{
+		PRAGMA_DISABLE_DEPRECATION_WARNINGS
 		return Timestamp;
+		PRAGMA_ENABLE_DEPRECATION_WARNINGS
 	}
 
 	/**
@@ -133,7 +138,7 @@ public:
 	void AppendVertex(const TVector2<T>& Position)
 	{
 		Vertices.Add(Position);
-		Timestamp++;
+		IncrementDeprecatedTimestamp();
 	}
 
 	/**
@@ -142,7 +147,7 @@ public:
 	void AppendVertices(const TArray<TVector2<T>>& NewVertices)
 	{
 		Vertices.Append(NewVertices);
-		Timestamp++;
+		IncrementDeprecatedTimestamp();
 	}
 
 	/**
@@ -151,7 +156,7 @@ public:
 	void Set(int VertexIndex, const TVector2<T>& Position)
 	{
 		Vertices[VertexIndex] = Position;
-		Timestamp++;
+		IncrementDeprecatedTimestamp();
 	}
 
 	/**
@@ -160,7 +165,7 @@ public:
 	void RemoveVertex(int VertexIndex)
 	{
 		Vertices.RemoveAt(VertexIndex);
-		Timestamp++;
+		IncrementDeprecatedTimestamp();
 	}
 
 	/**
@@ -169,7 +174,7 @@ public:
 	void SetVertices(const TArray<TVector2<T>>& NewVertices)
 	{
 		Vertices = NewVertices;
-		Timestamp++;
+		IncrementDeprecatedTimestamp();
 	}
 
 
@@ -183,7 +188,7 @@ public:
 		{
 			Swap(Vertices[VertexIndex], Vertices[j]);
 		}
-		Timestamp++;
+		IncrementDeprecatedTimestamp();
 	}
 
 
@@ -759,7 +764,7 @@ public:
 		{
 			Vertices[i] += Translate;
 		}
-		Timestamp++;
+		IncrementDeprecatedTimestamp();
 		return *this;
 	}
 
@@ -774,7 +779,7 @@ public:
 		{
 			Vertices[i] = Scale * (Vertices[i] - Origin) + Origin;
 		}
-		Timestamp++;
+		IncrementDeprecatedTimestamp();
 		return *this;
 	}
 
@@ -790,7 +795,7 @@ public:
 		{
 			Vertices[i] = TransformFunc(Vertices[i]);
 		}
-		Timestamp++;
+		IncrementDeprecatedTimestamp();
 		return *this;
 	}
 
@@ -825,7 +830,7 @@ public:
 			Vertices[k] = NewVertices[k];
 		}
 
-		Timestamp++;
+		IncrementDeprecatedTimestamp();
 	}
 
 
@@ -860,7 +865,7 @@ public:
 			Vertices[k] = NewVertices[k];
 		}
 
-		Timestamp++;
+		IncrementDeprecatedTimestamp();
 	}
 
 
@@ -1020,7 +1025,7 @@ public:
 			}
 		}
 
-		Timestamp++;
+		IncrementDeprecatedTimestamp();
 		return;
 	}
 
@@ -1086,7 +1091,7 @@ public:
 			iCur = iNext;
 		} while (iCur != 0);
 
-		Timestamp++;
+		IncrementDeprecatedTimestamp();
 	}
 
 
@@ -1156,6 +1161,16 @@ public:
 
 		return Circle;
 	}
+
+private:
+
+	// Note: this function will be removed when Timestamp is removed
+	inline void IncrementDeprecatedTimestamp()
+	{
+		PRAGMA_DISABLE_DEPRECATION_WARNINGS
+		Timestamp++;
+		PRAGMA_ENABLE_DEPRECATION_WARNINGS
+	}
 };
 
 typedef TPolygon2<double> FPolygon2d;
diff --git a/Engine/Source/Runtime/GeometryCore/Public/Polyline3.h b/Engine/Source/Runtime/GeometryCore/Public/Polyline3.h
--- a/Engine/Source/Runtime/GeometryCore/Public/Polyline3.h
+++ b/Engine/Source/Runtime/GeometryCore/Public/Polyline3.h
@@ -16,8 +16,6 @@ using namespace UE::Math;
 
 /**
  * TPolyline3 represents a 3D polyline stored as a list of Vertices.
- *
- * @todo move operators
  */
 template<typename T>
 class TPolyline3
@@ -27,48 +25,56 @@ protected:
 	TArray<TVector<T>> Vertices;
 
 	/** A counter that is incremented every time the polyline vertices are modified */
-	int Timestamp;
+	UE_DEPRECATED(5.1, "Timestamps for TPolyline3 were not being used and will be removed in the future")
+	int Timestamp = 0;
 
 public:
 
+	// Note: We need all 5 of these explicitly defaulted just because of the deprecated Timestamp member
+	// Once Timestamp is deleted, we can delete these as well
+	PRAGMA_DISABLE_DEPRECATION_WARNINGS
+	~TPolyline3() = default;
+	TPolyline3(const TPolyline3& Other) = default;
+	TPolyline3(TPolyline3&& Other) noexcept = default;
+	TPolyline3& operator=(const TPolyline3& Other) = default;
+	TPolyline3& operator=(TPolyline3&& Other) noexcept = default;
+	PRAGMA_ENABLE_DEPRECATION_WARNINGS
 
-	TPolyline3() : Timestamp(0)
-	{
-	}
 
-	/**
-	 * Construct polyline that is a copy of another polyline
-	 */
-	TPolyline3(const TPolyline3& Copy) : Vertices(Copy.Vertices), Timestamp(Copy.Timestamp)
+	TPolyline3()
 	{
 	}
 
 	/**
 	 * Construct polyline with given list of vertices
 	 */
-	TPolyline3(const TArray<TVector<T>>& VertexList) : Vertices(VertexList), Timestamp(0)
+	TPolyline3(const TArray<TVector<T>>& VertexList) : Vertices(VertexList)
 	{
 	}
 
 	/**
 	 * Construct a single-segment polyline
 	 */
-	TPolyline3(const TVector<T>& Point0, const TVector<T>& Point1) : Timestamp(0)
+	TPolyline3(const TVector<T>& Point0, const TVector<T>& Point1)
 	{
 		Vertices.Add(Point0);
 		Vertices.Add(Point1);
 	}
 
 	/** @return the Timestamp for the polyline, which is updated every time the polyline is modified */
+	UE_DEPRECATED(5.1, "Timestamps for TPolyline3 were not being used and will be removed in the future")
 	int GetTimestamp() const
 	{
+		PRAGMA_DISABLE_DEPRECATION_WARNINGS
 		return Timestamp;
+		PRAGMA_ENABLE_DEPRECATION_WARNINGS
 	}
 
 	/** Explicitly increment the Timestamp */
+	UE_DEPRECATED(5.1, "Timestamps for TPolyline3 were not being used and will be removed in the future")
 	void IncrementTimestamp()
 	{
-		Timestamp++;
+		IncrementDeprecatedTimestamp();
 	}
 
 	/**
@@ -135,7 +141,7 @@ public:
 	void Clear()
 	{
 		Vertices.Reset();
-		Timestamp++;
+		IncrementDeprecatedTimestamp();
 	}
 
 	/**
@@ -144,7 +150,7 @@ public:
 	void AppendVertex(const TVector<T>& Position)
 	{
 		Vertices.Add(Position);
-		Timestamp++;
+		IncrementDeprecatedTimestamp();
 	}
 
 	/**
@@ -153,7 +159,7 @@ public:
 	void AppendVertices(const TArray<TVector<T>>& NewVertices)
 	{
 		Vertices.Append(NewVertices);
-		Timestamp++;
+		IncrementDeprecatedTimestamp();
 	}
 
 	/**
@@ -168,7 +174,7 @@ public:
 		{
 			Vertices.Append( (TVector<T>)NewVertices[k] );
 		}
-		Timestamp++;
+		IncrementDeprecatedTimestamp();
 	}
 
 	/**
@@ -177,7 +183,7 @@ public:
 	void Set(int VertexIndex, const TVector<T>& Position)
 	{
 		Vertices[VertexIndex] = Position;
-		Timestamp++;
+		IncrementDeprecatedTimestamp();
 	}
 
 	/**
@@ -186,7 +192,7 @@ public:
 	void RemoveVertex(int VertexIndex)
 	{
 		Vertices.RemoveAt(VertexIndex);
-		Timestamp++;
+		IncrementDeprecatedTimestamp();
 	}
 
 	/**
@@ -200,7 +206,7 @@ public:
 		{
 			Vertices[k] = NewVertices[k];
 		}
-		Timestamp++;
+		IncrementDeprecatedTimestamp();
 	}
 
 
@@ -214,7 +220,7 @@ public:
 		{
 			Swap(Vertices[VertexIndex], Vertices[j]);
 		}
-		Timestamp++;
+		IncrementDeprecatedTimestamp();
 	}
 
 
@@ -465,7 +471,15 @@ public:
 		NewPolyline.Vertices[k] = Vertices[N];
 	}
 
+private:
 
+	// Note: this function will be removed when Timestamp is removed
+	inline void IncrementDeprecatedTimestamp()
+	{
+		PRAGMA_DISABLE_DEPRECATION_WARNINGS
+		Timestamp++;
+		PRAGMA_ENABLE_DEPRECATION_WARNINGS
+	}
 };
 
 typedef TPolyline3<double> FPolyline3d;
