GeometryCore: Fix BilinearSample to sample from the bounding pixel centers.

diff --git a/Engine/Source/Runtime/GeometryCore/Public/Image/ImageBuilder.h b/Engine/Source/Runtime/GeometryCore/Public/Image/ImageBuilder.h
--- a/Engine/Source/Runtime/GeometryCore/Public/Image/ImageBuilder.h
+++ b/Engine/Source/Runtime/GeometryCore/Public/Image/ImageBuilder.h
@@ -172,24 +172,52 @@ public:
 	template<typename ScalarType>
 	PixelType BilinearSample(const FVector2d& PixelCoords, const PixelType& InvalidValue) const
 	{
-		double X = PixelCoords.X;
-		double Y = PixelCoords.Y;
+		const double X = PixelCoords.X;
+		const double Y = PixelCoords.Y;
+		const double XMax = Dimensions.GetWidth();
+		const double YMax = Dimensions.GetHeight(); 
 
-		int X0 = (int)X, X1 = X0 + 1;
-		int Y0 = (int)Y, Y1 = Y0 + 1;
-
-		// make sure we are in range
-		if (X0 < 0 || X1 >= Dimensions.GetWidth() ||
-			Y0 < 0 || Y1 >= Dimensions.GetHeight())
+		if (X < 0 || X > XMax || Y < 0 || Y > YMax)
 		{
 			return InvalidValue;
 		}
 
-		// convert double coords to [0,1] range
-		double Ax = PixelCoords.X - (double)X0;
-		double Ay = PixelCoords.Y - (double)Y0;
-		double OneMinusAx = 1.0 - Ax;
-		double OneMinusAy = 1.0 - Ay;
+		// Compute nearest pixel edge intersection
+		const double XIntersect = FMath::RoundToDouble(X);
+		const double YIntersect = FMath::RoundToDouble(Y);
+
+		// Compute bounding pixel centers
+		double X0 = XIntersect - 0.5;
+		double X1 = XIntersect + 0.5;
+		double Y0 = YIntersect - 0.5;
+		double Y1 = YIntersect + 0.5;
+
+		// Clamp to min/max pixel centers in the range ([0.5, 0.5], [XMax-0.5, YMax-0.5])
+		// For edge intersections along a given axis boundary, the bounding pixel centers
+		// for that axis will be equivalent. This ensures that samples along an axis
+		// boundary will not interpolate.
+		if (X0 < 0.5)
+		{
+			X0 = 0.5;
+		}
+		if (X1 > XMax - 0.5)
+		{
+			X1 = XMax - 0.5;
+		}
+		if (Y0 < 0.5)
+		{
+			Y0 = 0.5;
+		}
+		if (Y1 > YMax - 0.5)
+		{
+			Y1 = YMax - 0.5;
+		}
+
+		// Convert double coords to [0,1] range
+		const double Ax = X - X0;
+		const double Ay = Y - Y0;
+		const double OneMinusAx = 1.0 - Ax;
+		const double OneMinusAy = 1.0 - Ay;
 
 		PixelType V00 = GetPixel(FVector2i(X0, Y0));
 		PixelType V10 = GetPixel(FVector2i(X1, Y0));
