deprecate timestamp on TDynamicPointSet3

diff --git a/Engine/Source/Runtime/GeometryCore/Public/DynamicMesh/DynamicPointSet3.h b/Engine/Source/Runtime/GeometryCore/Public/DynamicMesh/DynamicPointSet3.h
--- a/Engine/Source/Runtime/GeometryCore/Public/DynamicMesh/DynamicPointSet3.h
+++ b/Engine/Source/Runtime/GeometryCore/Public/DynamicMesh/DynamicPointSet3.h
@@ -49,6 +49,7 @@ protected:
 	TDynamicVector<RealType> Vertices;
 
 	/** The timestamp is incremented any time a function that modifies the point set is called */
+	UE_DEPRECATED(5.1, "Timestamps for TDynamicPointSet3 were not being used and will be removed in the future")
 	int Timestamp = 0;
 
 
@@ -63,6 +64,17 @@ protected:
 	//
 public:
 
+	// Note: We need explicit defaults because of the deprecated Timestamp member
+	// Once Timestamp is deleted, we can delete all of these except the constructor
+	PRAGMA_DISABLE_DEPRECATION_WARNINGS
+	TDynamicPointSet3() = default;
+	~TDynamicPointSet3() = default;
+	TDynamicPointSet3(const TDynamicPointSet3& Other) = default;
+	TDynamicPointSet3(TDynamicPointSet3&& Other) noexcept = default;
+	TDynamicPointSet3& operator=(const TDynamicPointSet3& Other) = default;
+	TDynamicPointSet3& operator=(TDynamicPointSet3&& Other) noexcept = default;
+	PRAGMA_ENABLE_DEPRECATION_WARNINGS
+
 	void SetExternallyManagedAttributes(TDynamicAttributeSetBase<TDynamicPointSet3<RealType>>* AttributeSet)
 	{
 		BaseAttributeSet = AttributeSet;
@@ -74,9 +86,6 @@ public:
 	}
 
 
-public:
-
-
 
 	//
 	// Indices
@@ -115,7 +124,7 @@ public:
 		Vertices.Clear();
 		VertexRefCounts = FRefCountVector();
 		BaseAttributeSet = nullptr;
-		Timestamp = 0;
+		IncrementDeprecatedTimestamp();
 	}
 
 
@@ -133,7 +142,7 @@ public:
 			GetBaseAttributeSet()->OnNewVertex(vid, false);
 		}
 
-		UpdateTimeStamp();
+		IncrementDeprecatedTimestamp();
 		return vid;
 	}
 
@@ -177,7 +186,7 @@ public:
 			GetBaseAttributeSet()->OnNewVertex(VertexID, true);
 		}
 
-		UpdateTimeStamp();
+		IncrementDeprecatedTimestamp();
 		return EMeshResult::Ok;
 	}
 
@@ -205,7 +214,7 @@ public:
 		Vertices[i] = vNewPos.X;
 		Vertices[i + 1] = vNewPos.Y;
 		Vertices[i + 2] = vNewPos.Z;
-		UpdateTimeStamp();
+		IncrementDeprecatedTimestamp();
 	}
 
 
@@ -245,9 +254,12 @@ public:
 	}
 
 	/** Timestamp is incremented any time any change is made to the point set */
+	UE_DEPRECATED(5.1, "Timestamps for TDynamicPointSet3 were not being used and will be removed in the future")
 	inline int GetTimestamp() const
 	{
+		PRAGMA_DISABLE_DEPRECATION_WARNINGS
 		return Timestamp;
+		PRAGMA_ENABLE_DEPRECATION_WARNINGS
 	}
 
 	//
@@ -309,19 +321,29 @@ public:
 			GetBaseAttributeSet()->OnRemoveVertex(VertexID);
 		}
 
-		UpdateTimeStamp();
+		IncrementDeprecatedTimestamp();
 		return EMeshResult::Ok;
 	}
 
 
 
 protected:
+
+	UE_DEPRECATED(5.1, "Timestamps for TDynamicPointSet3 were not being used and will be removed in the future")
 	void UpdateTimeStamp()
 	{
-		Timestamp++;
+		IncrementDeprecatedTimestamp();
 	}
 
+private:
 
+	// Note: this function will be removed when Timestamp is removed
+	inline void IncrementDeprecatedTimestamp()
+	{
+		PRAGMA_DISABLE_DEPRECATION_WARNINGS
+		Timestamp++;
+		PRAGMA_ENABLE_DEPRECATION_WARNINGS
+	}
 
 };
 
