GeometryCore: Add wrap tiling method support to TImageBuilder sampling.

diff --git a/Engine/Plugins/Experimental/GeometryScripting/Source/GeometryScriptingCore/Private/TextureMapFunctions.cpp b/Engine/Plugins/Experimental/GeometryScripting/Source/GeometryScriptingCore/Private/TextureMapFunctions.cpp
--- a/Engine/Plugins/Experimental/GeometryScripting/Source/GeometryScriptingCore/Private/TextureMapFunctions.cpp
+++ b/Engine/Plugins/Experimental/GeometryScripting/Source/GeometryScriptingCore/Private/TextureMapFunctions.cpp
@@ -41,21 +41,23 @@ void UGeometryScriptLibrary_TextureMapFunctions::SampleTexture2DAtUVPositions(
 	{
 		FVector2d UV = UVs[k];
 
-		// Adjust UV value and tile it. 
+		// Adjust UV value. 
 		UV = UV * SampleOptions.UVScale + SampleOptions.UVOffset;
+
+		FVector4f InterpValue;
 		if (SampleOptions.bWrap)
 		{
-			UV = UV - FVector2d(FMathd::Floor(UV.X), FMathd::Floor(UV.Y));
+			constexpr EImageTilingMethod TilingMethod = EImageTilingMethod::Wrap;
+			InterpValue = (SampleOptions.SamplingMethod == EGeometryScriptPixelSamplingMethod::Bilinear) ?
+				ImageData.BilinearSampleUV<double, TilingMethod>(UV, FVector4f::Zero()) : ImageData.NearestSampleUV<TilingMethod>(UV);
 		}
 		else
 		{
-			UV.X = FMathd::Clamp(UV.X, 0.0, 1.0);
-			UV.Y = FMathd::Clamp(UV.Y, 0.0, 1.0);
+			constexpr EImageTilingMethod TilingMethod = EImageTilingMethod::Clamp;
+			InterpValue = (SampleOptions.SamplingMethod == EGeometryScriptPixelSamplingMethod::Bilinear) ?
+				ImageData.BilinearSampleUV<double, TilingMethod>(UV, FVector4f::Zero()) : ImageData.NearestSampleUV<TilingMethod>(UV);
 		}
 
-		FVector4f InterpValue = (SampleOptions.SamplingMethod == EGeometryScriptPixelSamplingMethod::Bilinear) ?
-			ImageData.BilinearSampleUV<double>(UV, FVector4f::Zero()) : ImageData.NearestSampleUV(UV, FVector4f::Zero());
-
 		Colors[k] = (FLinearColor)InterpValue;
 	}
 
diff --git a/Engine/Source/Runtime/GeometryCore/Public/Image/ImageBuilder.h b/Engine/Source/Runtime/GeometryCore/Public/Image/ImageBuilder.h
--- a/Engine/Source/Runtime/GeometryCore/Public/Image/ImageBuilder.h
+++ b/Engine/Source/Runtime/GeometryCore/Public/Image/ImageBuilder.h
@@ -12,6 +12,12 @@ namespace UE
 namespace Geometry
 {
 
+enum class EImageTilingMethod
+{
+	Clamp,
+	Wrap
+};
+
 /**
  * TImageBuilder is used to create and populate a 2D image with a templated "pixel" type.
  */
@@ -168,96 +174,81 @@ public:
 
 	/**
 	 * Sample the image value at floating-point pixel coords with Bilinear interpolation
+	 * The pixel coords are expected to be in the [0,Width]x[0,Height] image pixel rectangle.
 	 */
-	template<typename ScalarType>
+	template<typename ScalarType, EImageTilingMethod TilingMethod=EImageTilingMethod::Clamp>
 	PixelType BilinearSample(const FVector2d& PixelCoords, const PixelType& InvalidValue) const
 	{
-		const double X = PixelCoords.X;
-		const double Y = PixelCoords.Y;
-		const double XMax = Dimensions.GetWidth();
-		const double YMax = Dimensions.GetHeight(); 
-
-		if (X < 0 || X > XMax || Y < 0 || Y > YMax)
+		const FVector2i Max = FVector2i(Dimensions.GetWidth(), Dimensions.GetHeight());
+		if (PixelCoords.X < 0.0 || PixelCoords.X > Max.X || PixelCoords.Y < 0.0 || PixelCoords.Y > Max.Y)
 		{
 			return InvalidValue;
 		}
 
-		// Compute nearest pixel edge intersection
-		const double XIntersect = FMath::RoundToDouble(X);
-		const double YIntersect = FMath::RoundToDouble(Y);
-
-		// Compute bounding pixel centers
-		double X0 = XIntersect - 0.5;
-		double X1 = XIntersect + 0.5;
-		double Y0 = YIntersect - 0.5;
-		double Y1 = YIntersect + 0.5;
-
-		// Clamp to min/max pixel centers in the range ([0.5, 0.5], [XMax-0.5, YMax-0.5])
-		// For edge intersections along a given axis boundary, the bounding pixel centers
-		// for that axis will be equivalent. This ensures that samples along an axis
-		// boundary will not interpolate.
-		if (X0 < 0.5)
-		{
-			X0 = 0.5;
-		}
-		if (X1 > XMax - 0.5)
-		{
-			X1 = XMax - 0.5;
-		}
-		if (Y0 < 0.5)
+		// Compute bounding pixel corners
+		const FVector2d XY0d(FMathd::Floor(PixelCoords.X - 0.5), FMathd::Floor(PixelCoords.Y - 0.5));
+		const FVector2d XY1d(FMathd::Floor(PixelCoords.X + 0.5), FMathd::Floor(PixelCoords.Y + 0.5));
+		FVector2i XY0 = FVector2i(XY0d);
+		FVector2i XY1 = FVector2i(XY1d);
+
+		// Compute interpolation factors in [0,1] range.
+		// This must be done before wrapped tiling to ensure the [0,1] range.
+		// Rather than transforming the two pixel corners (XY0,XY1), we shift
+		// our PixelCoords by half a pixel to account for pixel centers.
+		const FVector2d A = PixelCoords - FVector2d(0.5) - XY0d;
+		const FVector2d OneMinusA = FVector2d::One() - A;
+
+		if constexpr (TilingMethod == EImageTilingMethod::Wrap)
 		{
-			Y0 = 0.5;
+			XY0.X = XY0.X < 0 ? XY0.X + Max.X : XY0.X;
+			XY0.Y = XY0.Y < 0 ? XY0.Y + Max.Y : XY0.Y;
+			XY1.X = XY1.X >= Max.X ? 0 : XY1.X;
+			XY1.Y = XY1.Y >= Max.Y ? 0 : XY1.Y;
 		}
-		if (Y1 > YMax - 0.5)
+		else // TilingMethod == EImageTilingMethod::Clamp
 		{
-			Y1 = YMax - 0.5;
+			XY0.X = FMath::Max(0, XY0.X);
+			XY0.Y = FMath::Max(0, XY0.Y);
+			XY1.X = FMath::Min(Max.X - 1, XY1.X);
+			XY1.Y = FMath::Min(Max.Y - 1, XY1.Y);
 		}
 
-		// Convert double coords to [0,1] range
-		const double Ax = X - X0;
-		const double Ay = Y - Y0;
-		const double OneMinusAx = 1.0 - Ax;
-		const double OneMinusAy = 1.0 - Ay;
-
-		PixelType V00 = GetPixel(FVector2i(X0, Y0));
-		PixelType V10 = GetPixel(FVector2i(X1, Y0));
-		PixelType V01 = GetPixel(FVector2i(X0, Y1));
-		PixelType V11 = GetPixel(FVector2i(X1, Y1));
-
-		return V00 * (ScalarType)(OneMinusAx * OneMinusAy) +
-			   V01 * (ScalarType)(OneMinusAx * Ay) +
-			   V10 * (ScalarType)(Ax * OneMinusAy) +
-			   V11 * (ScalarType)(Ax * Ay);
+		PixelType V00 = GetPixel(XY0);
+		PixelType V10 = GetPixel(FVector2i(XY1.X, XY0.Y));
+		PixelType V01 = GetPixel(FVector2i(XY0.X, XY1.Y));
+		PixelType V11 = GetPixel(XY1);
+
+		return V00 * static_cast<ScalarType>(OneMinusA.X * OneMinusA.Y) +
+			   V01 * static_cast<ScalarType>(OneMinusA.X * A.Y) +
+			   V10 * static_cast<ScalarType>(A.X * OneMinusA.Y) +
+			   V11 * static_cast<ScalarType>(A.X * A.Y);
 	}
 
 
 	/**
 	 * Sample the image value at floating-point UV coords with Bilinear interpolation.
-	 * The UV coords are assumed to be in range [0,1]x[0,1], and that this maps to the [0,Width]x[0,Height] image pixel rectangle.
 	 */
-	template<typename ScalarType>
+	template<typename ScalarType, EImageTilingMethod TilingMethod=EImageTilingMethod::Clamp>
 	PixelType BilinearSampleUV(const FVector2d& UVCoords, const PixelType& InvalidValue) const
 	{
+		const FVector2d UV = GetTiledUV<TilingMethod>(UVCoords);
+		
 		FVector2d PixelCoords(
-			UVCoords.X * (double)Dimensions.GetWidth(),
-			UVCoords.Y * (double)Dimensions.GetHeight());
+			UV.X * (double)Dimensions.GetWidth(),
+			UV.Y * (double)Dimensions.GetHeight());
 
-		return BilinearSample<ScalarType>(PixelCoords, InvalidValue);
+		return BilinearSample<ScalarType, TilingMethod>(PixelCoords, InvalidValue);
 	}
 
 	/**
-	* Sample the image value at floating-point UV coords with Nearest sampling.
-	* The UV coords are assumed to be in range [0,1]x[0,1], and that this maps to the [0,Width]x[0,Height] image pixel rectangle.
-	*/
-	PixelType NearestSampleUV(const FVector2d& UVCoords, const PixelType& InvalidValue) const
-	{
-		int X = (int)(UVCoords.X * (double)Dimensions.GetWidth());
-		int Y = (int)(UVCoords.Y * (double)Dimensions.GetHeight());
-		if (X < 0 || X >= Dimensions.GetWidth() ||
-			Y < 0 || Y >= Dimensions.GetHeight())
-		{
-			return InvalidValue;
-		}
+	 * Sample the image value at floating-point UV coords with Nearest sampling.
+	 */
+	template<EImageTilingMethod TilingMethod=EImageTilingMethod::Clamp>
+	PixelType NearestSampleUV(const FVector2d& UVCoords) const
+	{
+		const FVector2d UV = GetTiledUV<TilingMethod>(UVCoords);
+		const int X = (int)(UV.X * (double)Dimensions.GetWidth());
+		const int Y = (int)(UV.Y * (double)Dimensions.GetHeight());
 		return GetPixel(X, Y);
 	}
 
@@ -333,6 +324,24 @@ public:
 
 		return DownsampleImage;
 	}
+
+protected:
+	template<EImageTilingMethod TilingMethod>
+	static FVector2d GetTiledUV(const FVector2d& UVCoords)
+	{
+		FVector2d UV(UVCoords);
+		if constexpr (TilingMethod == EImageTilingMethod::Wrap)
+		{
+			UV = UV - FVector2d(FMathd::Floor(UV.X), FMathd::Floor(UV.Y));
+		}
+		else // TilingMethod == EImageTilingMethod::Clamp
+		{
+			UV.X = FMathd::Clamp(UV.X, 0.0, 1.0);
+			UV.Y = FMathd::Clamp(UV.Y, 0.0, 1.0);
+		}
+		return UV;
+	}
+	
 };
 
 
