GeometryCore: DiTOrientedBox - fix bug that resulted in fallback to AABB in some cases.

diff --git a/Engine/Source/Runtime/GeometryCore/Private/CompGeom/DiTOrientedBox.cpp b/Engine/Source/Runtime/GeometryCore/Private/CompGeom/DiTOrientedBox.cpp
--- a/Engine/Source/Runtime/GeometryCore/Private/CompGeom/DiTOrientedBox.cpp
+++ b/Engine/Source/Runtime/GeometryCore/Private/CompGeom/DiTOrientedBox.cpp
@@ -288,7 +288,7 @@ namespace
 			const VectorType InitialOrigin(0, 0, 0);
 			const VectorType SetZ(en0);
 			const UE::Geometry::TFrame3<RealType> E0DirFrame(InitialOrigin, SetZ);
-			UE::Geometry::TOrientedBox3<RealType> OrientedBox(E0DirFrame, VectorType(-1));
+			UE::Geometry::TOrientedBox3<RealType> OrientedBox(E0DirFrame, VectorType(-TMathUtilConstants<RealType>::MaxReal));
 
 			// expand the best oriented box with all the points 
 			ExpandBoxToContain(OrientedBox, GetPointFunc, NumPoints);
@@ -370,7 +370,7 @@ namespace
 				{
 					const UE::Geometry::TFrame3<RealType> Frame(VectorType(0), en[i], FaceNormal, en[i].Cross(FaceNormal));
 					// create OrientedBox and expand to hold all the sample points
-					UE::Geometry::TOrientedBox3<RealType> OBox(Frame, VectorType(-1));
+					UE::Geometry::TOrientedBox3<RealType> OBox(Frame, VectorType(-TMathUtilConstants<RealType>::MaxReal));
 					ExpandBoxToContain(OBox, [&SamplePoints](int32 j){return SamplePoints[j];}, NumSamplePoints);
 
 					const RealType SurfaceArea = OBox.SurfaceArea();
