Refactor some Polygon2 functionality to be more broadly available via CurveUtil (e.g., functions for curve area, normals, length, winding)
+ add test coverage of these

diff --git a/Engine/Source/Runtime/GeometryCore/Public/Curve/CurveUtil.h b/Engine/Source/Runtime/GeometryCore/Public/Curve/CurveUtil.h
--- a/Engine/Source/Runtime/GeometryCore/Public/Curve/CurveUtil.h
+++ b/Engine/Source/Runtime/GeometryCore/Public/Curve/CurveUtil.h
@@ -20,25 +20,177 @@ using namespace UE::Math;
 /**
  * Curve utility functions
  */
-	template<typename RealType, typename VectorType>
-	TVector<RealType> Tangent(const TArrayView<const VectorType>& Vertices, int32 Idx, bool bLoop = false)
+
+	/**
+	 * Get (by reference) the vertices surrounding the given vertex index
+	 * If bLoop is false and the Prev or Next vertex would be out of bounds, the vertex at Idx is used instead.
+	 */
+	template<typename RealType, typename VectorType, bool bLoop>
+	inline void GetPrevNext(const TArrayView<const VectorType>& Vertices, int32 Idx, VectorType& OutPrev, VectorType& OutNext)
 	{
-		int32 EndIdx = Idx + 1;
-		int32 StartIdx = Idx - 1;
+		int32 NextIdx = Idx + 1;
+		int32 PrevIdx = Idx - 1;
 		int32 NV = Vertices.Num();
-		if (bLoop)
+		if constexpr (bLoop)
+		{
+			NextIdx = NextIdx % NV;
+			PrevIdx = (PrevIdx + NV) % NV;
+		}
+		else
+		{
+			NextIdx = FMath::Min(NV - 1, NextIdx);
+			PrevIdx = FMath::Max(0, PrevIdx);
+		}
+		OutNext = Vertices[NextIdx];
+		OutPrev = Vertices[PrevIdx];
+	}
+
+	/**
+	 * Get (by reference) vectors pointing toward the given vertex index, from its surrounding vertices
+	 * If bLoop is false and the Prev or Next vertex would be out of bounds, a zero vector is used instead
+	 */
+	template<typename RealType, typename VectorType, bool bLoop>
+	inline void GetVectorsToPrevNext(const TArrayView<const VectorType>& Vertices, int32 VertexIndex, VectorType& OutToPrev, VectorType& OutToNext, bool bNormalize)
+	{
+		GetPrevNext<RealType, VectorType, bLoop>(Vertices, VertexIndex, OutToPrev, OutToNext);
+		OutToNext -= Vertices[VertexIndex];
+		OutToPrev -= Vertices[VertexIndex];
+		if (bNormalize)
+		{
+			Normalize(OutToNext);
+			Normalize(OutToPrev);
+		}
+	}
+	
+	/**
+	 * @return the tangent direction (normalized) of the path at Idx, using the surrounding vertices (as found by GetPrevNext())
+	 */
+	template<typename RealType, typename VectorType, bool bLoop>
+	inline VectorType Tangent(const TArrayView<const VectorType>& Vertices, int32 Idx)
+	{
+		VectorType Prev, Next;
+		GetPrevNext<RealType, VectorType, bLoop>(Vertices, Idx, Prev, Next);
+		return Normalized(Next - Prev);
+	}
+
+	/**
+	 * @return the tangent direction (normalized) of the path at Idx, using the surrounding vertices (as found by GetPrevNext())
+	 */
+	template<typename RealType, typename VectorType>
+	inline VectorType Tangent(const TArrayView<const VectorType>& Vertices, int32 Idx, bool bLoop = false)
+	{
+		return bLoop ?
+			Tangent<RealType, VectorType, true>(Vertices, Idx) :
+			Tangent<RealType, VectorType, false>(Vertices, Idx);
+	}
+
+	/**
+	 * Construct a normal at a vertex of the Polygon by averaging the adjacent face normals. 
+	 * This vector is independent of the lengths of the adjacent segments.
+	 * Points "inward" for a Clockwise Polygon, and outward for CounterClockwise
+	 * Note: Specific to 2D curves; normals computed by rotating tangents in the XY plane
+	 */
+	template<typename RealType, typename VectorType, bool bLoop>
+	VectorType GetNormal_FaceAvg2(const TArrayView<const VectorType>& Vertices, int VertexIndex)
+	{
+		TVector2<RealType> ToPrev, ToNext;
+		GetVectorsToPrevNext<RealType, VectorType, bLoop>(Vertices, VertexIndex, ToPrev, ToNext, true);
+
+		TVector2<RealType> N = (PerpCW(ToNext) - PerpCW(ToPrev));
+		RealType Len = Normalize(N);
+		if (Len == 0)
 		{
-			EndIdx = EndIdx % NV;
-			StartIdx = (StartIdx + NV) % NV;
+			return Normalized(ToNext + ToPrev);   // this gives right direction for degenerate angle
 		}
 		else
 		{
-			EndIdx = FMath::Min(NV - 1, EndIdx);
-			StartIdx = FMath::Max(0, StartIdx);
+			return N;
+		}
+	}
+
+	/** 
+	 * Compute the signed area of a closed curve, assuming vertices in the XY plane
+	 */
+	template<typename RealType, typename VectorType>
+	RealType SignedArea2(const TArrayView<const VectorType>& Vertices)
+	{
+		RealType Area = 0;
+		int N = Vertices.Num();
+		if (N == 0)
+		{
+			return 0;
 		}
-		return Normalized(Vertices[EndIdx] - Vertices[StartIdx]);
+		for (int Idx = 0, PrevIdx = N-1; Idx < N; PrevIdx=Idx++)
+		{
+			const TVector2<RealType>& V1 = Vertices[PrevIdx];
+			const TVector2<RealType>& V2 = Vertices[Idx];
+			Area += V1.X * V2.Y - V1.Y * V2.X;
+		}
+		return Area * 0.5;
 	}
 
+	/**
+	 * Compute the winding of a point relative to a closed curve, assuming vertices in the XY plane
+	 */
+	template<typename RealType, typename VectorType>
+	RealType WindingIntegral2(const TArrayView<const VectorType>& Vertices, const VectorType& QueryPoint)
+	{
+		RealType Sum = 0;
+		int N = Vertices.Num();
+		VectorType A = Vertices[N-1] - QueryPoint, B = VectorType::Zero();
+		for (int Idx = 0; Idx < N; ++Idx)
+		{
+			B = Vertices[Idx] - QueryPoint;
+			// TODO: Consider computing closed curve winding w/out trig functions, i.e. see Contains2 below
+			Sum += TMathUtil<RealType>::Atan2(A.X * B.Y - A.Y * B.X, A.X * B.X + A.Y * B.Y);
+			A = B;
+		}
+		return Sum / FMathd::TwoPi;
+	}
+
+	/**
+	 * @return true if the given query point is inside the closed curve (in the XY plane), based on the winding integral
+	 * Note: Negative and positive windings are both considered 'inside'
+	 */
+	template<typename RealType, typename VectorType>
+	bool Contains2(const TArrayView<const VectorType>& Vertices, const VectorType& QueryPoint)
+	{
+		int WindingNumber = 0;
+
+		int N = Vertices.Num();
+		if (N == 0)
+		{
+			return false;
+		}
+		VectorType A = Vertices[N - 1], B = VectorType::Zero();
+		for (int Idx = 0; Idx < N; ++Idx)
+		{
+			B = Vertices[Idx];
+
+			if (A.Y <= QueryPoint.Y)     // y <= P.Y (below)
+			{
+				if (B.Y > QueryPoint.Y)									// an upward crossing
+				{
+					if (Orient(A, B, QueryPoint) > 0)  // P left of edge
+					{
+						++WindingNumber;                       // have a valid up intersect
+					}
+				}
+			}
+			else     // y > P.Y  (above)
+			{
+				if (B.Y <= QueryPoint.Y)									// a downward crossing
+				{
+					if (Orient(A, B, QueryPoint) < 0)  // P right of edge
+					{
+						--WindingNumber;						// have a valid down intersect
+					}
+				}
+			}
+			A = B;
+		}
+		return WindingNumber != 0;
+	}
 
 	template<typename RealType, typename VectorType>
 	static RealType ArcLength(const TArrayView<const VectorType>& Vertices, bool bLoop = false) {
@@ -59,11 +211,11 @@ using namespace UE::Math;
 	int FindNearestIndex(const TArrayView<const VectorType>& Vertices, const VectorType& V)
 	{
 		int iNearest = -1;
-		double dNear = TMathUtil<double>::MaxReal;
+		RealType dNear = TMathUtil<RealType>::MaxReal;
 		int N = Vertices.Num();
 		for ( int i = 0; i < N; ++i )
 		{
-			double dSqr = DistanceSquared(Vertices[i], V);
+			RealType dSqr = DistanceSquared(Vertices[i], V);
 			if (dSqr < dNear)
 			{
 				dNear = dSqr;
diff --git a/Engine/Source/Runtime/GeometryCore/Public/Polygon2.h b/Engine/Source/Runtime/GeometryCore/Public/Polygon2.h
--- a/Engine/Source/Runtime/GeometryCore/Public/Polygon2.h
+++ b/Engine/Source/Runtime/GeometryCore/Public/Polygon2.h
@@ -12,6 +12,8 @@
 #include "LineTypes.h"
 #include "MathUtil.h"
 #include "Intersection/IntrSegment2Segment2.h"
+#include "Curve/CurveUtil.h"
+#include "Algo/Reverse.h"
 
 namespace UE
 {
@@ -183,11 +185,7 @@ public:
 	 */
 	void Reverse()
 	{
-		int32 j = Vertices.Num()-1;
-		for (int32 VertexIndex = 0; VertexIndex < j; VertexIndex++, j--)
-		{
-			Swap(Vertices[VertexIndex], Vertices[j]);
-		}
+		Algo::Reverse(Vertices);
 		IncrementDeprecatedTimestamp();
 	}
 
@@ -198,9 +196,7 @@ public:
 	 */
 	TVector2<T> GetTangent(int VertexIndex) const
 	{
-		TVector2<T> next = Vertices[(VertexIndex + 1) % Vertices.Num()];
-		TVector2<T> prev = Vertices[VertexIndex == 0 ? Vertices.Num() - 1 : VertexIndex - 1];
-		return Normalized(next - prev);
+		return CurveUtil::Tangent<T, TVector2<T>, true>(Vertices, VertexIndex);
 	}
 
 
@@ -221,21 +217,7 @@ public:
 	 */
 	TVector2<T> GetNormal_FaceAvg(int VertexIndex) const
 	{
-		TVector2<T> next = Vertices[(VertexIndex + 1) % Vertices.Num()];
-		TVector2<T> prev = Vertices[VertexIndex == 0 ? Vertices.Num() - 1 : VertexIndex - 1];
-		next -= Vertices[VertexIndex]; Normalize(next);
-		prev -= Vertices[VertexIndex]; Normalize(prev);
-
-		TVector2<T> n = (PerpCW(next) - PerpCW(prev));
-		T len = Normalize(n);
-		if (len == 0) 
-		{
-			return Normalized(next + prev);   // this gives right direction for degenerate angle
-		}
-		else 
-		{
-			return n;
-		}
+		return CurveUtil::GetNormal_FaceAvg2<T, TVector2<T>, true>(Vertices, VertexIndex);
 	}
 
 
@@ -244,9 +226,7 @@ public:
 	 */
 	TAxisAlignedBox2<T> Bounds() const
 	{
-		TAxisAlignedBox2<T> box = TAxisAlignedBox2<T>::Empty();
-		box.Contain(Vertices);
-		return box;
+		return TAxisAlignedBox2<T>(Vertices);
 	}
 
 	
@@ -327,19 +307,7 @@ public:
 	 */
 	T SignedArea() const
 	{
-		T fArea = 0;
-		int N = Vertices.Num();
-		if (N == 0)
-		{
-			return 0;
-		}
-		for (int i = 0; i < N; ++i) 
-		{
-			const TVector2<T>& v1 = Vertices[i];
-			const TVector2<T>& v2 = Vertices[(i + 1) % N];
-			fArea += v1.X * v2.Y - v1.Y * v2.X;
-		}
-		return fArea * 0.5;
+		return CurveUtil::SignedArea2<T, TVector2<T>>(Vertices);
 	}
 
 	/**
@@ -355,40 +323,25 @@ public:
 	 */
 	T Perimeter() const
 	{
-		T fPerim = 0;
-		int N = Vertices.Num();
-		for (int i = 0; i < N; ++i)
-		{
-			fPerim += Distance(Vertices[i], Vertices[(i + 1) % N]);
-		}
-		return fPerim;
+		return CurveUtil::ArcLength<T, TVector2<T>>(Vertices, true);
 	}
 
 
 	/**
 	 * Get the previous and next vertex positions for a given vertex of the Polygon
 	 */
-	void NeighbourPoints(int iVertex, TVector2<T> &PrevNbrOut, TVector2<T> &NextNbrOut) const
+	void NeighbourPoints(int VertexIdx, TVector2<T>& OutPrevNbr, TVector2<T>& OutNextNbr) const
 	{
-		int N = Vertices.Num();
-		PrevNbrOut = Vertices[(iVertex == 0) ? N - 1 : iVertex - 1];
-		NextNbrOut = Vertices[(iVertex + 1) % N];
+		CurveUtil::GetPrevNext<T, TVector2<T>, true>(Vertices, VertexIdx, OutPrevNbr, OutNextNbr);
 	}
 
 
 	/**
 	 * Get the vectors from a given vertex to the previous and next Vertices, optionally normalized
 	 */
-	void NeighbourVectors(int iVertex, TVector2<T> &ToPrevOut, TVector2<T> &ToNextOut, bool bNormalize = false) const
+	void NeighbourVectors(int VertexIdx, TVector2<T>& OutToPrev, TVector2<T>& OutToNext, bool bNormalize = false) const
 	{
-		int N = Vertices.Num();
-		ToPrevOut = Vertices[(iVertex == 0) ? N - 1 : iVertex - 1] - Vertices[iVertex];
-		ToNextOut = Vertices[(iVertex + 1) % N] - Vertices[iVertex];
-		if (bNormalize) 
-		{
-			Normalize(ToPrevOut);
-			Normalize(ToNextOut);
-		}
+		CurveUtil::GetVectorsToPrevNext<T, TVector2<T>, true>(Vertices, VertexIdx, OutToPrev, OutToNext, bNormalize);
 	}
 
 
@@ -408,16 +361,7 @@ public:
 	 */
 	T WindingIntegral(const TVector2<T>& QueryPoint) const
 	{
-		T sum = 0;
-		int N = Vertices.Num();
-		TVector2<T> a = Vertices[0] - QueryPoint, b = TVector2<T>::Zero();
-		for (int i = 0; i < N; ++i) 
-		{
-			b = Vertices[(i + 1) % N] - QueryPoint;
-			sum += TMathUtil<T>::Atan2(a.X * b.Y - a.Y * b.X, a.X * b.X + a.Y * b.Y);
-			a = b;
-		}
-		return sum / FMathd::TwoPi;
+		return CurveUtil::WindingIntegral2<T, TVector2<T>>(Vertices, QueryPoint);
 	}
 
 
@@ -426,39 +370,7 @@ public:
 	 */
 	bool Contains(const TVector2<T>& QueryPoint) const
 	{
-		int nWindingNumber = 0;
-
-		int N = Vertices.Num();
-		if (N == 0)
-		{
-			return false;
-		}
-		TVector2<T> a = Vertices[0], b = TVector2<T>::Zero();
-		for (int i = 0; i < N; ++i) 
-		{
-			b = Vertices[(i + 1) % N];
-
-			if (a.Y <= QueryPoint.Y)     // y <= P.Y (below)
-			{
-				if (b.Y > QueryPoint.Y)									// an upward crossing
-				{
-					if (Orient(a, b, QueryPoint) > 0)  // P left of edge
-						++nWindingNumber;                       // have a valid up intersect
-				}
-			}
-			else     // y > P.Y  (above)
-			{
-				if (b.Y <= QueryPoint.Y)									// a downward crossing
-				{
-					if (Orient(a, b, QueryPoint) < 0)  // P right of edge
-					{
-						--nWindingNumber;						// have a valid down intersect
-					}
-				}
-			}
-			a = b;
-		}
-		return nWindingNumber != 0;
+		return CurveUtil::Contains2<T, TVector2<T>>(Vertices, QueryPoint);
 	}
 
 	/**
