Rule-of-5-or-0 fixes for some geometry types:
 - remove TGeneralPolygon2 and IntBoxTypes.h copy constructors, to allow everything to be autogenerated (especially so that a proper move should be generated for TGeneralPolygon2)
 - add explicit moves + delete copies for TSparseGrid3 and FProgressCancel::ProgressScope, as these have custom destructors and the default move/copy behavior would cause bugs
 - make voro++ move assignment protect vs self-move

diff --git a/Engine/Source/Runtime/GeometryCore/Public/Curve/GeneralPolygon2.h b/Engine/Source/Runtime/GeometryCore/Public/Curve/GeneralPolygon2.h
--- a/Engine/Source/Runtime/GeometryCore/Public/Curve/GeneralPolygon2.h
+++ b/Engine/Source/Runtime/GeometryCore/Public/Curve/GeneralPolygon2.h
@@ -42,13 +42,6 @@ public:
 	{
 	}
 
-	/**
-	* Construct a copy of another general polygon
-	*/
-	TGeneralPolygon2(const TGeneralPolygon2& Copy) : Outer(Copy.Outer), bOuterIsCW(Copy.bOuterIsCW), Holes(Copy.Holes)
-	{
-	}
-
 	/**
 	* Construct a general polygon with the given polygon as boundary
 	*/
diff --git a/Engine/Source/Runtime/GeometryCore/Public/IntBoxTypes.h b/Engine/Source/Runtime/GeometryCore/Public/IntBoxTypes.h
--- a/Engine/Source/Runtime/GeometryCore/Public/IntBoxTypes.h
+++ b/Engine/Source/Runtime/GeometryCore/Public/IntBoxTypes.h
@@ -86,8 +86,6 @@ struct FAxisAlignedBox2i
 	{
 	}
 
-	FAxisAlignedBox2i(const FAxisAlignedBox2i& OtherBox) = default;
-
 	FAxisAlignedBox2i(int32 SquareSize)
 		: Min((int32)0, (int32)0), Max(SquareSize, SquareSize)
 	{
@@ -152,8 +150,6 @@ struct FAxisAlignedBox3i
 		this->Max = Max;
 	}
 
-	FAxisAlignedBox3i(const FAxisAlignedBox3i& OtherBox) = default;
-
 	bool operator==(const FAxisAlignedBox3i& Other) const
 	{
 		return Max == Other.Max && Min == Other.Min;
diff --git a/Engine/Source/Runtime/GeometryCore/Public/Spatial/SparseGrid3.h b/Engine/Source/Runtime/GeometryCore/Public/Spatial/SparseGrid3.h
--- a/Engine/Source/Runtime/GeometryCore/Public/Spatial/SparseGrid3.h
+++ b/Engine/Source/Runtime/GeometryCore/Public/Spatial/SparseGrid3.h
@@ -50,6 +50,23 @@ public:
 	}
 
 
+	TSparseGrid3(const TSparseGrid3& Other) = delete;
+	TSparseGrid3& operator=(const TSparseGrid3& Other) = delete;
+	TSparseGrid3(TSparseGrid3&& Other) noexcept : Elements(MoveTemp(Other.Elements)), Bounds(MoveTemp(Other.Bounds))
+	{
+		Other.Elements.Reset();
+	}
+	TSparseGrid3& operator=(TSparseGrid3&& Other) noexcept
+	{
+		if (this != &Other)
+		{
+			Elements = MoveTemp(Other.Elements);
+			Bounds = MoveTemp(Other.Bounds);
+			Other.Elements.Reset();
+		}
+		return *this;
+	}
+
 	/**
 	 * @param Index an integer grid index
 	 * @return true if there is an allocated element at this index
diff --git a/Engine/Source/Runtime/GeometryCore/Public/Util/ProgressCancel.h b/Engine/Source/Runtime/GeometryCore/Public/Util/ProgressCancel.h
--- a/Engine/Source/Runtime/GeometryCore/Public/Util/ProgressCancel.h
+++ b/Engine/Source/Runtime/GeometryCore/Public/Util/ProgressCancel.h
@@ -271,6 +271,26 @@ public:
 		FProgressScope() : ProgressCancel(nullptr)
 		{}
 
+		FProgressScope(const FProgressScope& Other) = delete;
+		FProgressScope& operator=(const FProgressScope& Other) = delete;
+		FProgressScope(FProgressScope&& Other) noexcept : ProgressCancel(Other.ProgressCancel), SavedProgressData(MoveTemp(Other.SavedProgressData)), bEnded(Other.bEnded)
+		{
+			Other.ProgressCancel = nullptr;
+			Other.bEnded = true;
+		}
+		FProgressScope& operator=(FProgressScope&& Other) noexcept
+		{
+			if (this != &Other)
+			{
+				ProgressCancel = Other.ProgressCancel;
+				SavedProgressData = MoveTemp(Other.SavedProgressData);
+				bEnded = Other.bEnded;
+				Other.ProgressCancel = nullptr;
+				Other.bEnded = true;
+			}
+			return *this;
+		}
+
 		~FProgressScope()
 		{
 			if (!bEnded)
