GeometryCore: Fixing unsafe typecast warnings in preparation for enabling such warnings for GeometryCore

diff --git a/Engine/Source/Runtime/GeometryCore/Private/CompGeom/ConvexHull3.cpp b/Engine/Source/Runtime/GeometryCore/Private/CompGeom/ConvexHull3.cpp
--- a/Engine/Source/Runtime/GeometryCore/Private/CompGeom/ConvexHull3.cpp
+++ b/Engine/Source/Runtime/GeometryCore/Private/CompGeom/ConvexHull3.cpp
@@ -118,14 +118,14 @@ void TExtremePoints3<RealType>::Init(int32 NumPoints, TFunctionRef<TVector<RealT
 	Basis[1] = GetPointFunc(Extreme[2]) - Origin;
 	// project Basis[1] to be orthogonal to Basis[0]
 	Basis[1] -= (Basis[0].Dot(Basis[1])) * Basis[0];
-	if (!Basis[1].Normalize(Epsilon)) // points too collinear to form a valid basis
+	if (!Basis[1].Normalize(RealType(Epsilon))) // points too collinear to form a valid basis
 	{
 		Dimension = 1;
 		Extreme[3] = Extreme[2] = Extreme[1];
 		return;
 	}
 	Basis[2] = Basis[0].Cross(Basis[1]);
-	if (!Basis[2].Normalize(Epsilon)) // points too collinear to form a valid basis
+	if (!Basis[2].Normalize(RealType(Epsilon))) // points too collinear to form a valid basis
 	{
 		Dimension = 1;
 		Extreme[3] = Extreme[2] = Extreme[1];
@@ -142,7 +142,7 @@ void TExtremePoints3<RealType>::Init(int32 NumPoints, TFunctionRef<TVector<RealT
 			{
 				continue;
 			}
-			RealType DistSigned = Plane.DistanceTo(GetPointFunc(Idx));
+			RealType DistSigned = (RealType)Plane.DistanceTo(GetPointFunc(Idx));
 			RealType Dist = TMathUtil<RealType>::Abs(DistSigned);
 			if (Dist > MaxDist)
 			{
diff --git a/Engine/Source/Runtime/GeometryCore/Private/CompGeom/Delaunay2.cpp b/Engine/Source/Runtime/GeometryCore/Private/CompGeom/Delaunay2.cpp
--- a/Engine/Source/Runtime/GeometryCore/Private/CompGeom/Delaunay2.cpp
+++ b/Engine/Source/Runtime/GeometryCore/Private/CompGeom/Delaunay2.cpp
@@ -911,12 +911,12 @@ namespace DelaunayInternal
 			{
 				return false;
 			}
-			int32 SignA = FMath::Sign(OrientA);
+			int32 SignA = (int32)FMath::Sign(OrientA);
 			if (SignA >= 0)
 			{
 				return false;
 			}
-			int32 SignB = FMath::Sign(OrientBOut);
+			int32 SignB = (int32)FMath::Sign(OrientBOut);
 			// A properly oriented edge crossing the AB segment, on a tri that includes A, must go from the negative side to the positive side of AB
 			// (positive to negative would be behind the AB edge, and a zero would either be behind or would prevent the edge from being inserted)
 			return SignB == 1;
diff --git a/Engine/Source/Runtime/GeometryCore/Private/CompGeom/DiTOrientedBox.cpp b/Engine/Source/Runtime/GeometryCore/Private/CompGeom/DiTOrientedBox.cpp
--- a/Engine/Source/Runtime/GeometryCore/Private/CompGeom/DiTOrientedBox.cpp
+++ b/Engine/Source/Runtime/GeometryCore/Private/CompGeom/DiTOrientedBox.cpp
@@ -17,51 +17,51 @@ namespace
 	TArray<UE::Math::TVector<RealType>>  DiTODirections(const UE::Geometry::EDiTO DiTO)
 	{
 		typedef UE::Math::TVector<RealType>   VectorType;
-		
+		typedef RealType    T;
 
-		const RealType a = 0.5 * (FMath::Sqrt(5.) - 1);
+		const double a =  0.5 * (FMath::Sqrt(5.) - 1.);
 		switch (DiTO)
 		{
 			case UE::Geometry::EDiTO::DiTO_12:
 			{
 				// length Sqrt(5) * a  vectors
-				return TArray<VectorType>({ VectorType(0.,  1.,  a),
-										    VectorType(0.,  1., -a),
-											VectorType(1.,  a,   0.),
-											VectorType(1., -a,   0.),
-											VectorType(a,   0.,  1.),
-											VectorType(a,   0., -1.) });
+				return TArray<VectorType>({ VectorType(T(0),  T(1),  T(a)),
+										    VectorType(T(0),  T(1), -T(a)),
+											VectorType(T(1),  T(a),  T(0)),
+											VectorType(T(1), -T(a),  T(0)),
+											VectorType(T(a),  T(0),  T(1)),
+											VectorType(T(a),  T(0), -T(1)) });
 			}
 			break;
 
 			case  UE::Geometry::EDiTO::DiTO_14:
 			{
-				const RealType b = FMath::Sqrt(3.);
+				const double b = FMath::Sqrt(3.);
 				// length Sqrt(3) vectors
 				// note: in game engine gems vol 2.  b =  1. 
-				return TArray<VectorType>({ VectorType(b, 0., 0.),
-											VectorType(0., b, 0.),
-											VectorType(0., 0., b),
-											VectorType(1., 1., 1.),
-											VectorType(1., 1., -1.),
-											VectorType(1., -1., 1.),
-											VectorType(1., -1., -1.) });
+				return TArray<VectorType>({ VectorType(T(b), T(0), T(0)),
+											VectorType(T(0), T(b), T(0)),
+											VectorType(T(0), T(0), T(b)),
+											VectorType(T(1), T(1), T(1)),
+											VectorType(T(1), T(1),-T(1)),
+											VectorType(T(1),-T(1), T(1)),
+											VectorType(T(1),-T(1),-T(1)) });
 			}
 			break;
 
 			case UE::Geometry::EDiTO::DiTO_20:
 			{
 				// length Sqrt(3) vectors
-				return TArray<VectorType>({ VectorType(0.,         a,  1. + a),
-											VectorType(0.,         a,  -1.- a),
-											VectorType(a,     1. + a,      0.),
-											VectorType(a,    -1. - a,      0.),
-											VectorType(1. + a,     0.,      a),
-											VectorType(1. + a,     0.,     -a),
-											VectorType(1.,         1.,      1.),
-											VectorType(1.,         1.,     -1.),
-											VectorType(1.,        -1.,      1.),
-											VectorType(1.,        -1.,     -1.) });
+				return TArray<VectorType>({ VectorType(T(0),          T(a),  T(1. + a)),
+											VectorType(T(0),          T(a), -T(1. + a)),
+											VectorType(T(a),     T(1. + a),       T(0)),
+											VectorType(T(a),    -T(1. + a),       T(0)),
+											VectorType(T(1. + a),     T(0),       T(a)),
+											VectorType(T(1. + a),     T(0),      -T(a)),
+											VectorType(T(1),          T(1),       T(1)),
+											VectorType(T(1),          T(1),      -T(1)),
+											VectorType(T(1),         -T(1),       T(1)),
+											VectorType(T(1),         -T(1),      -T(1)) });
 			}
 			break;
 
@@ -69,21 +69,21 @@ namespace
 			{
 				// length Sqrt(3) vectors.
 				// note: in game engine gems vol 2.  b = c = 1. 
-				const RealType b = FMath::Sqrt(3.);
-				const RealType c = FMath::Sqrt(1.5);
-				return TArray<VectorType>({ VectorType(b,  0.,  0.),
-											VectorType(0.,  b,  0.),
-											VectorType(0.,  0., b),
-											VectorType(1.,  1.,  1.),
-											VectorType(1.,  1., -1.),
-											VectorType(1., -1.,  1.),
-											VectorType(1., -1., -1.),
-											VectorType(c,   c,  0.),
-											VectorType(c,  -c,  0.),
-											VectorType(c,   0.,  c),
-											VectorType(c,   0., -c),
-											VectorType(0.,   c,  c),
-											VectorType(0.,   c, -c) });
+				const double b = FMath::Sqrt(3.);
+				const double c = FMath::Sqrt(1.5);
+				return TArray<VectorType>({ VectorType(T(b),  T(0),  T(0)),
+											VectorType(T(0),  T(b),  T(0)),
+											VectorType(T(0),  T(0),  T(b)),
+											VectorType(T(1),  T(1),  T(1)),
+											VectorType(T(1),  T(1), -T(1)),
+											VectorType(T(1), -T(1),  T(1)),
+											VectorType(T(1), -T(1), -T(1)),
+											VectorType(T(c),  T(c),  T(0)),
+											VectorType(T(c), -T(c),  T(0)),
+											VectorType(T(c),  T(0),  T(c)),
+											VectorType(T(c),  T(0), -T(c)),
+											VectorType(T(0),  T(c),  T(c)),
+											VectorType(T(0),  T(c), -T(c)) });
 			}
 			break;
 
diff --git a/Engine/Source/Runtime/GeometryCore/Private/CompGeom/PolygonTriangulation.cpp b/Engine/Source/Runtime/GeometryCore/Private/CompGeom/PolygonTriangulation.cpp
--- a/Engine/Source/Runtime/GeometryCore/Private/CompGeom/PolygonTriangulation.cpp
+++ b/Engine/Source/Runtime/GeometryCore/Private/CompGeom/PolygonTriangulation.cpp
@@ -52,7 +52,7 @@ void PolygonTriangulation::TriangulateSimplePolygon(const TArray<TVector2<T>>& V
 	}
 
 	bool bIsClockwise = PolySignedArea2 < 0;
-	double OrientationSign = (bIsClockwise) ? -1.0 : 1.0;
+	T OrientationSign = (bIsClockwise) ? -T(1) : T(1);
 
 
 	OutTriangles.Reset();
@@ -200,8 +200,8 @@ static bool VectorsOnSameSide(const FVector& Vec, const FVector& A, const FVecto
 {
 	const FVector CrossA = FVector::CrossProduct(Vec, A);
 	const FVector CrossB = FVector::CrossProduct(Vec, B);
-	float DotWithEpsilon = SameSideDotProductEpsilon + FVector::DotProduct(CrossA, CrossB);
-	return !(DotWithEpsilon < 0.0f);
+	double DotWithEpsilon = (double)SameSideDotProductEpsilon + FVector::DotProduct(CrossA, CrossB);
+	return !(DotWithEpsilon < 0.0);
 }
 
 
diff --git a/Engine/Source/Runtime/GeometryCore/Private/DynamicMesh/Operations/MergeCoincidentMeshEdges.cpp b/Engine/Source/Runtime/GeometryCore/Private/DynamicMesh/Operations/MergeCoincidentMeshEdges.cpp
--- a/Engine/Source/Runtime/GeometryCore/Private/DynamicMesh/Operations/MergeCoincidentMeshEdges.cpp
+++ b/Engine/Source/Runtime/GeometryCore/Private/DynamicMesh/Operations/MergeCoincidentMeshEdges.cpp
@@ -122,8 +122,8 @@ bool FMergeCoincidentMeshEdges::Apply()
 				continue;
 			}
 		}
-
-		DuplicatesQueue.Insert(eid, EquivalenceSets[eid]->Num());
+		const float Priority = (float)EquivalenceSets[eid]->Num();
+		DuplicatesQueue.Insert(eid, Priority);
 	}
 
 	//
diff --git a/Engine/Source/Runtime/GeometryCore/Private/Generators/DiscMeshGenerator.cpp b/Engine/Source/Runtime/GeometryCore/Private/Generators/DiscMeshGenerator.cpp
--- a/Engine/Source/Runtime/GeometryCore/Private/Generators/DiscMeshGenerator.cpp
+++ b/Engine/Source/Runtime/GeometryCore/Private/Generators/DiscMeshGenerator.cpp
@@ -53,7 +53,9 @@ FMeshShapeGenerator& FDiscMeshGenerator::Generate()
 			int VertIdx = AngleIdx + 1 + (RadIdx-1) * AngleNV;
 			double R = RadIdx * IdxToRadiusScale;
 			Vertices[VertIdx] = MakeVertex(R*CosA, R*SinA);
-			UVs[VertIdx] = FVector2f(.5f + CosA * RadIdx * IdxToUVRadiusScale, .5f + SinA * RadIdx * IdxToUVRadiusScale);
+			double Uvalue = 0.5 + CosA * RadIdx * IdxToUVRadiusScale;
+			double Vvalue = 0.5 + SinA * RadIdx * IdxToUVRadiusScale;
+			UVs[VertIdx] = FVector2f(float(Uvalue), float(Vvalue));
 		}
 	}
 
@@ -153,7 +155,9 @@ FMeshShapeGenerator& FPuncturedDiscMeshGenerator::Generate()
 			int VertIdx = AngleIdx + RadIdx * AngleNV;
 			double R = RadIdx * IdxToRadiusScale + IdxToRadiusOffset;
 			Vertices[VertIdx] = MakeVertex(R*CosA, R*SinA);
-			UVs[VertIdx] = FVector2f(.5f + CosA * R * RadiusToUVScale, .5f + SinA * R * RadiusToUVScale);
+			double Uvalue = 0.5 + CosA * R * RadiusToUVScale;
+			double Vvalue = 0.5 + SinA * R * RadiusToUVScale;
+			UVs[VertIdx] = FVector2f(float(Uvalue), float(Vvalue));
 		}
 	}
 
diff --git a/Engine/Source/Runtime/GeometryCore/Private/Generators/PolygonEdgeMeshGenerator.cpp b/Engine/Source/Runtime/GeometryCore/Private/Generators/PolygonEdgeMeshGenerator.cpp
--- a/Engine/Source/Runtime/GeometryCore/Private/Generators/PolygonEdgeMeshGenerator.cpp
+++ b/Engine/Source/Runtime/GeometryCore/Private/Generators/PolygonEdgeMeshGenerator.cpp
@@ -449,11 +449,11 @@ FMeshShapeGenerator& FPolygonEdgeMeshGenerator::Generate()
 	{
 		if (UVWidth > UVHeight)
 		{
-			UVTop = UVHeight / UVWidth;
+			UVTop = float( UVHeight / UVWidth );
 		}
 		else
 		{
-			UVRight = UVWidth / UVHeight;
+			UVRight = float( UVWidth / UVHeight );
 		}
 	}
 	FVector2f UV00 = FVector2f(UVLeft, UVBottom);
@@ -472,10 +472,10 @@ FMeshShapeGenerator& FPolygonEdgeMeshGenerator::Generate()
 		const int32 NewVertexAIndex = 2 * NewVertexIndex;
 		const int32 NewVertexBIndex = NewVertexAIndex + 1;
 
-		float UParam = PartialArcLength(NewLeftSidePath, NewVertexIndex, MinSegmentArcLength) / TotalArcLength;
+		double UParam = PartialArcLength(NewLeftSidePath, NewVertexIndex, MinSegmentArcLength) / TotalArcLength;
 
-		UVs[NewVertexAIndex] = BilinearInterp(UV00, UV01, UV11, UV10, UParam, 0.0f);
-		UVs[NewVertexBIndex] = BilinearInterp(UV00, UV01, UV11, UV10, UParam, 1.0f);
+		UVs[NewVertexAIndex] = BilinearInterp(UV00, UV01, UV11, UV10, float(UParam), 0.0f);
+		UVs[NewVertexBIndex] = BilinearInterp(UV00, UV01, UV11, UV10, float(UParam), 1.0f);
 		UVParentVertex[NewVertexAIndex] = NewVertexAIndex;
 		UVParentVertex[NewVertexBIndex] = NewVertexBIndex;
 	}
diff --git a/Engine/Source/Runtime/GeometryCore/Private/Generators/RectangleMeshGenerator.cpp b/Engine/Source/Runtime/GeometryCore/Private/Generators/RectangleMeshGenerator.cpp
--- a/Engine/Source/Runtime/GeometryCore/Private/Generators/RectangleMeshGenerator.cpp
+++ b/Engine/Source/Runtime/GeometryCore/Private/Generators/RectangleMeshGenerator.cpp
@@ -42,11 +42,11 @@ FMeshShapeGenerator& FRectangleMeshGenerator::Generate()
 	{
 		if (Width > Height)
 		{
-			uvtop = Height / Width;
+			uvtop = float( Height / Width );
 		}
 		else
 		{
-			uvright = Width / Height;
+			uvright = float( Width / Height );
 		}
 	}
 
@@ -144,8 +144,8 @@ FMeshShapeGenerator& FRoundedRectangleMeshGenerator::Generate()
 		Normals[VertIdx] = Normal;
 	}
 
-	float TotWidth = FMathf::Max(FMathf::ZeroTolerance, Radius * 2 + Width);
-	float TotHeight = FMathf::Max(FMathf::ZeroTolerance, Radius * 2 + Height);
+	float TotWidth = FMathf::Max(FMathf::ZeroTolerance, float(Radius * 2 + Width));
+	float TotHeight = FMathf::Max(FMathf::ZeroTolerance, float(Radius * 2 + Height));
 	
 
 	// corner vertices
diff --git a/Engine/Source/Runtime/GeometryCore/Private/Generators/StairGenerator.cpp b/Engine/Source/Runtime/GeometryCore/Private/Generators/StairGenerator.cpp
--- a/Engine/Source/Runtime/GeometryCore/Private/Generators/StairGenerator.cpp
+++ b/Engine/Source/Runtime/GeometryCore/Private/Generators/StairGenerator.cpp
@@ -867,7 +867,7 @@ FVector3d FLinearStairGenerator::GenerateVertex(ESide Side, int VertexColumn, in
 		check(false);
 	}
 	float X = VertexColumn * StepDepth;
-	float Y = (Side == ESide::Right ? 0.5 * StepWidth : -0.5 * StepWidth);
+	float Y = (Side == ESide::Right ? float(0.5 * StepWidth) : -float(0.5 * StepWidth));
 	float Z = VertexRow * StepHeight;
 	return FVector3d(X, Y, Z);
 }
@@ -980,7 +980,7 @@ FVector3d FFloatingStairGenerator::GenerateVertex(ESide Side, int VertexColumn,
 		check(false);
 	}
 	float X = VertexColumn * StepDepth;
-	float Y = (Side == ESide::Right ? 0.5 * StepWidth : -0.5 * StepWidth);
+	float Y = (Side == ESide::Right ? float(0.5 * StepWidth) : -float(0.5 * StepWidth));
 	float Z = VertexColumn > 1 ? ((VertexColumn - 2) + VertexRow) * StepHeight : VertexRow * StepHeight;
 	return FVector3d(X, Y, Z);
 }
diff --git a/Engine/Source/Runtime/GeometryCore/Private/Generators/SweepGenerator.cpp b/Engine/Source/Runtime/GeometryCore/Private/Generators/SweepGenerator.cpp
--- a/Engine/Source/Runtime/GeometryCore/Private/Generators/SweepGenerator.cpp
+++ b/Engine/Source/Runtime/GeometryCore/Private/Generators/SweepGenerator.cpp
@@ -543,12 +543,12 @@ void FProfileSweepGenerator::InitializeUvBuffer(const TArray<int32>& VertPositio
 			if (WeldedVertices.Contains(ProfileIndex))
 			{
 				int32 VertIndex = GetVertIndex(true, SweepIndex, ProfileIndex, NumWelded, NumNonWelded, VertPositionOffsets);
-				SetUV(i * NumUvColumnsOut + j, FVector2f(WeldedUs[i], Vs[j]), VertIndex);
+				SetUV(i * NumUvColumnsOut + j, FVector2f((float)WeldedUs[i], (float)Vs[j]), VertIndex);
 			}
 			else
 			{
 				int32 VertIndex = GetVertIndex(false, SweepIndex, ProfileIndex, NumWelded, NumNonWelded, VertPositionOffsets);
-				SetUV(i * NumUvColumnsOut + j, FVector2f(RegularUs[i], Vs[j]), VertIndex);
+				SetUV(i * NumUvColumnsOut + j, FVector2f((float)RegularUs[i], (float)Vs[j]), VertIndex);
 			}
 		}
 	});
diff --git a/Engine/Source/Runtime/GeometryCore/Private/Parameterization/UVPacking.cpp b/Engine/Source/Runtime/GeometryCore/Private/Parameterization/UVPacking.cpp
--- a/Engine/Source/Runtime/GeometryCore/Private/Parameterization/UVPacking.cpp
+++ b/Engine/Source/Runtime/GeometryCore/Private/Parameterization/UVPacking.cpp
@@ -526,8 +526,8 @@ bool FStandardChartPacker::PackCharts(TArray<FUVIsland>& Charts, double UVScale,
 
 						if (StartPos)
 						{
-							RasterRect.X = StartPos->X;
-							RasterRect.Y = StartPos->Y;
+							RasterRect.X = FMath::FloorToInt32(StartPos->X);
+							RasterRect.Y = FMath::FloorToInt32(StartPos->Y);
 						}
 
 						bFound = LayoutRaster.FindWithSegments(RasterRect, ChartRaster, IsBestRect);
diff --git a/Engine/Source/Runtime/GeometryCore/Private/Parameterization/UVSpaceAllocator.cpp b/Engine/Source/Runtime/GeometryCore/Private/Parameterization/UVSpaceAllocator.cpp
--- a/Engine/Source/Runtime/GeometryCore/Private/Parameterization/UVSpaceAllocator.cpp
+++ b/Engine/Source/Runtime/GeometryCore/Private/Parameterization/UVSpaceAllocator.cpp
@@ -553,7 +553,7 @@ void FUVSpaceAllocator::CreateUsedSegments()
 	// Create segments along the X axis for each rows
 	for ( uint32 y = 0; y < RasterHeight; ++y )
 	{
-		SortedRowsIndex[y] = y;
+		SortedRowsIndex[y] = (uint16)y;
 
 		FRun& CurrentRow = Rows[y];
 		CurrentRow.LongestSegment = 0;
@@ -622,10 +622,10 @@ void FUVSpaceAllocator::CreateUsedSegments()
 
 	// Create segments along the Y axis for each columns
 	SortedColumnsIndex.SetNum(RasterWidth, false);
-
+	
 	for ( uint32 x = 0; x < RasterWidth; ++x )
 	{
-		SortedColumnsIndex[x] = x;
+		SortedColumnsIndex[x] = (uint16)x;
 
 		FRun& CurrentColumn = Columns[x];
 		CurrentColumn.LongestSegment = 0;
@@ -745,7 +745,7 @@ void FUVSpaceAllocator::MergeRun(FRun& ThisRun, const FRun& OtherRun, uint32 Rec
 		uint32 StopIndex  = Segment.StartPos + Segment.Length;
 		for (uint32 LookupIndex = LastIndex; LookupIndex < StopIndex; ++LookupIndex)
 		{
-			FreeSegmentsLookupPtr[LookupIndex] = SegmentIndex;
+			FreeSegmentsLookupPtr[LookupIndex] = uint16(SegmentIndex);
 		}
 
 		LastIndex = StopIndex;
@@ -753,7 +753,7 @@ void FUVSpaceAllocator::MergeRun(FRun& ThisRun, const FRun& OtherRun, uint32 Rec
 
 	for (uint32 Index = LastIndex; Index < PerpendicularResolution; ++Index)
 	{
-		FreeSegmentsLookupPtr[Index] = SegmentIndex;
+		FreeSegmentsLookupPtr[Index] = uint16(SegmentIndex);
 	}
 }
 
diff --git a/Engine/Source/Runtime/GeometryCore/Public/Generators/MarchingCubes.h b/Engine/Source/Runtime/GeometryCore/Public/Generators/MarchingCubes.h
--- a/Engine/Source/Runtime/GeometryCore/Public/Generators/MarchingCubes.h
+++ b/Engine/Source/Runtime/GeometryCore/Public/Generators/MarchingCubes.h
@@ -305,7 +305,7 @@ protected:
 	// Hash table for edge vertices
 	//
 
-	int64 NumEdgeVertexSections = 64;
+	const int64 NumEdgeVertexSections = 64;
 	TArray<TMap<int64, int>> EdgeVertexSections;
 	TArray<FCriticalSection> EdgeVertexSectionLocks;
 	
@@ -518,9 +518,9 @@ protected:
 	void InitHashTables()
 	{
 		EdgeVertexSections.Reset();
-		EdgeVertexSections.SetNum(NumEdgeVertexSections);
+		EdgeVertexSections.SetNum((int32)NumEdgeVertexSections);
 		EdgeVertexSectionLocks.Reset();
-		EdgeVertexSectionLocks.SetNum(NumEdgeVertexSections);
+		EdgeVertexSectionLocks.SetNum((int32)NumEdgeVertexSections);
 	}
 
 
@@ -911,14 +911,14 @@ protected:
 	 */
 	void ResetMesh()
 	{
-		VertexSectionLocks.SetNum(NumVertexSections);
+		VertexSectionLocks.SetNum((int32)NumVertexSections);
 		VertexSectionLists.Reset();
-		VertexSectionLists.SetNum(NumVertexSections);
+		VertexSectionLists.SetNum((int32)NumVertexSections);
 		VertexCounter = 0;
 
-		TriangleSectionLocks.SetNum(NumTriangleSections);
+		TriangleSectionLocks.SetNum((int32)NumTriangleSections);
 		TriangleSectionLists.Reset();
-		TriangleSectionLists.SetNum(NumTriangleSections);
+		TriangleSectionLists.SetNum((int32)NumTriangleSections);
 	}
 
 	/**
diff --git a/Engine/Source/Runtime/GeometryCore/Public/Generators/SweepGenerator.h b/Engine/Source/Runtime/GeometryCore/Public/Generators/SweepGenerator.h
--- a/Engine/Source/Runtime/GeometryCore/Public/Generators/SweepGenerator.h
+++ b/Engine/Source/Runtime/GeometryCore/Public/Generators/SweepGenerator.h
@@ -179,7 +179,7 @@ protected:
 						SetTrianglePolygon(TriIdx, PolyIdx);
 						TriIdx++;
 					}
-					float SideScale = 2 * CapIdx - 1;
+					float SideScale = float(2 * CapIdx - 1);
 					for (int32 Idx = 0; Idx < XVerts; Idx++)
 					{
 						FVector2f CenteredVert = (FVector2f)CrossSection.GetVertices()[Idx] * CapUVScale + CapUVOffset;
@@ -251,10 +251,10 @@ protected:
 			int32 NextDupVertIdx = UVSection < NumSections ? UVSections[UVSection] : -1;
 			for (int32 VertSubIdx = 0; VertSubIdx < XVerts; UVSubIdx++)
 			{
-				float UVX = bEvenlySpaceUVs ? CrossSectionPercentages[VertSubIdx] : VertSubIdx / float(XVerts);
+				float UVX = bEvenlySpaceUVs ? float(CrossSectionPercentages[VertSubIdx]) : VertSubIdx / float(XVerts);
 				for (int32 XIdx = 0; XIdx < NumCrossSections; XIdx++)
 				{
-					float UVY = bEvenlySpaceUVs ? PathPercentages[XIdx] : XIdx / float(NumCrossSections - 1);
+					float UVY = bEvenlySpaceUVs ? float(PathPercentages[XIdx]) : XIdx / float(NumCrossSections - 1);
 					SetUV(XIdx * XUVs + UVSubIdx, FVector2f(1-UVX, 1-UVY) * SectionsUVScale, (XIdx % CrossSectionsMod) * XVerts + VertSubIdx);
 				}
 
@@ -286,7 +286,7 @@ protected:
 				int32 VertSubIdx = 0;
 				for (int32 XIdx = 0; XIdx < NumCrossSections; XIdx++)
 				{
-					float UVY = bEvenlySpaceUVs ? PathPercentages[XIdx] : XIdx / float(NumCrossSections - 1);
+					float UVY = bEvenlySpaceUVs ? float(PathPercentages[XIdx]) : XIdx / float(NumCrossSections - 1);
 					SetUV(XIdx * XUVs + UVSubIdx, FVector2f(1-UVX, 1-UVY) * SectionsUVScale, (XIdx % CrossSectionsMod) * XVerts + VertSubIdx);
 				}
 			}
@@ -379,8 +379,8 @@ public:
 		AlongPercents[0] = 0;
 		for (int XIdx = 0; XIdx+1 < NumX; XIdx++)
 		{
-			float Dist = Distance( FVector2d(Radii[XIdx], Heights[XIdx]), FVector2d(Radii[XIdx + 1], Heights[XIdx + 1]) );
-			LenAlong += Dist;
+			double Dist = Distance( FVector2d(Radii[XIdx], Heights[XIdx]), FVector2d(Radii[XIdx + 1], Heights[XIdx + 1]) );
+			LenAlong += float(Dist);
 			AlongPercents[XIdx + 1] = LenAlong;
 		}
 		for (int XIdx = 0; XIdx+1 < NumX; XIdx++)
@@ -448,19 +448,19 @@ public:
 					// write sharp normals
 					if (ensure(XIdx > 0)) // very first index cannot be sharp
 					{
-						Normals[SubIdx + NormalXIdx * AngleSamples] = FVector3f(XVerts[SubIdx].X*NormalSides[XIdx-1].X, XVerts[SubIdx].Y*NormalSides[XIdx-1].X, NormalSides[XIdx-1].Y);
+						Normals[SubIdx + NormalXIdx * AngleSamples] = FVector3f(float(XVerts[SubIdx].X*NormalSides[XIdx-1].X), float(XVerts[SubIdx].Y*NormalSides[XIdx-1].X), float(NormalSides[XIdx-1].Y));
 					}
 					NormalXIdx++;
 					if (ensure(XIdx + 1 < NumX)) // very last index cannot be sharp
 					{
-						Normals[SubIdx + NormalXIdx * AngleSamples] = FVector3f(XVerts[SubIdx].X*NormalSides[XIdx].X, XVerts[SubIdx].Y*NormalSides[XIdx].X, NormalSides[XIdx].Y);
+						Normals[SubIdx + NormalXIdx * AngleSamples] = FVector3f(float(XVerts[SubIdx].X*NormalSides[XIdx].X), float(XVerts[SubIdx].Y*NormalSides[XIdx].X), float(NormalSides[XIdx].Y));
 					}
 					SharpNormalIdx++;
 				}
 				else
 				{
 					// write smoothed normal
-					Normals[SubIdx + NormalXIdx * AngleSamples] = FVector3f(XVerts[SubIdx].X*SmoothedNormalSides[XIdx].X, XVerts[SubIdx].Y*SmoothedNormalSides[XIdx].X, SmoothedNormalSides[XIdx].Y);
+					Normals[SubIdx + NormalXIdx * AngleSamples] = FVector3f(float(XVerts[SubIdx].X*SmoothedNormalSides[XIdx].X), float(XVerts[SubIdx].Y*SmoothedNormalSides[XIdx].X), float(SmoothedNormalSides[XIdx].Y));
 				}
 			}
 		}
@@ -481,14 +481,14 @@ public:
 			{
 				for (int XBotTop = 0; XBotTop < 2; ++XBotTop)
 				{
-					Normals[CapNormalStart[XBotTop] + SubIdx] = FVector3f(0, 0, 2 * XBotTop - 1);
+					Normals[CapNormalStart[XBotTop] + SubIdx] = FVector3f(0.f, 0.f, float(2 * XBotTop - 1));
 				}
 			}
 			for (int CapIdx = 0; CapIdx < 2; CapIdx++)
 			{
 				if (Caps[CapIdx] == ECapType::FlatMidpointFan)
 				{
-					Normals[CapNormalStart[CapIdx] + X.VertexCount()] = FVector3f(0, 0, 2 * CapIdx - 1);
+					Normals[CapNormalStart[CapIdx] + X.VertexCount()] = FVector3f(0.f, 0.f, float(2 * CapIdx - 1));
 				}
 			}
 		}
@@ -620,9 +620,9 @@ public:
 				break;
 			}
 
-			for (float ExtraSeg = 1, NumExtraSegs = AdditionalLengthSamples[SegIdx] + 1; ExtraSeg < NumExtraSegs; ++ExtraSeg)
+			for (int ExtraSeg = 1, NumExtraSegs = AdditionalLengthSamples[SegIdx] + 1; ExtraSeg < NumExtraSegs; ++ExtraSeg)
 			{
-				const float Along = ExtraSeg / NumExtraSegs;
+				const float Along = ExtraSeg / float(NumExtraSegs);
 				SetVert(FMath::Lerp(SrcRadii[SegIdx], SrcRadii[SegIdx + 1], Along),
 						FMath::Lerp(SrcHeights[SegIdx], SrcHeights[SegIdx + 1], Along));
 			}
@@ -689,12 +689,12 @@ public:
 
 		FAxisAlignedBox2f Bounds = (FAxisAlignedBox2f)CrossSection.Bounds();
 		double BoundsMaxDimInv = 1.0 / FMathd::Max(Bounds.MaxDim(), .001);
-		FVector2f SectionScale(1, 1), CapScale(BoundsMaxDimInv, BoundsMaxDimInv);
+		FVector2f SectionScale(1.f, 1.f), CapScale((float)BoundsMaxDimInv, (float)BoundsMaxDimInv);
 		if (bUVScaleRelativeWorld)
 		{
 			double Perimeter = CrossSection.Perimeter();
-			SectionScale.X = Perimeter / UnitUVInWorldCoordinates;
-			SectionScale.Y = TotalPathArcLength / UnitUVInWorldCoordinates;
+			SectionScale.X = float( Perimeter / UnitUVInWorldCoordinates );
+			SectionScale.Y = float( TotalPathArcLength / UnitUVInWorldCoordinates);
 			CapScale.X = CapScale.Y = 1.0f / UnitUVInWorldCoordinates;
 		}
 		ConstructMeshTopology(CrossSection, {}, {}, {}, true, Path, PathNum + (bLoop ? 1 : 0), bLoop, Caps, SectionScale, CapScale, Bounds.Center());
