add a ClipConvex method for polygon2 and curveutil, to clip polygons to given bounds.

Uses the Sutherland–Hodgman algorithm of progressively clipping each side, so may leave folded overlaps directly on the boundaries if applied to non-convex polygons (but otherwise should still work)

diff --git a/Engine/Source/Runtime/GeometryCore/Public/Curve/CurveUtil.h b/Engine/Source/Runtime/GeometryCore/Public/Curve/CurveUtil.h
--- a/Engine/Source/Runtime/GeometryCore/Public/Curve/CurveUtil.h
+++ b/Engine/Source/Runtime/GeometryCore/Public/Curve/CurveUtil.h
@@ -225,6 +225,79 @@ using namespace UE::Math;
 		return iNearest;
 	}
 
+	/**
+	 * Use the Sutherland–Hodgman algorithm to clip the vertices to the given bounds
+	 * Note if the path/polygon is concave, this may leave overlapping edges at the boundary.
+	 * Note this modifies and re-sizes the input array, so is written for TArray instead of TArrayView
+	 * @tparam bLoop	Whether the path is a loop
+	 * @tparam ClipDims	How many dimensions of Min and Max to use for clipping
+	 *					(e.g. if ClipDims is 2 and the vectors are 3D, clips only vs a rectangle in X and Y)
+	 */
+	template<typename RealType, typename VectorType, bool bLoop = true, int ClipDims = 2>
+	static void ClipConvexToBounds(TArray<VectorType>& Vertices, VectorType Min, VectorType Max)
+	{
+		TArray<VectorType> Clipped;
+		Clipped.Reserve(Vertices.Num());
+		VectorType SidePt = Min;
+		for (int32 SideIdx = 0; SideIdx < 2; ++SideIdx, SidePt = Max)
+		{
+			int32 SideSign = -SideIdx * 2 + 1;
+			for (int32 ClipDim = 0; ClipDim < ClipDims; ++ClipDim)
+			{
+				RealType ClipCoord = SidePt[ClipDim];
+				int32 VertNum = Vertices.Num();
+				int32 StartCur, StartPrev;
+				if constexpr (bLoop)
+				{
+					StartCur = 0;
+					StartPrev = VertNum - 1;
+				}
+				else
+				{
+					StartCur = 1;
+					StartPrev = 0;
+				}
+				RealType PrevDist = (Vertices[StartPrev][ClipDim] - ClipCoord) * SideSign;
+				if constexpr (!bLoop)
+				{
+					if (PrevDist >= 0)
+					{
+						Clipped.Add(Vertices[0]);
+					}
+				}
+				for (int32 CurIdx = StartCur, PrevIdx = StartPrev; CurIdx < VertNum; PrevIdx = CurIdx++)
+				{
+					RealType CurDist = (Vertices[CurIdx][ClipDim] - ClipCoord) * SideSign;
+					if (CurDist >= 0)
+					{
+						if (PrevDist < 0 && CurDist > 0)
+						{
+							RealType T = CurDist / (CurDist - PrevDist);
+							VectorType LerpVec = Lerp(Vertices[CurIdx], Vertices[PrevIdx], T);
+							LerpVec[ClipDim] = ClipCoord; // snap to exact bounds
+							Clipped.Add(LerpVec);
+						}
+						Clipped.Add(Vertices[CurIdx]);
+					}
+					else if (PrevDist > 0)
+					{
+						RealType T = CurDist / (CurDist - PrevDist);
+						VectorType LerpVec = Lerp(Vertices[CurIdx], Vertices[PrevIdx], T);
+						LerpVec[ClipDim] = ClipCoord; // snap to exact bounds
+						Clipped.Add(LerpVec);
+					}
+					PrevDist = CurDist;
+				}
+				Swap(Vertices, Clipped);
+				if (Vertices.IsEmpty())
+				{
+					return;
+				}
+				Clipped.Reset();
+			}
+		}
+	}
+
 
 	/**
 	 * smooth vertices in-place (will not produce a symmetric result, but does not require extra buffer)
diff --git a/Engine/Source/Runtime/GeometryCore/Public/Polygon2.h b/Engine/Source/Runtime/GeometryCore/Public/Polygon2.h
--- a/Engine/Source/Runtime/GeometryCore/Public/Polygon2.h
+++ b/Engine/Source/Runtime/GeometryCore/Public/Polygon2.h
@@ -459,6 +459,14 @@ public:
 		}
 		return true;
 	}
+	
+	/**
+	 * Clip polygon to the given bounds. Note that if the polygon is not convex, it may have overlapping edges at the bounds.
+	 */
+	void ClipConvex(const TAxisAlignedBox2<T>& Bounds)
+	{
+		CurveUtil::ClipConvexToBounds<T, TVector2<T>, true, 2>(Vertices, Bounds.Min, Bounds.Max);
+	}
 
 
 	/**
