add an IsConvex() test to polygon2/curveutil

diff --git a/Engine/Source/Runtime/GeometryCore/Public/Curve/CurveUtil.h b/Engine/Source/Runtime/GeometryCore/Public/Curve/CurveUtil.h
--- a/Engine/Source/Runtime/GeometryCore/Public/Curve/CurveUtil.h
+++ b/Engine/Source/Runtime/GeometryCore/Public/Curve/CurveUtil.h
@@ -298,6 +298,62 @@ using namespace UE::Math;
 		}
 	}
 
+	/**
+	 * Tests closed, 2D curve for convexity, with an optional tolerance allowing for approximately-collinear points
+	 * Note that tolerance is per vertex angle, so a very-finely-sampled smooth concavity could be reported as convex
+	 * 
+	 * @param RadiansTolerance		Maximum turn in the 'wrong' direction that will still be considered convex (in radians)
+	 * @param bDegenerateIsConvex	What to return for degenerate input (less than 3 points or equivalent due to repeated points)
+	 * @return						true if polygon is convex
+	 */
+	template<typename RealType, typename VectorType>
+	bool IsConvex2(const TArrayView<const VectorType> Vertices, RealType Tolerance = TMathUtil<RealType>::ZeroTolerance,
+		bool bDegenerateIsConvex = true)
+	{
+		int32 N = Vertices.Num();
+		if (N < 3)
+		{
+			// degenerate case, less than 3 points
+			return bDegenerateIsConvex;
+		}
+		VectorType FromPrev = Vertices[N - 1] - Vertices[N - 2];
+		int32 UsedPrevIdx = N - 2;
+		while (!FromPrev.Normalize(0))
+		{
+			UsedPrevIdx--;
+			if (UsedPrevIdx < 0)
+			{
+				return bDegenerateIsConvex; // degenerate case: all points coincident
+			}
+			FromPrev = Vertices[N - 1] - Vertices[UsedPrevIdx];
+		}
+		bool bSeenNeg = false, bSeenPos = false;
+		RealType AngleSum = 0;
+		int32 NonZeroEdges = 0;
+		for (int32 Cur = N-1, Next = 0; Next < N; Next++)
+		{
+			VectorType ToNext = Vertices[Next] - Vertices[Cur];
+			if (!ToNext.Normalize(0)) // skip duplicates
+			{
+				continue;
+			}
+			RealType Angle = SignedAngleR(FromPrev, ToNext);
+			bSeenNeg |= (Angle < -Tolerance);
+			bSeenPos |= (Angle > Tolerance);
+			AngleSum += Angle;
+			FromPrev = ToNext;
+			Cur = Next;
+			NonZeroEdges++;
+		}
+		if (NonZeroEdges < 3)
+		{
+			return bDegenerateIsConvex;
+		}
+		return !(bSeenNeg && bSeenPos) // curve should only turn one way, and
+			// curve should loop only once, turning 2*Pi radians in that one direction
+			&& int(TMathUtil<RealType>::Round(AngleSum / TMathUtil<RealType>::TwoPi)) == (int)bSeenPos - (int)bSeenNeg;
+	}
+
 
 	/**
 	 * smooth vertices in-place (will not produce a symmetric result, but does not require extra buffer)
diff --git a/Engine/Source/Runtime/GeometryCore/Public/Polygon2.h b/Engine/Source/Runtime/GeometryCore/Public/Polygon2.h
--- a/Engine/Source/Runtime/GeometryCore/Public/Polygon2.h
+++ b/Engine/Source/Runtime/GeometryCore/Public/Polygon2.h
@@ -364,6 +364,16 @@ public:
 		return CurveUtil::WindingIntegral2<T, TVector2<T>>(Vertices, QueryPoint);
 	}
 
+	/**
+	 * @param RadiansTolerance		Maximum turn in the 'wrong' direction that will still be considered convex (in radians)
+	 * @param bDegenerateIsConvex	What to return for degenerate input (less than 3 points or equivalent due to repeated points)
+	 * @return						true if polygon is convex
+	 */
+	bool IsConvex(T RadiansTolerance = TMathUtil<T>::ZeroTolerance, bool bDegenerateIsConvex = true)
+	{
+		return CurveUtil::IsConvex2<T, TVector2<T>>(Vertices, RadiansTolerance, bDegenerateIsConvex);
+	}
+
 
 	/**
 	 * @return true if the given query point is inside the Polygon, based on the winding integral
