GeometryCore: add FMeshEdgeTriID struct that identifies an edge in a mesh by a 30-bit triangle ID and a 2-bit edge-index. This allows an edge to be represented on meshes that do not have explicit edge IDs, or on a FDynamicMesh3 after add/remove changes (as edge IDs are currently not stable across such operations).

GeometryProcessing: add functions to FGroupTopology to query some information by FMeshEdgeTriID

diff --git a/Engine/Plugins/Runtime/GeometryProcessing/Source/DynamicMesh/Private/GroupTopology.cpp b/Engine/Plugins/Runtime/GeometryProcessing/Source/DynamicMesh/Private/GroupTopology.cpp
--- a/Engine/Plugins/Runtime/GeometryProcessing/Source/DynamicMesh/Private/GroupTopology.cpp
+++ b/Engine/Plugins/Runtime/GeometryProcessing/Source/DynamicMesh/Private/GroupTopology.cpp
@@ -229,6 +229,23 @@ const TArray<int>& FGroupTopology::GetGroupNbrGroups(int GroupID) const
 }
 
 
+bool FGroupTopology::IsGroupEdge(FMeshTriEdgeID TriEdgeID, bool bIncludeMeshBoundary) const
+{
+	int32 EdgeID = Mesh->GetTriEdge((int32)TriEdgeID.TriangleID, (int32)TriEdgeID.TriEdgeIndex);
+	if (EdgeID != IndexConstants::InvalidID)
+	{
+		FIndex2i EdgeT = Mesh->GetEdgeT(EdgeID);
+		if (EdgeT.B == IndexConstants::InvalidID)
+		{
+			return bIncludeMeshBoundary;
+		}
+		int32 GroupA = GetGroupID(EdgeT.A);
+		int32 GroupB = GetGroupID(EdgeT.B);
+		return GroupA != GroupB;
+	}
+	return false;
+}
+
 
 int FGroupTopology::FindGroupEdgeID(int MeshEdgeID) const
 {
@@ -252,6 +269,11 @@ int FGroupTopology::FindGroupEdgeID(int MeshEdgeID) const
 	return -1;
 }
 
+int FGroupTopology::FindGroupEdgeID(FMeshTriEdgeID TriEdgeID) const
+{
+	int32 EdgeID = Mesh->GetTriEdge((int32)TriEdgeID.TriangleID, (int32)TriEdgeID.TriEdgeIndex);
+	return (EdgeID != IndexConstants::InvalidID) ? FindGroupEdgeID(EdgeID) : -1;
+}
 
 const TArray<int>& FGroupTopology::GetGroupEdgeVertices(int GroupEdgeID) const
 {
diff --git a/Engine/Plugins/Runtime/GeometryProcessing/Source/DynamicMesh/Public/GroupTopology.h b/Engine/Plugins/Runtime/GeometryProcessing/Source/DynamicMesh/Public/GroupTopology.h
--- a/Engine/Plugins/Runtime/GeometryProcessing/Source/DynamicMesh/Public/GroupTopology.h
+++ b/Engine/Plugins/Runtime/GeometryProcessing/Source/DynamicMesh/Public/GroupTopology.h
@@ -275,8 +275,12 @@ public:
 	/** @return the list of neighbour GroupIDs for the given GroupID, or empty list if not found */
 	const TArray<int>& GetGroupNbrGroups(int GroupID) const;
 
+	/** @return true if the mesh edge identified by FMeshTriEdgeID is a group boundary edge */
+	bool IsGroupEdge(FMeshTriEdgeID TriEdgeID, bool bIncludeMeshBoundary = false) const; 
 	/** @return the ID of the FGroupEdge that contains the given Mesh Edge ID*/
 	int FindGroupEdgeID(int MeshEdgeID) const;
+	/** @return the ID of the FGroupEdge that contains the given Mesh Edge ID */
+	int FindGroupEdgeID(FMeshTriEdgeID FMeshTriEdgeID) const;
 	/** @return the list of vertices of a FGroupEdge identified by the GroupEdgeID */
 	const TArray<int>& GetGroupEdgeVertices(int GroupEdgeID) const;
 	/** @return the list of edges of a FGroupEdge identified by the GroupEdgeID */
diff --git a/Engine/Source/Runtime/GeometryCore/Public/DynamicMesh/InfoTypes.h b/Engine/Source/Runtime/GeometryCore/Public/DynamicMesh/InfoTypes.h
--- a/Engine/Source/Runtime/GeometryCore/Public/DynamicMesh/InfoTypes.h
+++ b/Engine/Source/Runtime/GeometryCore/Public/DynamicMesh/InfoTypes.h
@@ -33,6 +33,70 @@ struct FVertexInfo
 };
 
 
+
+/**
+ * FMeshTriEdgeID identifies an edge in a triangle mesh based on 
+ * the triangle ID/Index and the "edge index" 0/1/2 in the triangle.
+ * If the ordered triangle vertices are [A,B,C], then [A,B]=0, [B,C]=1, and [C,A]=2.
+ * 
+ * This type of edge identifier is applicable on any indexed mesh, even if 
+ * the mesh does not store explicit edge IDs. 
+ * 
+ * Values are stored unsigned, so it is currently *not* possible to store an "invalid"
+ * edge identifier as a FMeshTriEdgeID (0xFFFFFFFF could potentially be used as such an identifier).
+ * 
+ * The TriangleID is stored in 30 bits, while FDynamicMesh3 stores (valid) triangle IDs in 31 bits.
+ * So, only ~1 billion triangles are allowed in a mesh when using FMeshTriEdgeID, vs ~2 billion in FDynamicMesh3.
+ * This limit has not been encountered in practice, to date.
+ * 
+ * Note that cycling or permuting the vertices of a triangle will change these indices.
+ */
+struct FMeshTriEdgeID
+{
+	/** The 0/1/2 index of the edge in the triangle's tuple of edges */
+	unsigned TriEdgeIndex : 2;
+	/** The index of the mesh Triangle */
+	unsigned TriangleID : 30;
+
+	FMeshTriEdgeID()
+	{
+		TriEdgeIndex = 0;
+		TriangleID = 0;
+	}
+
+	/**
+	 * Construct a FMeshTriEdgeID for the given TriangleID and Edge Index in range
+	 * @param EdgeIndexIn index in range 0,1,2
+	 */
+	FMeshTriEdgeID(int32 TriangleIDIn, int32 EdgeIndexIn)
+	{
+		checkSlow(EdgeIndexIn >= 0 && EdgeIndexIn <= 2);
+		checkSlow(TriangleIDIn >= 0 && TriangleIDIn < (1<<30));
+		TriEdgeIndex = (unsigned int)EdgeIndexIn;
+		TriangleID = (unsigned int)TriangleIDIn;
+	}
+
+	/**
+	 * Decode an encoded FMeshTriEdgeID from a packed uint32 created by the Encoded() function
+	 */
+	FMeshTriEdgeID(uint32 EncodedEdgeKey)
+	{
+		TriangleID = EncodedEdgeKey & 0x8FFFFFFF;
+		TriEdgeIndex = (EncodedEdgeKey & 0xC0000000) >> 30;
+	}
+
+	/**
+	 * @return the (TriangleID, TriEdgeIndex) values packed into a 32 bit integer
+	 */
+	uint32 Encoded() const
+	{
+		return (TriEdgeIndex << 30) | TriangleID;
+	}
+};
+
+
+
+
 } // end namespace UE::Geometry
 } // end namespace UE
 
