GeometryCore:
- Add FDynamicMesh3::EnumerateEdgeTriangles() function, same API as ::EnumerateVertexTriangles()
- add function UE::Geometry::TrianglesToOverlayElements() in MeshIndexUtil.h, collects elements from triangles in a mesh attribute overlay
- minor change in TPolyline3::DistanceSquared to consistently use UE::Geometry:: function

diff --git a/Engine/Source/Runtime/GeometryCore/Private/DynamicMesh/DynamicMesh3.cpp b/Engine/Source/Runtime/GeometryCore/Private/DynamicMesh/DynamicMesh3.cpp
--- a/Engine/Source/Runtime/GeometryCore/Private/DynamicMesh/DynamicMesh3.cpp
+++ b/Engine/Source/Runtime/GeometryCore/Private/DynamicMesh/DynamicMesh3.cpp
@@ -599,6 +599,21 @@ void FDynamicMesh3::EnumerateVertexTriangles(int32 VertexID, TFunctionRef<void(i
 }
 
 
+void FDynamicMesh3::EnumerateEdgeTriangles(int32 EdgeID, TFunctionRef<void(int32)> ApplyFunc) const
+{
+	checkSlow(EdgeRefCounts.IsValid(EdgeID));
+	if (IsEdge(EdgeID))
+	{
+		const FEdge Edge = Edges[EdgeID];
+		ApplyFunc(Edge.Tri.A);
+		if (Edge.Tri.B != IndexConstants::InvalidID)
+		{
+			ApplyFunc(Edge.Tri.B);
+		}
+	}
+}
+
+
 FString FDynamicMesh3::MeshInfoString()
 {
 	FString VtxString = FString::Printf(TEXT("Vertices count %d max %d  %s  VtxEdges %s"),
diff --git a/Engine/Source/Runtime/GeometryCore/Public/DynamicMesh/DynamicMesh3.h b/Engine/Source/Runtime/GeometryCore/Public/DynamicMesh/DynamicMesh3.h
--- a/Engine/Source/Runtime/GeometryCore/Public/DynamicMesh/DynamicMesh3.h
+++ b/Engine/Source/Runtime/GeometryCore/Public/DynamicMesh/DynamicMesh3.h
@@ -518,6 +518,10 @@ public:
 	/** Call ApplyFunc for each one-ring triangle of a vertex. Currently this is significantly more efficient than VtxTrianglesItr() in many use cases. */
 	void EnumerateVertexTriangles(int32 VertexID, TFunctionRef<void(int32)> ApplyFunc) const;
 
+	/** Call ApplyFunc for each triangle connected to an Edge (1 or 2 triangles) */
+	void EnumerateEdgeTriangles(int32 EdgeID, TFunctionRef<void(int32)> ApplyFunc) const;
+
+
 	//
 	// Mesh Construction
 	//
diff --git a/Engine/Source/Runtime/GeometryCore/Public/DynamicMesh/MeshIndexUtil.h b/Engine/Source/Runtime/GeometryCore/Public/DynamicMesh/MeshIndexUtil.h
--- a/Engine/Source/Runtime/GeometryCore/Public/DynamicMesh/MeshIndexUtil.h
+++ b/Engine/Source/Runtime/GeometryCore/Public/DynamicMesh/MeshIndexUtil.h
@@ -3,6 +3,7 @@
 #pragma once
 
 #include "DynamicMesh/DynamicMesh3.h"
+#include "DynamicMesh/DynamicMeshOverlay.h"
 
 namespace UE
 {
@@ -26,6 +27,18 @@ namespace Geometry
 	GEOMETRYCORE_API void VertexToTriangleOneRing(const FDynamicMesh3* Mesh, const TArray<int>& VertexIDs, TSet<int>& TriangleIDsOut);
 
 
+	/**
+	 * For a given list/enumeration of Triangles, and an Overlay, find all the Elements in all the Triangles.
+	 * Note that if ElementsOut is a TArray, it may include an Element multiple times.
+	 * @param Overlay source overlay
+	 * @param TriangleEnumeration list of triangles
+	 * @param ElementsOut list of elements
+	 */
+	template<typename OverlayType, typename EnumeratorType, typename OutputSetType>
+	void TrianglesToOverlayElements(const OverlayType* Overlay, EnumeratorType TriangleEnumeration, OutputSetType& ElementsOut);
+
+
+
 	/**
 	 * Find the two edges in a Triangle that are connected to a given Vertex of the triangle
 	 * @return edge pair, or (IndexConstants::InvalidID, IndexConstants::InvalidID) on error
@@ -142,5 +155,23 @@ FIndex3i FindNextAdjacentTriangleAroundVtx(const FDynamicMesh3* Mesh,
 }
 
 
+
+template<typename OverlayType, typename EnumeratorType, typename OutputSetType>
+void TrianglesToOverlayElements(const OverlayType* Overlay, EnumeratorType TriangleEnumeration, OutputSetType& ElementsOut)
+{
+	for (int32 TriangleID : TriangleEnumeration)
+	{
+		if ( Overlay->IsSetTriangle(TriangleID) )
+		{
+			FIndex3i TriElements = Overlay->GetTriangle(TriangleID);
+			ElementsOut.Add(TriElements.A);
+			ElementsOut.Add(TriElements.B);
+			ElementsOut.Add(TriElements.C);
+		}
+	}
+}
+
+
+
 }
 }
diff --git a/Engine/Source/Runtime/GeometryCore/Public/Polyline3.h b/Engine/Source/Runtime/GeometryCore/Public/Polyline3.h
--- a/Engine/Source/Runtime/GeometryCore/Public/Polyline3.h
+++ b/Engine/Source/Runtime/GeometryCore/Public/Polyline3.h
@@ -398,15 +398,15 @@ public:
 			T d = TNumericLimits<T>::Max();
 			if (t >= seg.Extent)
 			{
-				d = TVector<T>::DistSquared(seg.EndPoint(), QueryPoint);
+				d = UE::Geometry::DistanceSquared(seg.EndPoint(), QueryPoint);
 			}
 			else if (t <= -seg.Extent)
 			{
-				d = TVector<T>::DistSquared(seg.StartPoint(), QueryPoint);
+				d = UE::Geometry::DistanceSquared(seg.StartPoint(), QueryPoint);
 			}
 			else
 			{
-				d = (seg.PointAt(t) - QueryPoint).SquaredLength();
+				d = UE::Geometry::DistanceSquared(seg.PointAt(t), QueryPoint);
 			}
 			if (d < dist)
 			{
