ParallelTransformReduce: fix implicit 64-bit to 32-bit conversion

diff --git a/Engine/Plugins/Runtime/GeometryProcessing/Source/DynamicMesh/Private/NormalFlowRemesher.cpp b/Engine/Plugins/Runtime/GeometryProcessing/Source/DynamicMesh/Private/NormalFlowRemesher.cpp
--- a/Engine/Plugins/Runtime/GeometryProcessing/Source/DynamicMesh/Private/NormalFlowRemesher.cpp
+++ b/Engine/Plugins/Runtime/GeometryProcessing/Source/DynamicMesh/Private/NormalFlowRemesher.cpp
@@ -366,8 +366,10 @@ void FNormalFlowRemesher::TrackedFaceProjectionPass(double& MaxDistanceMoved, bo
 
 
 	// Return the maximum distance moved by a vertex
-	MaxDistanceMoved = ParallelTransformReduce(Mesh->MaxVertexID(), 0.0, [&](int32 VID) -> double
+	MaxDistanceMoved = ParallelTransformReduce(Mesh->MaxVertexID(), 0.0, [&](int64 InVertexIndex) -> double
 	{
+		check(InVertexIndex < Mesh->MaxVertexID());
+		int VID = (int)InVertexIndex;
 		if (TempFlagBuffer[VID] && Mesh->IsVertex(VID))
 		{
 			const FVector3d& CurrentPosition = Mesh->GetVertex(VID);
diff --git a/Engine/Source/Runtime/GeometryCore/Public/Async/ParallelTransformReduce.h b/Engine/Source/Runtime/GeometryCore/Public/Async/ParallelTransformReduce.h
--- a/Engine/Source/Runtime/GeometryCore/Public/Async/ParallelTransformReduce.h
+++ b/Engine/Source/Runtime/GeometryCore/Public/Async/ParallelTransformReduce.h
@@ -18,15 +18,18 @@ namespace Geometry {
 							  const T& Init,
 							  TransformFuncT Transform,
 							  ReduceFuncT Reduce,
-							  int64 NumTasks)
+							  int64 InNumTasks)
 	{
-		check(NumTasks > 0);
-		int64 NumPerTask = (int64)FMathf::Ceil((float)Num / NumTasks);
+		check(InNumTasks > 0);
+		// ParallelFor doesn't yet support int64 NumTasks, so cap it
+		const int32 NumTasks = (int32) FMath::Min(InNumTasks, MAX_int32);
+
+		const int64 NumPerTask = (int64)FMathf::Ceil((float)Num / NumTasks);
 
 		TArray<T> PerTaskResults;
 		PerTaskResults.SetNum(NumTasks);
 
-		ParallelFor(NumTasks, [&](int64 TaskIndex)
+		ParallelFor(NumTasks, [&](int32 TaskIndex)
 		{
 			T LocalResult{ Init };
 
@@ -61,15 +64,19 @@ namespace Geometry {
 								 TransformFuncT Transform,
 								 ReduceFuncT Reduce,
 								 T& Out,
-								 int64 NumTasks)
+								 int64 InNumTasks)
 	{
-		check(NumTasks > 0);
+		check(InNumTasks > 0);
+
+		// ParallelFor doesn't yet support int64 NumTasks, so cap it
+		const int32 NumTasks = (int32)FMath::Min(InNumTasks, MAX_int32);
+
 		int64 NumPerTask = (int64)FMathf::Ceil((float)Num / NumTasks);
 
 		TArray<T> PerTaskResults;
 		PerTaskResults.SetNum(NumTasks);
 
-		ParallelFor(NumTasks, [&](int64 TaskIndex)
+		ParallelFor(NumTasks, [&](int32 TaskIndex)
 		{
 			InitFunc(PerTaskResults[TaskIndex]);
 
@@ -144,11 +151,15 @@ namespace Geometry {
 							  const T& Init,
 							  TransformFuncT Transform,
 							  ReduceFuncT Reduce,
-							  int64 NumTasks)
+							  int64 InNumTasks)
 	{
 		int64 Num = IteratorDistance(BeginIterator, EndIterator);
 
-		check(NumTasks > 0);
+		check(InNumTasks > 0);
+
+		// ParallelFor doesn't yet support int64 NumTasks, so cap it
+		const int32 NumTasks = (int32)FMath::Min(InNumTasks, MAX_int32);
+
 		int64 NumPerTask = (int64)FMathf::Ceil((float)Num / NumTasks);
 
 		TArray<T> PerTaskResults;
@@ -193,11 +204,14 @@ namespace Geometry {
 								 TransformFuncT Transform,
 								 ReduceFuncT Reduce,
 								 T& Out,
-								 int64 NumTasks)
+								 int64 InNumTasks)
 	{
 		int64 Num = IteratorDistance(BeginIterator, EndIterator);
 
-		check(NumTasks > 0);
+		check(InNumTasks > 0);
+		// ParallelFor doesn't yet support int64 NumTasks, so cap it
+		const int32 NumTasks = (int32)FMath::Min(InNumTasks, MAX_int32);
+
 		int64 NumPerTask = (int64)FMathf::Ceil((float)Num / NumTasks);
 
 		TArray<T> PerTaskResults;
diff --git a/Engine/Source/Runtime/GeometryCore/Public/Spatial/FastWinding.h b/Engine/Source/Runtime/GeometryCore/Public/Spatial/FastWinding.h
--- a/Engine/Source/Runtime/GeometryCore/Public/Spatial/FastWinding.h
+++ b/Engine/Source/Runtime/GeometryCore/Public/Spatial/FastWinding.h
@@ -127,7 +127,8 @@ namespace FastTriWinding
 
 		auto CentroidTransform = [&](int64 TriangleSubsetIndex) -> PData
 		{
-			int tid = TriangleArray[TriangleSubsetIndex];
+			check(TriangleSubsetIndex < TriangleArray.Num());
+			int tid = TriangleArray[(int)TriangleSubsetIndex];
 			FVector3d P0, P1, P2;
 			Mesh.GetTriVertices(tid, P0, P1, P2);
 
@@ -158,9 +159,10 @@ namespace FastTriWinding
 			double R;
 		};
 
-		auto ExpansionTransform = [&, P = CentroidData.P](int TriangleSubsetIndex) -> OrderData
+		auto ExpansionTransform = [&, P = CentroidData.P](int64 TriangleSubsetIndex) -> OrderData
 		{
-			int tid = TriangleArray[TriangleSubsetIndex];
+			check(TriangleSubsetIndex < TriangleArray.Num());
+			int tid = TriangleArray[(int)TriangleSubsetIndex];
 			FVector3d P0, P1, P2;
 			Mesh.GetTriVertices(tid, P0, P1, P2);
 
