GeometryCore: Add FColliderMesh::FindNearestHitTriangle() and FColliderMesh::FindNearestTriangle() wrappers around AABBTree functions. Replace most calls to these functions via AABBTree w/ calls to FColliderMesh functions. Rename ::GetAABBTree() to ::GetRawAABBTreeUnsafe() to discourage usage (intention long-term is to encapsulate the spatial data structure, so it is not exposed)

diff --git a/Engine/Plugins/Runtime/MeshModelingToolset/Source/ModelingComponents/Private/Scene/MeshSceneAdapter.cpp b/Engine/Plugins/Runtime/MeshModelingToolset/Source/ModelingComponents/Private/Scene/MeshSceneAdapter.cpp
--- a/Engine/Plugins/Runtime/MeshModelingToolset/Source/ModelingComponents/Private/Scene/MeshSceneAdapter.cpp
+++ b/Engine/Plugins/Runtime/MeshModelingToolset/Source/ModelingComponents/Private/Scene/MeshSceneAdapter.cpp
@@ -564,7 +564,8 @@ public:
 			if (bUseDistanceShellForWinding == false || bRequiresWindingQueryFallback)
 			{
 				TRACE_CPUPROFILER_EVENT_SCOPE(MeshScene_WrapperBuild_DMesh_FWNTree);
-				FWNTree = MakeUnique<TFastWindingTree<FColliderMesh>>(&ColliderMesh.GetAABBTree(), true);
+				check( ColliderMesh.GetRawAABBTreeUnsafe() != nullptr );
+				FWNTree = MakeUnique<TFastWindingTree<FColliderMesh>>(ColliderMesh.GetRawAABBTreeUnsafe(), true);
 			}
 		}
 		return true;
@@ -650,7 +651,7 @@ public:
 			{
 				FVector3d UseP = (bHasBakedTransform) ? P : LocalToWorldTransform.InverseTransformPosition(P);
 				double NearestDistSqr;
-				int32 NearTriID = ColliderMesh.GetAABBTree().FindNearestTriangle(UseP, NearestDistSqr, IMeshSpatial::FQueryOptions(WindingShellThickness));
+				int32 NearTriID = ColliderMesh.FindNearestTriangle(UseP, NearestDistSqr, IMeshSpatial::FQueryOptions(WindingShellThickness));
 				if (NearTriID != IndexConstants::InvalidID)
 				{
 					// Do we even need to do this? won't we return InvalidID if we don't find point within distance?
@@ -688,7 +689,7 @@ public:
 		}
 
 		double RayHitT; int32 HitTID; FVector3d HitBaryCoords;
-		if (ColliderMesh.GetAABBTree().FindNearestHitTriangle(UseRay, RayHitT, HitTID, HitBaryCoords))
+		if (ColliderMesh.FindNearestHitTriangle(UseRay, RayHitT, HitTID, HitBaryCoords))
 		{
 			WorldHitResultOut.HitMeshTriIndex = HitTID;
 			WorldHitResultOut.HitMeshSpatialWrapper = this;
diff --git a/Engine/Plugins/Runtime/MeshModelingToolset/Source/ModelingComponents/Private/Scene/SceneGeometrySpatialCache.cpp b/Engine/Plugins/Runtime/MeshModelingToolset/Source/ModelingComponents/Private/Scene/SceneGeometrySpatialCache.cpp
--- a/Engine/Plugins/Runtime/MeshModelingToolset/Source/ModelingComponents/Private/Scene/SceneGeometrySpatialCache.cpp
+++ b/Engine/Plugins/Runtime/MeshModelingToolset/Source/ModelingComponents/Private/Scene/SceneGeometrySpatialCache.cpp
@@ -122,7 +122,7 @@ bool FSceneVolumeSpatial::FindNearestHit(
 
 	FRay3d LocalRay = Transform.InverseTransformRay(WorldRay);
 	double RayHitT; int32 HitTID; FVector3d HitBaryCoords;
-	if ( ColliderMesh.GetAABBTree().FindNearestHitTriangle(LocalRay, RayHitT, HitTID, HitBaryCoords) )
+	if ( ColliderMesh.FindNearestHitTriangle(LocalRay, RayHitT, HitTID, HitBaryCoords) )
 	{
 		HitResultOut.Component = Component.Get();
 		HitResultOut.Actor = HitResultOut.Component->GetOwner();
diff --git a/Engine/Source/Runtime/GeometryCore/Private/DynamicMesh/ColliderMesh.cpp b/Engine/Source/Runtime/GeometryCore/Private/DynamicMesh/ColliderMesh.cpp
--- a/Engine/Source/Runtime/GeometryCore/Private/DynamicMesh/ColliderMesh.cpp
+++ b/Engine/Source/Runtime/GeometryCore/Private/DynamicMesh/ColliderMesh.cpp
@@ -103,9 +103,28 @@ void FColliderMesh::Initialize(const FDynamicMesh3& SourceMesh, const FBuildOpti
 }
 
 
-TMeshAABBTree3<FColliderMesh>& FColliderMesh::GetAABBTree()
+TMeshAABBTree3<FColliderMesh>* FColliderMesh::GetRawAABBTreeUnsafe()
 {
-	return AABBTree;
+	return &AABBTree;
+}
+
+
+bool FColliderMesh::FindNearestHitTriangle(const FRay3d& Ray, double& RayParameterOut, int& HitTriangleIDOut, FVector3d& BaryCoordsOut) const
+{
+	if ( AABBTree.IsValid(false) )
+	{
+		return AABBTree.FindNearestHitTriangle(Ray, RayParameterOut, HitTriangleIDOut, BaryCoordsOut);
+	}
+	return false;
+}
+
+int FColliderMesh::FindNearestTriangle(const FVector3d& Point, double& NearestDistSqrOut, const IMeshSpatial::FQueryOptions& Options) const
+{
+	if ( AABBTree.IsValid(false) )
+	{
+		return AABBTree.FindNearestTriangle(Point, NearestDistSqrOut, Options);
+	}
+	return IndexConstants::InvalidID;
 }
 
 
@@ -139,7 +158,7 @@ int32 FColliderMesh::GetSourceTriangleID(int32 TriangleID) const
 FVector3d FColliderMeshProjectionTarget::Project(const FVector3d& Point, int Identifier)
 {
 	double DistSqr;
-	int NearestTriID = ColliderMesh->GetAABBTree().FindNearestTriangle(Point, DistSqr);
+	int NearestTriID = ColliderMesh->FindNearestTriangle(Point, DistSqr);
 	if (NearestTriID < 0)
 	{
 		return Point;
@@ -165,7 +184,7 @@ FVector3d FColliderMeshProjectionTarget::Project(const FVector3d& Point, int Ide
 FVector3d FColliderMeshProjectionTarget::Project(const FVector3d& Point, FVector3d& ProjectNormalOut, int Identifier)
 {
 	double DistSqr;
-	int NearestTriID = ColliderMesh->GetAABBTree().FindNearestTriangle(Point, DistSqr);
+	int NearestTriID = ColliderMesh->FindNearestTriangle(Point, DistSqr);
 	if (NearestTriID < 0)
 	{
 		return Point;
diff --git a/Engine/Source/Runtime/GeometryCore/Public/DynamicMesh/ColliderMesh.h b/Engine/Source/Runtime/GeometryCore/Public/DynamicMesh/ColliderMesh.h
--- a/Engine/Source/Runtime/GeometryCore/Public/DynamicMesh/ColliderMesh.h
+++ b/Engine/Source/Runtime/GeometryCore/Public/DynamicMesh/ColliderMesh.h
@@ -93,7 +93,28 @@ public:
 
 	// spatial APIs
 
-	TMeshAABBTree3<FColliderMesh>& GetAABBTree();
+	/**
+	 * @param HitTriangleIDOut TriangleID on collider mesh - use GetSourceTriangleID() to map back to source mesh
+	 * @return true if a triangle was hit and Out parameters are initialized
+	 */
+	bool FindNearestHitTriangle( const FRay3d& Ray, double& RayParameterOut, int& HitTriangleIDOut, FVector3d& BaryCoordsOut ) const;
+
+	/**
+	 * Find the TriangleID closest to P on the collider mesh, and distance to it, within distance MaxDist, or return InvalidID.
+	 * Use GetSourceTriangleID() to map the TriangleID back to the source mesh.
+	 * Use MeshQueries.TriangleDistance() to get more information.
+	 */
+	int FindNearestTriangle( const FVector3d& Point, double& NearestDistSqrOut, const IMeshSpatial::FQueryOptions& Options = IMeshSpatial::FQueryOptions() ) const;
+
+
+	/**
+	 * Direct access to the AABBTree pointer inside the ColliderMesh. 
+	 * @warning this function should be avoided if possible, and may be removed in the future. 
+	 */
+	TMeshAABBTree3<FColliderMesh>* GetRawAABBTreeUnsafe();
+
+
+
 
 protected:
 	// For large meshes it might be better to optionally use TDynamicVector here (but keeping the TArray for
