add 2D Voronoi diagram support, built on the Delaunay triangulation code

diff --git a/Engine/Source/Runtime/GeometryCore/Private/CompGeom/Delaunay2.cpp b/Engine/Source/Runtime/GeometryCore/Private/CompGeom/Delaunay2.cpp
--- a/Engine/Source/Runtime/GeometryCore/Private/CompGeom/Delaunay2.cpp
+++ b/Engine/Source/Runtime/GeometryCore/Private/CompGeom/Delaunay2.cpp
@@ -5,6 +5,8 @@
 #include "CompGeom/ExactPredicates.h"
 #include "Spatial/ZOrderCurvePoints.h"
 
+#include "Async/ParallelFor.h"
+
 namespace UE
 {
 namespace Geometry
@@ -56,12 +58,39 @@ struct FDelaunay2Connectivity
 		}
 	}
 
+	const TArray<int32>& GetVertexAdjacencyCache() const
+	{
+		return VertexAdjCache;
+	}
+
+	TArray<int32> MakeVertexAdjacency(int32 NumVertices) const
+	{
+		TArray<int32> VertexAdj;
+		VertexAdj.Init(-1, VertexIDToAdjIndex(NumVertices));
+		for (const TPair<FIndex2i, int32>& EV : EdgeToVert)
+		{
+			FIndex2i AdjEdge(VertexIDToAdjIndex(EV.Key.A), VertexIDToAdjIndex(EV.Key.B));
+			VertexAdj[AdjEdge.A] = AdjEdge.B;
+		}
+		return VertexAdj;
+	}
+
 	void DisableVertexAdjacency()
 	{
 		bUseAdjCache = false;
 		VertexAdjCache.Empty();
 	}
 
+	bool HasVertexAdjacency() const
+	{
+		return bUseAdjCache;
+	}
+
+	bool HasCompleteVertexAdjacency(int32 NumVertices) const
+	{
+		return bUseAdjCache && VertexAdjCache.Num() == VertexIDToAdjIndex(NumVertices);
+	}
+
 	bool HasEdge(const FIndex2i& Edge) const
 	{
 		return EdgeToVert.Contains(Edge);
@@ -526,11 +555,12 @@ struct FDelaunay2Connectivity
 
 	// Get any edge BC opposite vertex A, such that triangle ABC is in the mesh (or return InvalidIndex edge if no such edge is present)
 	// Before calling this frequently, consider calling EnableVertexAdjacency()
-	FIndex2i GetEdge(int32 Vertex) const
+	// Not thread-safe if cache is enabled, because it may try to update the cache
+	FIndex2i GetEdge(int32 Vertex)
 	{
 		if (bUseAdjCache)
 		{
-			int32 AdjVertex = GetCachedAdjVertex(Vertex);
+			int32 AdjVertex = GetCachedAdjVertex(VertexAdjCache, Vertex);
 			if (AdjVertex != InvalidIndex)
 			{
 				int32 LastVertex = GetVertex(FIndex2i(Vertex, AdjVertex));
@@ -550,6 +580,19 @@ struct FDelaunay2Connectivity
 		}
 		return FIndex2i(InvalidIndex, InvalidIndex);
 	}
+
+	// Same as GetEdge but only reads from a vertex adjacency cache; returns InvalidIndex if not in cache
+	// Thread-safe because it will never try to update the cache.
+	FIndex2i GetEdgeFromCache(const TArray<int32>& Cache, int32 Vertex) const
+	{
+		int32 AdjVertex = GetCachedAdjVertex(Cache, Vertex);
+		if (AdjVertex == InvalidIndex)
+		{
+			return FIndex2i(InvalidIndex, InvalidIndex);
+		}
+		int32 LastVertex = GetVertex(FIndex2i(Vertex, AdjVertex));
+		return FIndex2i(AdjVertex, LastVertex);
+	}
 	
 	// Call a function on every oriented edge (+ next vertex) on the mesh
 	// (note the number of edges visited will be 3x the number of triangles)
@@ -594,7 +637,7 @@ protected:
 	
 	// Optional cache of a single vertex in the 1-ring of each vertex
 	// Makes GetEdge() constant time (as long as the cache hits) instead of O(#Edges), at the cost of additional storage and bookkeeping.
-	mutable TArray<int32> VertexAdjCache;
+	TArray<int32> VertexAdjCache;
 	bool bUseAdjCache = false;
 
 	bool bTrackDuplicateVertices = false;
@@ -609,12 +652,12 @@ protected:
 	{
 		return AdjIndex - 1;
 	}
-	inline int32 GetCachedAdjVertex(int32 VertexID) const
+	inline int32 GetCachedAdjVertex(const TArray<int32>& Cache, int32 VertexID) const
 	{
 		int32 AdjIndex = VertexIDToAdjIndex(VertexID);
-		return AdjIndexToVertexID(VertexAdjCache[AdjIndex]);
+		return AdjIndexToVertexID(Cache[AdjIndex]);
 	}
-	inline void UpdateAdjCache(FIndex2i Edge) const
+	inline void UpdateAdjCache(FIndex2i Edge)
 	{
 		FIndex2i AdjEdge(VertexIDToAdjIndex(Edge.A), VertexIDToAdjIndex(Edge.B));
 		VertexAdjCache[AdjEdge.A] = AdjEdge.B;
@@ -825,6 +868,166 @@ namespace DelaunayInternal
 		return Border.Num() > 0 ? FIndex3i(Border[0].B, Border[0].A, ToInsertV) : FIndex3i::Invalid();
 	}
 
+	template<typename RealType>
+	TArray<TArray<TVector2<RealType>>> GetVoronoiCells(const FDelaunay2Connectivity& Connectivity, TArrayView<const TVector2<RealType>> Vertices, bool bIncludeBoundary, TAxisAlignedBox2<RealType> BoundsClip, RealType ExpandBounds)
+	{
+		// get or create vertex adjacency data; we will use it heavily here
+		TArray<int32> ComputedVertexAdj;
+		const TArray<int32>* UseVertexAdj;
+		int32 NumVertices = Vertices.Num();
+		if (Connectivity.HasCompleteVertexAdjacency(NumVertices))
+		{
+			UseVertexAdj = &Connectivity.GetVertexAdjacencyCache();
+		}
+		else
+		{
+			ComputedVertexAdj = Connectivity.MakeVertexAdjacency(NumVertices);
+			UseVertexAdj = &ComputedVertexAdj;
+		}
+
+		bool bDoClip = !BoundsClip.IsEmpty();
+		if (bDoClip)
+		{
+			BoundsClip.Expand(ExpandBounds);
+		}
+
+		TAxisAlignedBox2<RealType> Bounds;
+
+		if (bIncludeBoundary)
+		{
+			for (const TVector2<RealType>& Vert : Vertices)
+			{
+				Bounds.Contain(Vert);
+			}
+
+			// Note: We're recomputing each circumcenter 4 times (once here, and once per triangle vertex),
+			// but they're cheap to compute, so still probably not worth caching
+			Connectivity.EnumerateTriangles([&](const FIndex2i& Edge, int32 Vertex)
+			{
+				Bounds.Contain(VectorUtil::Circumcenter(Vertices[Vertex], Vertices[Edge.A], Vertices[Edge.B]));
+				return true;
+			}, true);
+			Bounds.Expand(FMath::Max((RealType)UE_SMALL_NUMBER, ExpandBounds));
+			Bounds.Contain(BoundsClip);
+		}
+
+		// Helper finds ray-boundary intersection, and an index indicating which side of the boundary was crossed
+		auto GetBoundaryCrossing = [&Vertices, &Bounds](int32 A, int32 B, int32& OutCrossIdx)
+		{
+			// Cast the ray from the midpoint of a Delaunay edge, in the perpendicular direction
+			TVector2<RealType> Mid = (Vertices[A] + Vertices[B]) * .5;
+			TVector2<RealType> Edge = (Vertices[B] - Vertices[A]);
+			TVector2<RealType> EdgePerp(-Edge.Y, Edge.X);
+			RealType T = TMathUtilConstants<RealType>::MaxReal;
+			OutCrossIdx = -1;
+			if (EdgePerp.X < 0)
+			{
+				T = (Bounds.Min.X - Mid.X) / EdgePerp.X;
+				OutCrossIdx = 0;
+			}
+			else if (EdgePerp.X > 0)
+			{
+				T = (Bounds.Max.X - Mid.X) / EdgePerp.X;
+				OutCrossIdx = 2;
+			}
+			if (EdgePerp.Y < 0)
+			{
+				RealType TCand = (Bounds.Min.Y - Mid.Y) / EdgePerp.Y;
+				if (TCand < T)
+				{
+					T = TCand;
+					OutCrossIdx = 1;
+				}
+			}
+			else if (EdgePerp.Y > 0)
+			{
+				RealType TCand = (Bounds.Max.Y - Mid.Y) / EdgePerp.Y;
+				if (TCand < T)
+				{
+					T = TCand;
+					OutCrossIdx = 3;
+				}
+			}
+			// Delaunay triangulation will not include duplicate points in the triangulation, so OutCrossIdx should never stay -1 (corresponds to EdgePerp of 0,0)
+			checkSlow(OutCrossIdx != -1);
+			return Mid + T * EdgePerp;
+		};
+
+		const TVector2<RealType> Corners[4]{Bounds.GetCorner(0), Bounds.GetCorner(1), Bounds.GetCorner(2), Bounds.GetCorner(3)};
+
+		TArray<TArray<TVector2<RealType>>> Cells;
+		Cells.SetNum(Vertices.Num());
+		ParallelFor(Vertices.Num(), [&](int32 VertIdx)
+		{
+			TArray<TVector2<RealType>>& Polygon = Cells[VertIdx];
+			FIndex2i CacheEdge = Connectivity.GetEdgeFromCache(*UseVertexAdj, VertIdx);
+			if (CacheEdge.A == FDelaunay2Connectivity::InvalidIndex)
+			{
+				return;
+			}
+			if (!bIncludeBoundary)
+			{
+				if (Connectivity.HasEdge(FIndex2i(VertIdx, FDelaunay2Connectivity::GhostIndex)))
+				{
+					return;
+				}
+			}
+			else
+			{
+				// Make sure that our first triangle is not a ghost triangle, to simplify subsequent processing
+				while (CacheEdge.Contains(FDelaunay2Connectivity::GhostIndex))
+				{
+					int32 NextV = Connectivity.GetVertex(FIndex2i(VertIdx, CacheEdge.B));
+					CacheEdge = FIndex2i(CacheEdge.B, NextV);
+				}
+			}
+
+			FIndex2i WalkEdge = FIndex2i(VertIdx, CacheEdge.A);
+			int32 NextVert = CacheEdge.B;
+			int32 InitialBoundaryCrossIdx = -1;
+			do
+			{
+				FIndex3i Tri = Connectivity.AsUniqueTriangle(WalkEdge, NextVert);
+				if (Tri.A == -1) // passing through a ghost tri
+				{
+					if (NextVert == -1)
+					{
+						// Entering the infinite part of the cell: Add where it exits the bounds, and note the side of the bounding box
+						TVector2<RealType> CrossPt = GetBoundaryCrossing(VertIdx, WalkEdge.B, InitialBoundaryCrossIdx);
+						Polygon.Add(CrossPt);
+					}
+					else // WalkEdge.B == -1
+					{
+						// Exiting the infinite part of the cell: Add the corners between the exiting and entering bounding box sides
+						// and then add the crossing vertex
+						int32 EndBoundaryCrossIdx;
+						TVector2<RealType> CrossPt = GetBoundaryCrossing(NextVert, VertIdx, EndBoundaryCrossIdx);
+						for (int32 BdryIdx = InitialBoundaryCrossIdx; BdryIdx != EndBoundaryCrossIdx; BdryIdx = (BdryIdx + 1) % 4)
+						{
+							Polygon.Add(Corners[BdryIdx]);
+						}
+						Polygon.Add(CrossPt);
+					}
+				}
+				else // a regular tri; the cell vertex is at the circumcenter
+				{
+					TVector2<RealType> CenterPt = VectorUtil::Circumcenter(Vertices[Tri[0]], Vertices[Tri[1]], Vertices[Tri[2]]);
+					Polygon.Add(CenterPt);
+				}
+
+				WalkEdge = FIndex2i(VertIdx, NextVert);
+				NextVert = Connectivity.GetVertex(WalkEdge);
+			} while (NextVert != CacheEdge.B);
+
+			if (bDoClip)
+			{
+				CurveUtil::ClipConvexToBounds<RealType, TVector2<RealType>>(Polygon, BoundsClip.Min, BoundsClip.Max);
+			}
+		}); // end ParallelFor
+
+		return Cells;
+	}
+
 	template<typename RealType>
 	bool IsDelaunay(FDelaunay2Connectivity& Connectivity, TArrayView<const TVector2<RealType>> Vertices, TArrayView<const FIndex2i> SkipEdgesIn)
 	{
@@ -1292,6 +1495,26 @@ bool FDelaunay2::IsDelaunay(TArrayView<const FVector2d> Vertices, TArrayView<con
 	return false; // if no triangulation was performed, function should not be called
 }
 
+
+TArray<TArray<FVector2d>> FDelaunay2::GetVoronoiCells(TArrayView<const FVector2d> Vertices, bool bIncludeBoundary, FAxisAlignedBox2d Bounds, double ExpandBounds) const
+{
+	if (ensureMsgf(Connectivity.IsValid() && !bIsConstrained, TEXT("Voronoi diagram computation requires a valid, unconstrained Delaunay triangulation to be already computed")))
+	{
+		return DelaunayInternal::GetVoronoiCells<double>(*Connectivity, Vertices, bIncludeBoundary, Bounds, ExpandBounds);
+	}
+	return TArray<TArray<FVector2d>>();
+}
+
+TArray<TArray<FVector2f>> FDelaunay2::GetVoronoiCells(TArrayView<const FVector2f> Vertices, bool bIncludeBoundary, FAxisAlignedBox2f Bounds, float ExpandBounds) const
+{
+	if (ensureMsgf(Connectivity.IsValid() && !bIsConstrained, TEXT("Voronoi diagram computation required a valid, unconstrained Delaunay triangulation to be already computed")))
+	{
+		return DelaunayInternal::GetVoronoiCells<float>(*Connectivity, Vertices, bIncludeBoundary, Bounds, ExpandBounds);
+	}
+	return TArray<TArray<FVector2f>>();
+}
+
+
 bool FDelaunay2::HasEdges(TArrayView<const FIndex2i> Edges) const
 {
 	if (ensure(Connectivity.IsValid()))
diff --git a/Engine/Source/Runtime/GeometryCore/Public/CompGeom/Delaunay2.h b/Engine/Source/Runtime/GeometryCore/Public/CompGeom/Delaunay2.h
--- a/Engine/Source/Runtime/GeometryCore/Public/CompGeom/Delaunay2.h
+++ b/Engine/Source/Runtime/GeometryCore/Public/CompGeom/Delaunay2.h
@@ -220,6 +220,33 @@ public:
 	// @return true if the edge is in the triangulation
 	bool HasEdge(const FIndex2i& Edge, bool bRemapDuplicates);
 
+	/**
+	 * Get Voronoi diagram cells as dual of the Delaunay triangulation.  You must call Triangulate() before calling this.
+	 * @param Vertices			Vertices of the Delaunay triangulation, to be used as the sites of the Voronoi diagram
+	 * @param bIncludeBoundary	If true, include the cells on the boundary of the diagram.  These cells are conceptually infinite, but will be clipped to a bounding box.
+	 * @param ClipBounds		If non-empty, all Voronoi diagram cells will be clipped to this rectangle.
+	 * @param ExpandBounds		Amount to expand the clipping bounds beyond the Bounds argument (or the bounding box of non-boundary Voronoi cells, if Bounds was empty)
+	 * @return					The cells of each Voronoi site, or an empty array if the Triangulation was not yet computed
+	 */
+	TArray<TArray<FVector2d>> GetVoronoiCells(TArrayView<const FVector2d> Vertices, bool bIncludeBoundary = false, FAxisAlignedBox2d ClipBounds = FAxisAlignedBox2d::Empty(), double ExpandBounds = 0.0) const;
+	TArray<TArray<FVector2f>> GetVoronoiCells(TArrayView<const FVector2f> Vertices, bool bIncludeBoundary = false, FAxisAlignedBox2f ClipBounds = FAxisAlignedBox2f::Empty(), float ExpandBounds = 0.0f) const;
+
+	/**
+	 * Compute Voronoi diagram cells
+	 * @param Sites				Positions to use as the sites of the Voronoi diagram
+	 * @param bIncludeBoundary	If true, include the cells on the boundary of the diagram.  These cells are conceptually infinite, but will be clipped to a bounding box.
+	 * @param ClipBounds		If non-empty, all Voronoi diagram cells will be clipped to this rectangle.
+	 * @param ExpandBounds		Amount to expand the clipping bounds beyond the Bounds argument (or the bounding box of non-boundary Voronoi cells, if Bounds was empty)
+	 * @return					The cells of each Voronoi site, or an empty array if the Triangulation was not yet computed
+	 */
+	template<typename RealType>
+	static TArray<TArray<TVector2<RealType>>> ComputeVoronoiCells(TArrayView<const TVector2<RealType>> Sites, bool bIncludeBoundary = false, TAxisAlignedBox2<RealType> ClipBounds = FAxisAlignedBox2f::Empty(), RealType ExpandBounds = (RealType)0)
+	{
+		FDelaunay2 Delaunay;
+		Delaunay.Triangulate(Sites);
+		return Delaunay.GetVoronoiCells(Sites, bIncludeBoundary, ClipBounds, ExpandBounds);
+	}
+
 protected:
 	TPimplPtr<FDelaunay2Connectivity> Connectivity;
 
diff --git a/Engine/Source/Runtime/GeometryCore/Public/VectorUtil.h b/Engine/Source/Runtime/GeometryCore/Public/VectorUtil.h
--- a/Engine/Source/Runtime/GeometryCore/Public/VectorUtil.h
+++ b/Engine/Source/Runtime/GeometryCore/Public/VectorUtil.h
@@ -418,6 +418,31 @@ namespace VectorUtil
 	}
 
 
+	/**
+	 * @return circumcenter of triangle ABC
+	 */
+	template<typename RealType>
+	inline TVector2<RealType> Circumcenter(TVector2<RealType> A, TVector2<RealType> B, const TVector2<RealType>& C, RealType Epsilon = TMathUtilConstants<RealType>::Epsilon)
+	{
+		// Compute in offset space w/ A translated to origin
+		TVector2<RealType> AB = B - A;
+		TVector2<RealType> AC = C - A;
+		
+		RealType Denom = 2 * DotPerp(AB,AC);
+		if (FMath::Abs(Denom) <= Epsilon)
+		{
+			// degenerate (flat) triangle does not have a (finite) circumcenter ...
+			// we'll just fall back to the centroid in this case
+			return A + (AB + AC) * RealType(1.0/3.0);
+		}
+		RealType ABLenSq = AB.SquaredLength();
+		RealType ACLenSq = AC.SquaredLength();
+		TVector2<RealType> Center(
+			A.X + (AC.Y * ABLenSq - AB.Y * ACLenSq) / Denom,
+			A.Y + (AB.X * ACLenSq - AC.X * ABLenSq) / Denom
+			);
+		return Center;
+	}
 
 	/**
 	 * @return sign of Bitangent relative to Normal and Tangent
