GeometryCore: Curve: DynamicGraph.
Fix possible int to unsigned int conversion error

diff --git a/Engine/Source/Runtime/GeometryCore/Public/Curve/DynamicGraph.h b/Engine/Source/Runtime/GeometryCore/Public/Curve/DynamicGraph.h
--- a/Engine/Source/Runtime/GeometryCore/Public/Curve/DynamicGraph.h
+++ b/Engine/Source/Runtime/GeometryCore/Public/Curve/DynamicGraph.h
@@ -438,10 +438,15 @@ public:
 		}
 		
 		// update middle vertex refcounts
-		vertices_refcount.Increment(f, FIncr);
-		if (BIncr)
+		vertices_refcount.Increment(f, (unsigned short) FIncr);
+		if (BIncr > 0)
 		{
-			vertices_refcount.Increment(B, BIncr);
+			vertices_refcount.Increment(B, (unsigned short) BIncr);
+		}
+		if (BIncr < 0)
+		{
+			unsigned short BDecr = (unsigned short) (-BIncr);
+			vertices_refcount.Decrement(B, BDecr);
 		}
 
 		Split.VNew = f;
