add extra validation of CrossIdx values to make static analysis happy + be extra safe vs crash/inf loop

diff --git a/Engine/Source/Runtime/GeometryCore/Private/CompGeom/Delaunay2.cpp b/Engine/Source/Runtime/GeometryCore/Private/CompGeom/Delaunay2.cpp
--- a/Engine/Source/Runtime/GeometryCore/Private/CompGeom/Delaunay2.cpp
+++ b/Engine/Source/Runtime/GeometryCore/Private/CompGeom/Delaunay2.cpp
@@ -1002,9 +1002,13 @@ namespace DelaunayInternal
 						// and then add the crossing vertex
 						int32 EndBoundaryCrossIdx;
 						TVector2<RealType> CrossPt = GetBoundaryCrossing(NextVert, VertIdx, EndBoundaryCrossIdx);
-						for (int32 BdryIdx = InitialBoundaryCrossIdx; BdryIdx != EndBoundaryCrossIdx; BdryIdx = (BdryIdx + 1) % 4)
+						// Note: both CrossIdx values should always be > -1 but double-check anyway
+						if (InitialBoundaryCrossIdx > -1 && EndBoundaryCrossIdx > -1)
 						{
-							Polygon.Add(Corners[BdryIdx]);
+							for (int32 BdryIdx = InitialBoundaryCrossIdx; BdryIdx != EndBoundaryCrossIdx; BdryIdx = (BdryIdx + 1) % 4)
+							{
+								Polygon.Add(Corners[BdryIdx]);
+							}
 						}
 						Polygon.Add(CrossPt);
 					}
