GeometryCore: SweepingMeshSDF - fixing implicit conversion warnings.

diff --git a/Engine/Source/Runtime/GeometryCore/Public/Implicit/SweepingMeshSDF.h b/Engine/Source/Runtime/GeometryCore/Public/Implicit/SweepingMeshSDF.h
--- a/Engine/Source/Runtime/GeometryCore/Public/Implicit/SweepingMeshSDF.h
+++ b/Engine/Source/Runtime/GeometryCore/Public/Implicit/SweepingMeshSDF.h
@@ -12,6 +12,7 @@
 #include "Async/ParallelFor.h"
 #include "Misc/ScopeLock.h"
 #include "Implicit/GridInterpolant.h"
+#include "Math/NumericLimits.h"
 
 namespace UE
 {
@@ -131,7 +132,7 @@ public:
 	 * @param CellSize Spacing between Grid points
 	 * @param Spatial Optional AABB tree; note it *must* be provided if ComputeMode is set to NarrowBand_SpatialFloodFill
 	 */
-	TSweepingMeshSDF(const TriangleMeshType* Mesh = nullptr, double CellSize = 1, TMeshAABBTree3<TriangleMeshType>* Spatial = nullptr) : Mesh(Mesh), Spatial(Spatial), CellSize(CellSize)
+	TSweepingMeshSDF(const TriangleMeshType* Mesh = nullptr, float CellSize = 1, TMeshAABBTree3<TriangleMeshType>* Spatial = nullptr) : Mesh(Mesh), Spatial(Spatial), CellSize(CellSize)
 	{
 	}
 
@@ -203,7 +204,7 @@ public:
 		double MaxDim = MaxElement(Bounds.Max - Bounds.Min + ExpandBounds * 2.0);
 		if (!ensureMsgf(MaxDim / CellSize <= ApproxMaxCellsPerDimension - 2 * ExactBandWidth, TEXT("SDF resolution clamped to avoid excessive memory use")))
 		{
-			CellSize = MaxDim / (ApproxMaxCellsPerDimension - 2 * ExactBandWidth);
+			CellSize = float( MaxDim / (ApproxMaxCellsPerDimension - 2 * ExactBandWidth) );
 			if (!ensure(CellSize > 0 && FMath::IsFinite(CellSize)))
 			{
 				return false;
@@ -750,7 +751,8 @@ private:
 		{
 			int64 CellLinearIndex = Distances.ToLinear(CellGridIndex);
 			int64 SectionIndex = CellLinearIndex / SectionSize;
-			return &GridSections[SectionIndex];
+			checkSlow(SectionIndex <= MAX_int32);
+			return &GridSections[(int32)SectionIndex];
 		};
 
 		// per-grid-chunk queue, each one of these cooresponds to a GridSections lock
@@ -759,7 +761,9 @@ private:
 		auto AddToSectionQueue = [SectionSize, &SectionQueues](int64 CellLinearIndex)
 		{
 			int64 SectionIndex = CellLinearIndex / SectionSize;
-			SectionQueues[SectionIndex].Add(CellLinearIndex);
+			checkSlow(SectionIndex <= MAX_int32);
+			checkSlow(CellLinearIndex <= MAX_int32);
+			SectionQueues[(int32)SectionIndex].Add((int32)CellLinearIndex);
 		};
 
 
@@ -1077,10 +1081,10 @@ private:
 				double fir = (xr[0] - ox) * invdx, fjr = (xr[1] - oy) * invdx, fkr = (xr[2] - oz) * invdx;
 
 				// recompute J/K integer bounds of triangle w/o exact band
-				int32 j0 = FMath::Clamp(FMath::CeilToInt(FMath::Min3(fjp, fjq, fjr)), 0, NJ - 1);
-				int32 j1 = FMath::Clamp(FMath::FloorToInt(FMath::Max3(fjp, fjq, fjr)), 0, NJ - 1);
-				int32 k0 = FMath::Clamp(FMath::CeilToInt(FMath::Min3(fkp, fkq, fkr)), 0, NK - 1);
-				int32 k1 = FMath::Clamp(FMath::FloorToInt(FMath::Max3(fkp, fkq, fkr)), 0, NK - 1);
+				int32 j0 = FMath::Clamp(FMath::CeilToInt32(FMath::Min3(fjp, fjq, fjr)), 0, NJ - 1);
+				int32 j1 = FMath::Clamp(FMath::FloorToInt32(FMath::Max3(fjp, fjq, fjr)), 0, NJ - 1);
+				int32 k0 = FMath::Clamp(FMath::CeilToInt32(FMath::Min3(fkp, fkq, fkr)), 0, NK - 1);
+				int32 k1 = FMath::Clamp(FMath::FloorToInt32(FMath::Max3(fkp, fkq, fkr)), 0, NK - 1);
 
 				// and do intersection counts
 				for (int K = k0; K <= k1; ++K)
@@ -1091,7 +1095,7 @@ private:
 						if (PointInTriangle2d(J, K, fjp, fkp, fjq, fkq, fjr, fkr, A, B, C))
 						{
 							double fi = A * fip + B * fiq + C * fir; // intersection I coordinate
-							int i_interval = FMath::CeilToInt(fi); // intersection is in (i_interval-1,i_interval]
+							int i_interval = FMath::CeilToInt32(fi); // intersection is in (i_interval-1,i_interval]
 							if (i_interval < 0)
 							{
 								DenseGrid::AtomicIncDec(IntersectionCount, 0, J, K, neg_x);
