GeometryCore: Parameterization.   MeshLocalParam and MeshDijkstra fix double to float cast warning (priority queue is sorted on floats, but distances are calculated in double)

diff --git a/Engine/Source/Runtime/GeometryCore/Public/Parameterization/IncrementalMeshDijkstra.h b/Engine/Source/Runtime/GeometryCore/Public/Parameterization/IncrementalMeshDijkstra.h
--- a/Engine/Source/Runtime/GeometryCore/Public/Parameterization/IncrementalMeshDijkstra.h
+++ b/Engine/Source/Runtime/GeometryCore/Public/Parameterization/IncrementalMeshDijkstra.h
@@ -89,7 +89,7 @@ public:
 				Node->GraphDistance = SeedPoint.StartDistance;
 				Node->FrozenTimestamp = CurrentSeedTimestamp;
 				Node->SeedID = NewIndex;
-				Queue.Insert(PointID, Node->GraphDistance);
+				Queue.Insert(PointID, float(Node->GraphDistance));
 			}
 		}
 
@@ -273,7 +273,7 @@ private:
 					NbrNode->ParentPointID = Parent->PointID;
 					NbrNode->SeedID = Parent->SeedID;
 					NbrNode->GraphDistance = NbrDist;
-					Queue.Update(NbrPointID, NbrNode->GraphDistance);
+					Queue.Update(NbrPointID, float(NbrNode->GraphDistance));
 				}
 			}
 			else 
@@ -283,7 +283,7 @@ private:
 					NbrNode->ParentPointID = Parent->PointID;
 					NbrNode->SeedID = Parent->SeedID;
 					NbrNode->GraphDistance = NbrDist;
-					Queue.Insert(NbrPointID, NbrNode->GraphDistance);
+					Queue.Insert(NbrPointID, float(NbrNode->GraphDistance));
 				}
 			}
 		}
diff --git a/Engine/Source/Runtime/GeometryCore/Public/Parameterization/MeshDijkstra.h b/Engine/Source/Runtime/GeometryCore/Public/Parameterization/MeshDijkstra.h
--- a/Engine/Source/Runtime/GeometryCore/Public/Parameterization/MeshDijkstra.h
+++ b/Engine/Source/Runtime/GeometryCore/Public/Parameterization/MeshDijkstra.h
@@ -120,7 +120,7 @@ public:
 				Node->GraphDistance = SeedPoint.StartDistance;
 				Node->bFrozen = true;
 				Node->SeedPointID = NewIndex;
-				Queue.Insert(PointID, Node->GraphDistance);
+				Queue.Insert(PointID, float(Node->GraphDistance));
 			}
 		}
 
@@ -188,7 +188,7 @@ public:
 				Node->GraphDistance = SeedPoint.StartDistance;
 				Node->bFrozen = true;
 				Node->SeedPointID = NewIndex;
-				Queue.Insert(PointID, Node->GraphDistance);
+				Queue.Insert(PointID, (float)Node->GraphDistance);
 			}
 		}
 
@@ -401,7 +401,7 @@ private:
 					NbrNode->ParentPointID = Parent->PointID;
 					NbrNode->GraphDistance = NbrDist;
 					NbrNode->SeedPointID = Parent->SeedPointID;
-					Queue.Update(NbrPointID, NbrNode->GraphDistance);
+					Queue.Update(NbrPointID, float(NbrNode->GraphDistance));
 				}
 			}
 			else 
@@ -409,7 +409,7 @@ private:
 				NbrNode->ParentPointID = Parent->PointID;
 				NbrNode->GraphDistance = NbrDist;
 				NbrNode->SeedPointID = Parent->SeedPointID;
-				Queue.Insert(NbrPointID, NbrNode->GraphDistance);
+				Queue.Insert(NbrPointID, float(NbrNode->GraphDistance));
 			}
 		}
 	}
diff --git a/Engine/Source/Runtime/GeometryCore/Public/Parameterization/MeshLocalParam.h b/Engine/Source/Runtime/GeometryCore/Public/Parameterization/MeshLocalParam.h
--- a/Engine/Source/Runtime/GeometryCore/Public/Parameterization/MeshLocalParam.h
+++ b/Engine/Source/Runtime/GeometryCore/Public/Parameterization/MeshLocalParam.h
@@ -91,7 +91,7 @@ public:
 			Node->GraphDistance = Node->UV.Length();
 			Node->bFrozen = true;
 			check(Queue.Contains(NbrPointID) == false);
-			Queue.Insert(NbrPointID, Node->GraphDistance);
+			Queue.Insert(NbrPointID, float(Node->GraphDistance));
 		}
 
 		ProcessQueueUntilTermination(ComputeToMaxDistanceIn);
@@ -450,14 +450,14 @@ protected:
 				{
 					NbrNode->ParentPointID = Parent->PointID;
 					NbrNode->GraphDistance = NbrDist;
-					Queue.Update(NbrPointID, NbrNode->GraphDistance);
+					Queue.Update(NbrPointID, float(NbrNode->GraphDistance));
 				}
 			}
 			else 
 			{
 				NbrNode->ParentPointID = Parent->PointID;
 				NbrNode->GraphDistance = NbrDist;
-				Queue.Insert(NbrPointID, NbrNode->GraphDistance);
+				Queue.Insert(NbrPointID, float(NbrNode->GraphDistance));
 			}
 		}
 	}
