GeometryCore: Image.  Fixing type promotion warnings, and adding a checkslow to help identify a possible overflow case.

diff --git a/Engine/Source/Runtime/GeometryCore/Public/Image/BCSplineFilter.h b/Engine/Source/Runtime/GeometryCore/Public/Image/BCSplineFilter.h
--- a/Engine/Source/Runtime/GeometryCore/Public/Image/BCSplineFilter.h
+++ b/Engine/Source/Runtime/GeometryCore/Public/Image/BCSplineFilter.h
@@ -80,14 +80,16 @@ float FBCSplineFilter<SplineType, bIsRadial>::GetWeight(const FVector2d& Dist) c
 	auto ComputeWeight = [this](float X) -> float
 	{
 		const float AbsX = FMathf::Abs(X);
+		const float AbsX2 = AbsX * AbsX;
+		const float AbsX3 = AbsX * AbsX2;
 		float Weight = 0.0f;
 		if(AbsX < Radius)
 		{
-			Weight = Coeff1[0] * AbsX * AbsX * AbsX + Coeff1[1] * AbsX * AbsX + Coeff1[2];
+			Weight = Coeff1[0] * AbsX3 + Coeff1[1] * AbsX2 + Coeff1[2];
 		}
 		else if(AbsX < Radius * 2.0f)
 		{
-			Weight = Coeff2[0] * AbsX * AbsX * AbsX + Coeff2[1] * AbsX * AbsX + Coeff2[2] * AbsX + Coeff2[3];
+			Weight = Coeff2[0] * AbsX3 + Coeff2[1] * AbsX2 + Coeff2[2] * AbsX + Coeff2[3];
 		}
 		return Weight;
 	};
@@ -98,8 +100,8 @@ float FBCSplineFilter<SplineType, bIsRadial>::GetWeight(const FVector2d& Dist) c
 	}
 	else
 	{
-		const float WeightX = ComputeWeight(Dist.X);
-		const float WeightY = ComputeWeight(Dist.Y);
+		const float WeightX = ComputeWeight((float)Dist.X);
+		const float WeightY = ComputeWeight((float)Dist.Y);
 		return WeightX * WeightY;
 	}
 }
diff --git a/Engine/Source/Runtime/GeometryCore/Public/Image/ImageOccupancyMap.h b/Engine/Source/Runtime/GeometryCore/Public/Image/ImageOccupancyMap.h
--- a/Engine/Source/Runtime/GeometryCore/Public/Image/ImageOccupancyMap.h
+++ b/Engine/Source/Runtime/GeometryCore/Public/Image/ImageOccupancyMap.h
@@ -8,6 +8,7 @@
 #include "Image/ImageTile.h"
 #include "Spatial/MeshAABBTree3.h"
 #include "Sampling/GridSampler.h"
+#include "MathUtil.h"
 
 namespace UE
 {
@@ -118,11 +119,11 @@ public:
 		// make flat mesh
 		TMeshAABBTree3<MeshType> FlatSpatial(&UVSpaceMesh, true);
 
-		const int32 LinearImageSize = Tile.Num();
+		const int64 LinearImageSize = Tile.Num();
 		TexelInteriorSamples.Init(0, LinearImageSize);
 		TexelQueryUVChart.Init(IndexConstants::InvalidID, LinearImageSize);
 		
-		const int32 LinearSampleSize = Tile.Num() * PixelSampler.Num();
+		const int64 LinearSampleSize = Tile.Num() * PixelSampler.Num();
 		TexelType.Init(EmptyTexel, LinearSampleSize);
 		TexelQueryUV.Init(FVector2f::Zero(), LinearSampleSize);
 		TexelQueryTriangle.Init(IndexConstants::InvalidID, LinearSampleSize);
@@ -250,7 +251,9 @@ public:
 	) const
 	{
 		int64 N = Dimensions.Num();
-		PassBuffer.SetNum(N);
+		checkSlow(N <= TMathUtilConstants<int32>::MaxReal); // TArray< , FDefaultAllocator>::SetNum(int32 ), so max Dimension ~ 65k x 65k
+		int32 N32 =  int32(N);
+		PassBuffer.SetNum( N32 );
 		
 		ParallelFor(Dimensions.GetHeight(), 
 			[this, &BeginTexel, & AccumulateTexel, &CompleteTexel, &WeightFunction, FilterWidth, &PassBuffer]
@@ -258,7 +261,7 @@ public:
 		{
 			for (int32 ImgX = 0; ImgX < Dimensions.GetWidth(); ++ImgX)
 			{
-				int64 LinearIdx = Dimensions.GetIndex(ImgX, ImgY);
+				int32 LinearIdx = (int32)Dimensions.GetIndex(ImgX, ImgY);
 				if (TexelType[LinearIdx] != EmptyTexel)
 				{
 					TexelValueType AccumValue = BeginTexel(LinearIdx);
@@ -296,7 +299,7 @@ public:
 		});
 
 		// write results
-		for (int64 k = 0; k < N; ++k)
+		for (int32 k = 0; k < N32; ++k)
 		{
 			if (TexelType[k] != EmptyTexel)
 			{
@@ -313,7 +316,7 @@ protected:
 	void InitializePixelSampler(const int32 SamplesPerPixelIn)
 	{
 		const float GridDimensionFloat = SamplesPerPixelIn > 0 ? static_cast<float>(SamplesPerPixelIn) : 1.0f;
-		const int32 GridDimension = FMathf::Clamp(FMathf::Sqrt(GridDimensionFloat), 1.0f, GridDimensionFloat);
+		const int32 GridDimension = FMath::FloorToInt32(FMathf::Clamp(FMathf::Sqrt(GridDimensionFloat), 1.0f, GridDimensionFloat));
 		PixelSampler = FGridSampler(GridDimension);
 	}
 
diff --git a/Engine/Source/Runtime/GeometryCore/Public/Image/ImageTile.h b/Engine/Source/Runtime/GeometryCore/Public/Image/ImageTile.h
--- a/Engine/Source/Runtime/GeometryCore/Public/Image/ImageTile.h
+++ b/Engine/Source/Runtime/GeometryCore/Public/Image/ImageTile.h
@@ -87,7 +87,7 @@ public:
 	/** @return the source image's coordinates given linear index into this tile. */
 	FVector2i GetSourceCoords(const int64 LinearIdx) const
 	{
-		return FVector2i(Start.X + LinearIdx % Width, Start.Y + LinearIdx / Width);
+		return FVector2i(Start.X + int32(LinearIdx % Width), Start.Y + int32(LinearIdx / Width));
 	}
 
 	/** @return the source image's coordinates given local XY coordinates into this tile. */
