GeometryCore: Spatial and Implicit:  fix type conversion warnings with  MeshWindingNumber, CachingMeshSDF and SparseNarrowBandMeshSDF

diff --git a/Engine/Source/Runtime/GeometryCore/Public/Implicit/CachingMeshSDF.h b/Engine/Source/Runtime/GeometryCore/Public/Implicit/CachingMeshSDF.h
--- a/Engine/Source/Runtime/GeometryCore/Public/Implicit/CachingMeshSDF.h
+++ b/Engine/Source/Runtime/GeometryCore/Public/Implicit/CachingMeshSDF.h
@@ -142,7 +142,7 @@ public:
 		double MaxDim = MaxElement(Bounds.Max - Bounds.Min + ExpandBounds * 2.0) + 4.0 * MaxOffsetDistance;
 		if (!ensureMsgf(MaxDim / CellSize <= ApproxMaxCellsPerDimension - 4, TEXT("SDF resolution clamped to avoid excessive memory use")))
 		{
-			CellSize = MaxDim / (ApproxMaxCellsPerDimension - 4);
+			CellSize = float (MaxDim / (ApproxMaxCellsPerDimension - 4));
 			if (!ensure(CellSize > 0 && FMath::IsFinite(CellSize)))
 			{
 				return;
@@ -340,10 +340,10 @@ private:
 			double fir = (xr[0] - ox) * invdx, fjr = (xr[1] - oy) * invdx, fkr = (xr[2] - oz) * invdx;
 
 			// recompute J/K integer bounds of triangle w/o exact band
-			int32 j0 = FMath::Clamp(FMath::CeilToInt(FMath::Min3(fjp, fjq, fjr)), 0, NJ - 1);
-			int32 j1 = FMath::Clamp(FMath::FloorToInt(FMath::Max3(fjp, fjq, fjr)), 0, NJ - 1);
-			int32 k0 = FMath::Clamp(FMath::CeilToInt(FMath::Min3(fkp, fkq, fkr)), 0, NK - 1);
-			int32 k1 = FMath::Clamp(FMath::FloorToInt(FMath::Max3(fkp, fkq, fkr)), 0, NK - 1);
+			int32 j0 = FMath::Clamp(FMath::CeilToInt32(FMath::Min3(fjp, fjq, fjr)), 0, NJ - 1);
+			int32 j1 = FMath::Clamp(FMath::FloorToInt32(FMath::Max3(fjp, fjq, fjr)), 0, NJ - 1);
+			int32 k0 = FMath::Clamp(FMath::CeilToInt32(FMath::Min3(fkp, fkq, fkr)), 0, NK - 1);
+			int32 k1 = FMath::Clamp(FMath::FloorToInt32(FMath::Max3(fkp, fkq, fkr)), 0, NK - 1);
 
 			// and do intersection counts
 			for (int K = k0; K <= k1; ++K)
@@ -354,7 +354,7 @@ private:
 					if (PointInTriangle2d(J, K, fjp, fkp, fjq, fkq, fjr, fkr, A, B, C))
 					{
 						double fi = A * fip + B * fiq + C * fir; // intersection I coordinate
-						int i_interval = FMath::CeilToInt(fi); // intersection is in (i_interval-1,i_interval]
+						int i_interval = FMath::CeilToInt32(fi); // intersection is in (i_interval-1,i_interval]
 						if (i_interval < 0)
 						{
 							DenseGrid::AtomicIncDec(IntersectionCount, 0, J, K, neg_x);
diff --git a/Engine/Source/Runtime/GeometryCore/Public/Implicit/SparseNarrowBandMeshSDF.h b/Engine/Source/Runtime/GeometryCore/Public/Implicit/SparseNarrowBandMeshSDF.h
--- a/Engine/Source/Runtime/GeometryCore/Public/Implicit/SparseNarrowBandMeshSDF.h
+++ b/Engine/Source/Runtime/GeometryCore/Public/Implicit/SparseNarrowBandMeshSDF.h
@@ -249,7 +249,7 @@ public:
 	 */
 	TSparseNarrowBandMeshSDF(
 		const TriangleMeshType* Mesh = nullptr, 
-		double CellSize = 1, 
+		float CellSize = 1.f, 
 		TMeshAABBTree3<TriangleMeshType>* Spatial = nullptr ) 
 		: Mesh(Mesh), Spatial(Spatial), CellSize(CellSize)
 	{
@@ -313,7 +313,7 @@ public:
 		double MaxDim = MaxElement(Bounds.Max - Bounds.Min + ExpandBounds * 2.0);
 		if (!ensureMsgf(MaxDim / CellSize <= ApproxMaxCellsPerDimension - 2 * ExactBandWidth, TEXT("SDF resolution clamped to avoid excessive memory use")))
 		{
-			CellSize = MaxDim / (ApproxMaxCellsPerDimension - 2 * ExactBandWidth);
+			CellSize = float ( MaxDim / (ApproxMaxCellsPerDimension - 2 * ExactBandWidth) );
 			if (!ensure(CellSize > 0 && FMath::IsFinite(CellSize)))
 			{
 				return false;
@@ -716,7 +716,7 @@ private:
 										const FVector3i local_coords = ijk - BlockMin;
 										const int32 local_index = DistanceBlock.ToLinear(local_coords[0], local_coords[1], local_coords[2]);
 										const FVector3d cellpos((double)ijk[0], (double)ijk[1], (double)ijk[2]);
-										float d = DX * PointTriangleDistance(cellpos, fp, fq, fr);
+										float d = float(DX * PointTriangleDistance(cellpos, fp, fq, fr));
 
 										float& grid_dist = DistanceBlock.At(local_index);
 										if (d < grid_dist)
@@ -1122,10 +1122,10 @@ private:
 				double fir = (xr[0] - ox) * invdx, fjr = (xr[1] - oy) * invdx, fkr = (xr[2] - oz) * invdx;
 
 				// recompute J/K integer bounds of triangle w/o exact band
-				int32 j0 = FMath::Clamp(FMath::CeilToInt(FMath::Min3(fjp, fjq, fjr)),  0, NJ - 1);
-				int32 j1 = FMath::Clamp(FMath::FloorToInt(FMath::Max3(fjp, fjq, fjr)), 0, NJ - 1);
-				int32 k0 = FMath::Clamp(FMath::CeilToInt(FMath::Min3(fkp, fkq, fkr)),  0, NK - 1);
-				int32 k1 = FMath::Clamp(FMath::FloorToInt(FMath::Max3(fkp, fkq, fkr)), 0, NK - 1);
+				int32 j0 = FMath::Clamp(FMath::CeilToInt32(FMath::Min3(fjp, fjq, fjr)),  0, NJ - 1);
+				int32 j1 = FMath::Clamp(FMath::FloorToInt32(FMath::Max3(fjp, fjq, fjr)), 0, NJ - 1);
+				int32 k0 = FMath::Clamp(FMath::CeilToInt32(FMath::Min3(fkp, fkq, fkr)),  0, NK - 1);
+				int32 k1 = FMath::Clamp(FMath::FloorToInt32(FMath::Max3(fkp, fkq, fkr)), 0, NK - 1);
 
 				
 				auto AtomicIncDec = [&AtomicGrid3i, neg_x](const int32 i, const int32 j, const int32 k)
@@ -1155,7 +1155,7 @@ private:
 						if (PointInTriangle2d(J, K, fjp, fkp, fjq, fkq, fjr, fkr, A, B, C))
 						{
 							double fi = A * fip + B * fiq + C * fir; // intersection I coordinate
-							int i_interval = FMath::CeilToInt(fi); // intersection is in (i_interval-1,i_interval]
+							int i_interval = FMath::CeilToInt32(fi); // intersection is in (i_interval-1,i_interval]
 							if (i_interval < 0)
 							{
 								AtomicIncDec(0, J, K);
diff --git a/Engine/Source/Runtime/GeometryCore/Public/Spatial/MeshWindingNumberGrid.h b/Engine/Source/Runtime/GeometryCore/Public/Spatial/MeshWindingNumberGrid.h
--- a/Engine/Source/Runtime/GeometryCore/Public/Spatial/MeshWindingNumberGrid.h
+++ b/Engine/Source/Runtime/GeometryCore/Public/Spatial/MeshWindingNumberGrid.h
@@ -72,7 +72,7 @@ public:
 
 
 	TMeshWindingNumberGrid(const TriangleMeshType* Mesh, TFastWindingTree<TriangleMeshType>* FastWinding, double CellSize)
-		: Mesh(Mesh), FastWinding(FastWinding), CellSize(CellSize), MeshSDF(Mesh, CellSize, FastWinding->GetTree())
+		: Mesh(Mesh), FastWinding(FastWinding), CellSize(CellSize), MeshSDF(Mesh, (float)CellSize, FastWinding->GetTree())
 	{
 	}
 
