GeometryCore: MathUtil - adding SignAsInt(TypeT& ) function.

diff --git a/Engine/Source/Runtime/GeometryCore/Public/MathUtil.h b/Engine/Source/Runtime/GeometryCore/Public/MathUtil.h
--- a/Engine/Source/Runtime/GeometryCore/Public/MathUtil.h
+++ b/Engine/Source/Runtime/GeometryCore/Public/MathUtil.h
@@ -144,6 +144,7 @@ public:
 	static inline RealType Abs(const RealType Value);
 	static inline RealType Clamp(const RealType Value, const RealType ClampMin, const RealType ClampMax);
 	static inline RealType Sign(const RealType Value);
+	static inline int32 SignAsInt(const RealType Value);
 	static inline RealType SignNonZero(const RealType Value);
 	static inline RealType Max(const RealType A, const RealType B);
 	static inline RealType Max3(const RealType A, const RealType B, const RealType C);
@@ -207,11 +208,16 @@ RealType TMathUtil<RealType>::Clamp(const RealType Value, const RealType ClampMi
 	return (Value < ClampMin) ? ClampMin : ((Value > ClampMax) ? ClampMax : Value);
 }
 
+template<typename RealType>
+int32 TMathUtil<RealType>::SignAsInt(const RealType Value)
+{
+	return (Value > (RealType)0) ? 1 : ((Value < (RealType)0) ? -1 : 0);
+}
 
 template<typename RealType>
 RealType TMathUtil<RealType>::Sign(const RealType Value)
 {
-	return (Value > (RealType)0) ? (RealType)1 : ((Value < (RealType)0) ? (RealType)-1 : (RealType)0);
+	return (RealType)SignAsInt(Value);
 }
 
 template<typename RealType>
