GeometryCore, Generators.  SphereGenerator and GridBoxMeshGenerator.  Fixing implicit type conversion warning

diff --git a/Engine/Source/Runtime/GeometryCore/Public/Generators/GridBoxMeshGenerator.h b/Engine/Source/Runtime/GeometryCore/Public/Generators/GridBoxMeshGenerator.h
--- a/Engine/Source/Runtime/GeometryCore/Public/Generators/GridBoxMeshGenerator.h
+++ b/Engine/Source/Runtime/GeometryCore/Public/Generators/GridBoxMeshGenerator.h
@@ -189,10 +189,10 @@ public:
 			for (int Side = 0; Side < 2; Side++)
 			{
 				int SideOpp = 1 - Side;
-				float SideSign = Side * 2 - 1;
+				int SideSign = Side * 2 - 1;
 
 				FVector3f Normal(0, 0, 0);
-				Normal[Dim] = (2 * Side - 1);
+				Normal[Dim] = float(2 * Side - 1);
 				Normal = (FVector3f)Box.Frame.FromFrameVector((FVector3d)Normal);
 
 				int MajorFaceInd = FaceIdxBase + Side;
@@ -207,14 +207,14 @@ public:
 					for (int D1 = 0; D1 < N[D[1][Dim]]; D1++)
 					{
 						// put the grid coordinates (centered at 0,0) into the UVs
-						UV[UVXDim] = D0 * Nscale[D[0][Dim]] - .5;
-						UV[UVYDim] = D1 * Nscale[D[1][Dim]] - .5;
+						UV[UVXDim] = float (D0 * Nscale[D[0][Dim]] - .5);
+						UV[UVYDim] = float (D1 * Nscale[D[1][Dim]] - .5);
 						// invert axes to match the desired UV patterns & so the opp faces are not backwards
 						UV.X *= SideSign * Minor1Flip[Dim];
 						UV.Y *= Minor2Flip[Dim];
 						// recenter and scale up
-						UV[UVXDim] = (UV[UVXDim] + .5) * WidthUVScale;
-						UV[UVYDim] = (UV[UVYDim] + .5) * HeightUVScale;
+						UV[UVXDim] = float ( (UV[UVXDim] + .5f) * WidthUVScale);
+						UV[UVYDim] = float ( (UV[UVYDim] + .5f) * HeightUVScale);
 						UVs[CurrentUVIdx] = UV;
 						Normals[CurrentUVIdx] = Normal;
 						UVParentVertex[CurrentUVIdx] = FaceVertIndices[MajorFaceInd][ToFaceV(Dim, D0, D1)];
diff --git a/Engine/Source/Runtime/GeometryCore/Public/Generators/SphereGenerator.h b/Engine/Source/Runtime/GeometryCore/Public/Generators/SphereGenerator.h
--- a/Engine/Source/Runtime/GeometryCore/Public/Generators/SphereGenerator.h
+++ b/Engine/Source/Runtime/GeometryCore/Public/Generators/SphereGenerator.h
@@ -73,8 +73,8 @@ private:
 	void GenerateUVVertices()
 	{
 		// generate the UV's
-		const float DUVphi = 1.0 / float(NumPhi - 1);
-		const float DUVtheta = -1.0 / float(NumTheta);
+		const float DUVphi = 1.f / float(NumPhi - 1);
+		const float DUVtheta = -1.f / float(NumTheta);
 
 		int32 UVIdx = 0;
 		int32 p,t;
