workaround internal compiler error on MSVC2019 + ASan

#rnx

diff --git a/Engine/Source/Runtime/GeometryCore/Public/Util/DynamicVector.h b/Engine/Source/Runtime/GeometryCore/Public/Util/DynamicVector.h
--- a/Engine/Source/Runtime/GeometryCore/Public/Util/DynamicVector.h
+++ b/Engine/Source/Runtime/GeometryCore/Public/Util/DynamicVector.h
@@ -105,24 +105,34 @@ public:
 		return Blocks[CurBlock][CurBlockUsed - 1];
 	}
 
+
+#if USING_ADDRESS_SANITISER
+	FORCENOINLINE Type& operator[](unsigned int Index)
+	{
+		checkSlow(Index < Num());
+
+		return Blocks[Index >> nShiftBits][Index & BlockIndexBitmask];
+	}
+	FORCENOINLINE const Type& operator[](unsigned int Index) const
+	{
+		checkSlow(Index < Num());
+
+		return Blocks[Index >> nShiftBits][Index & BlockIndexBitmask];
+	}
+#else
 	inline Type& operator[](unsigned int Index)
 	{
 		checkSlow(Index < Num());
 
-		int ArrayIndex, BlockIndex;
-		GetIndices( Index, ArrayIndex, BlockIndex );
-	
-		return Blocks[ArrayIndex][BlockIndex];
+		return Blocks[Index >> nShiftBits][Index & BlockIndexBitmask];
 	}
 	inline const Type& operator[](unsigned int Index) const
 	{
 		checkSlow(Index < Num());
 
-		int ArrayIndex, BlockIndex;
-		GetIndices( Index, ArrayIndex, BlockIndex );
-	
-		return Blocks[ArrayIndex][BlockIndex];
+		return Blocks[Index >> nShiftBits][Index & BlockIndexBitmask];
 	}
+#endif
 
 	// apply f() to each member sequentially
 	template <typename Func>
@@ -317,7 +327,7 @@ private:
 	protected:
 		TArray<ArrayType*> Elements;
 
-	public:
+	public: 
 		TBlockVector() = default;
 		TBlockVector(TBlockVector&& Moved) = default;
 
@@ -557,13 +567,6 @@ private:
 		return !(Lhs == Rhs);
 	}
 
-	// helper function - avoids MSVC+ASan internal compiler error in [] operators
-	inline void GetIndices( unsigned int Index, int& ArrayIndex, int& BlockIndex ) const
-	{
-		ArrayIndex = (Index >> nShiftBits);
-		BlockIndex = (Index & BlockIndexBitmask);
-	}
-
 	void SetCurBlock(SIZE_T Count)
 	{
 		// Reset block index for the last item and used item count within the last block.
