FProfileSweepGenerator per-point Scale

diff --git a/Engine/Source/Runtime/GeometryCore/Private/Generators/SweepGenerator.cpp b/Engine/Source/Runtime/GeometryCore/Private/Generators/SweepGenerator.cpp
--- a/Engine/Source/Runtime/GeometryCore/Private/Generators/SweepGenerator.cpp
+++ b/Engine/Source/Runtime/GeometryCore/Private/Generators/SweepGenerator.cpp
@@ -94,7 +94,7 @@ FMeshShapeGenerator& FProfileSweepGenerator::Generate()
 	}
 
 	// Check that we have our inputs
-	if (ProfileCurve.Num() < 2 || SweepCurve.Num() < 2)
+	if (ProfileCurve.Num() < 2 || SweepCurve.Num() < 2 || (!SweepScaleCurve.IsEmpty() && SweepScaleCurve.Num() != SweepCurve.Num()))
 	{
 		Reset();
 		return *this;;
@@ -221,17 +221,27 @@ FMeshShapeGenerator& FProfileSweepGenerator::Generate()
 	{
 		if (WeldedVertices.Contains(ProfileIndex))
 		{
+			FVector3d FramePoint = ProfileCurve[ProfileIndex];
 			// Position stays locked into the first frame
+			if (!SweepScaleCurve.IsEmpty())
+			{
+				FramePoint *= SweepScaleCurve[0];
+			}
 			SetVertex(GetVertIndex(true, 0, ProfileIndex, NumWelded, NumNonWelded, VertPositionOffsets), 
-				SweepCurve[0].FromFramePoint(ProfileCurve[ProfileIndex]));
+				SweepCurve[0].FromFramePoint(FramePoint));
 		}
 		else
 		{
 			// Generate copies of the vertex in all the sweep frames.
 			for (int32 SweepIndex = 0; SweepIndex < NumSweepPoints; ++SweepIndex)
 			{
+				FVector3d FramePoint = ProfileCurve[ProfileIndex];
+				if (!SweepScaleCurve.IsEmpty())
+				{
+					FramePoint *= SweepScaleCurve[SweepIndex];
+				}
 				SetVertex(GetVertIndex(false, SweepIndex, ProfileIndex, NumWelded, NumNonWelded, VertPositionOffsets),
-					SweepCurve[SweepIndex].FromFramePoint(ProfileCurve[ProfileIndex]));
+					SweepCurve[SweepIndex].FromFramePoint(FramePoint));
 			}
 		}
 	});
diff --git a/Engine/Source/Runtime/GeometryCore/Public/Generators/SweepGenerator.h b/Engine/Source/Runtime/GeometryCore/Public/Generators/SweepGenerator.h
--- a/Engine/Source/Runtime/GeometryCore/Public/Generators/SweepGenerator.h
+++ b/Engine/Source/Runtime/GeometryCore/Public/Generators/SweepGenerator.h
@@ -801,9 +801,12 @@ public:
 	// Curve that will be swept along the curve, given in coordinates of the frames used in the sweep curve.
 	TArray<FVector3d> ProfileCurve;
 
-	// Curve along which to sweep the profile curve
+	// Curve along which to sweep the profile curve.
 	TArray<FFrame3d> SweepCurve;
 
+	// (Optional) Curve along which to scale the profile curve, corresponding to each frame in SweepCurve.
+	TArray<FVector3d> SweepScaleCurve;
+
 	// Indices into ProfileCurve that should not be swept along the curve, instead being instantiated
 	// just once. This is useful for welding vertices on an axis of rotation if the sweep curve denotes
 	// a revolution.
