Add OrientedBox2 to BoxTypes.h + some test coverage
Add a TransformByTranspose to Matrix3, so we can efficiently inverse-transform when we know we have a rotation matrix

diff --git a/Engine/Source/Runtime/GeometryCore/Public/MatrixTypes.h b/Engine/Source/Runtime/GeometryCore/Public/MatrixTypes.h
--- a/Engine/Source/Runtime/GeometryCore/Public/MatrixTypes.h
+++ b/Engine/Source/Runtime/GeometryCore/Public/MatrixTypes.h
@@ -212,6 +212,14 @@ struct TMatrix3
 			Row0.Z, Row1.Z, Row2.Z);
 	}
 
+	UE::Math::TVector<RealType> TransformByTranspose(const UE::Math::TVector<RealType>& V) const
+	{
+		return UE::Math::TVector<RealType>(
+			Row0.X * V.X + Row1.X * V.Y + Row2.X * V.Z,
+			Row0.Y * V.X + Row1.Y * V.Y + Row2.Y * V.Z,
+			Row0.Z * V.X + Row1.Z * V.Y + Row2.Z * V.Z);
+	}
+
 	// Computes |A|(A^-1)^T. 
 	// This value is sometimes useful for computing rotated surface normals from a continuous deformation.
 	// Since computing the 3x3 inverse using Cramer's rule divides by the 3x3 determinant in the final step, this 
diff --git a/Engine/Source/Runtime/GeometryCore/Public/OrientedBoxTypes.h b/Engine/Source/Runtime/GeometryCore/Public/OrientedBoxTypes.h
--- a/Engine/Source/Runtime/GeometryCore/Public/OrientedBoxTypes.h
+++ b/Engine/Source/Runtime/GeometryCore/Public/OrientedBoxTypes.h
@@ -205,83 +205,259 @@ struct TOrientedBox3
 	/**
 	 * Find squared distance to box.
 	 * @param Point input point
-	 * @return squared distance from point to box, or 0 if point is inside box
+	 * @return Squared distance from point to box, or 0 if point is inside box
 	 */
-	RealType DistanceSquared(TVector<RealType> Point)
-	{
-		 // Ported from WildMagic5 Wm5DistPoint3Box3.cpp
-			 
-		 // Work in the box's coordinate system.
-		 Point -= Frame.Origin;
-
-		// Compute squared distance and closest point on box.
-		RealType sqrDistance = 0;
-		RealType delta;
-		TVector<RealType> closest(0, 0, 0);
-		for (int i = 0; i < 3; ++i) 
-		{
-			closest[i] = Point.Dot(GetAxis(i));
-			if (closest[i] < -Extents[i]) 
-			{
-				delta = closest[i] + Extents[i];
-				sqrDistance += delta * delta;
-				closest[i] = -Extents[i];
-			}
-			else if (closest[i] > Extents[i])
-			{
-				delta = closest[i] - Extents[i];
-				sqrDistance += delta * delta;
-				closest[i] = Extents[i];
-			}
-		}
-
-		return sqrDistance;
+	RealType DistanceSquared(TVector<RealType> Point) const
+	{
+		TVector<RealType> Local = Frame.ToFramePoint(Point);
+		TVector<RealType> BeyondExtents = Local.GetAbs() - Extents;
+		TVector<RealType> ClampedBeyond( // clamp negative (inside) to zero
+			FMath::Max((RealType)0, BeyondExtents.X),
+			FMath::Max((RealType)0, BeyondExtents.Y),
+			FMath::Max((RealType)0, BeyondExtents.Z));
+		return ClampedBeyond.SquaredLength();
 	}
 
-
-
-
 	 /**
 	  * Find closest point on box
 	  * @param Point input point
-	  * @return closest point on box. Input point is returned if it is inside box.
+	  * @return Closest point on box. Input point is returned if it is inside box.
 	  */
-	 TVector<RealType> ClosestPoint(TVector<RealType> Point)
+	 TVector<RealType> ClosestPoint(const TVector<RealType>& Point) const
 	 {
-		 // Work in the box's coordinate system.
-		 Point -= Frame.Origin;
-
-		 // Compute squared distance and closest point on box.
-		 RealType sqrDistance = 0;
-		 RealType delta;
-		 TVector<RealType> closest;
-		 TVector<RealType> Axes[3];
-		 for (int i = 0; i < 3; ++i) 
-		 {
-			 Axes[i] = GetAxis(i);
-			 closest[i] = Point.Dot(Axes[i]);
-			 RealType extent = Extents[i];
-			 if (closest[i] < -extent) 
-			 {
-				 delta = closest[i] + extent;
-				 sqrDistance += delta * delta;
-				 closest[i] = -extent;
-			 }
-			 else if (closest[i] > extent) 
-			 {
-				 delta = closest[i] - extent;
-				 sqrDistance += delta * delta;
-				 closest[i] = extent;
-			 }
-		 }
-
-		 return Frame.Origin + closest.X*Axes[0] + closest.Y*Axes[1] + closest.Z*Axes[2];
+		 TMatrix3<RealType> RotMatrix = Frame.Rotation.ToRotationMatrix();
+		 TVector<RealType> FromOrigin = Point - Frame.Origin;
+		 // Transform to local space
+		 TVector<RealType> Local(RotMatrix.TransformByTranspose(FromOrigin));
+		 // Clamp to box
+		 Local.X = FMath::Clamp(Local.X, -Extents.X, Extents.X);
+		 Local.Y = FMath::Clamp(Local.Y, -Extents.Y, Extents.Y);
+		 Local.Z = FMath::Clamp(Local.Z, -Extents.Z, Extents.Z);
+		 // Transform back
+		 return Frame.Origin + RotMatrix * Local;
 	 }
 
 
 
 };
 
+template<typename RealType>
+struct TOrientedBox2
+{
+	// Center of the box
+	TVector2<RealType> Origin;
+	// X Axis of the box -- must be normalized
+	TVector2<RealType> UnitAxisX;
+	/** Half-dimensions of box measured along the two axes */
+	TVector2<RealType> Extents;
+
+	TOrientedBox2() : Origin(0, 0), UnitAxisX(1, 0), Extents(1, 1) {}
+
+	/**
+	 * Create axis-aligned box with given Origin and Extents
+	 */
+	TOrientedBox2(const TVector2<RealType>& OriginIn, const TVector2<RealType>& ExtentsIn)
+		: Origin(OriginIn), UnitAxisX(1, 0), Extents(ExtentsIn)
+	{
+	}
+
+	/**
+	 * Create oriented box with given Origin, X Axis and Extents
+	 */
+	TOrientedBox2(const TVector2<RealType>& OriginIn, const TVector2<RealType>& XAxisIn, const TVector2<RealType>& ExtentsIn)
+		: Origin(OriginIn), UnitAxisX(XAxisIn), Extents(ExtentsIn)
+	{
+		if (!UnitAxisX.Normalize())
+		{
+			UnitAxisX = TVector2<RealType>(1, 0);
+		}
+	}
+
+	/**
+	 * Create oriented box with given Origin, Angle (in radians) and Extents
+	 */
+	TOrientedBox2(const TVector2<RealType>& OriginIn, RealType AngleRad, const TVector2<RealType>& ExtentsIn)
+		: Origin(OriginIn), UnitAxisX(FMath::Cos(AngleRad), FMath::Sin(AngleRad)), Extents(ExtentsIn)
+	{
+	}
+
+	/**
+	 * Create oriented box from axis-aligned box
+	 */
+	TOrientedBox2(const TAxisAlignedBox2<RealType>& AxisBox)
+		: Origin(AxisBox.Center()), UnitAxisX(1,0), Extents(AxisBox.Extents())
+	{
+	}
+
+	void SetAngleRadians(RealType AngleRad)
+	{
+		UnitAxisX = TVector2<RealType>(FMath::Cos(AngleRad), FMath::Sin(AngleRad));
+	}
+
+	/** @return box with unit dimensions centered at origin */
+	static TOrientedBox2<RealType> UnitZeroCentered() { return TOrientedBox2<RealType>(TVector2<RealType>::Zero(), (RealType)0.5 * TVector2<RealType>::One()); }
+
+	/** @return box with unit dimensions where minimum corner is at origin */
+	static TOrientedBox2<RealType> UnitPositive() { return TOrientedBox2<RealType>((RealType)0.5 * TVector2<RealType>::One(), (RealType)0.5 * TVector2<RealType>::One()); }
+
+
+	/** @return center of the box */
+	inline TVector2<RealType> Center() const { return Origin; }
+
+	/** @return X axis of the box */
+	inline TVector2<RealType> AxisX() const { return UnitAxisX; }
+
+	/** @return Y axis of the box  (the X axis rotated 90 degrees counter-clockwise) */
+	inline TVector2<RealType> AxisY() const { return TVector2<RealType>(-UnitAxisX.Y, UnitAxisX.X); }
+
+	/** @return an axis of the box */
+	inline TVector2<RealType> GetAxis(int AxisIndex) const { return AxisIndex == 0 ? AxisX() : AxisY(); }
+
+	/** @return maximum extent of box */
+	inline RealType MaxExtent() const
+	{
+		return Extents.MaxAbs();
+	}
+
+	/** @return minimum extent of box */
+	inline RealType MinExtent() const
+	{
+		return Extents.MinAbs();
+	}
+
+	/** @return vector from minimum-corner to maximum-corner of box */
+	inline TVector2<RealType> Diagonal() const
+	{
+		return (AxisX() * Extents.X + AxisY() * Extents.Y) * (RealType)2;
+	}
+
+	/** @return area of box */
+	inline RealType Area() const
+	{
+		return (RealType)4 * (Extents.X) * (Extents.Y);
+	}
+
+	/** @return perimeter of box */
+	inline RealType Perimeter() const
+	{
+		return (RealType)4 * (Extents.X + Extents.Y);
+	}
+
+	/** @return Point transformed to the local space of the box (Origin at box center, Axes aligned to box) */
+	inline TVector2<RealType> ToLocalSpace(const TVector2<RealType>& Point) const
+	{
+		TVector2<RealType> FromOrigin = Point - Origin;
+		return TVector2<RealType>(UnitAxisX.Dot(FromOrigin), AxisY().Dot(FromOrigin));
+	}
+
+	/** @return Point transformed from the local space of the box (Origin at box center, Axes aligned to box) */
+	inline TVector2<RealType> FromLocalSpace(const TVector2<RealType>& Point) const
+	{
+		return Origin + Point.X * UnitAxisX + Point.Y * AxisY();
+	}
+
+	/** @return true if box contains point */
+	inline bool Contains(const TVector<RealType>& Point) const
+	{
+		TVector<RealType> LocalPoint = ToLocalPoint(Point);
+		return (TMathUtil<RealType>::Abs(LocalPoint.X) <= Extents.X) &&
+			(TMathUtil<RealType>::Abs(LocalPoint.Y) <= Extents.Y) &&
+			(TMathUtil<RealType>::Abs(LocalPoint.Z) <= Extents.Z);
+	}
+
+
+	//// corners 
+
+	/**
+	 * @param Index corner index in range 0-3
+	 * @return Corner point on the box identified by the given index.
+	 * Ordering is: [ (-x,-y), (x,-y), (x,y), (-x,y) ]
+	 */
+	TVector2<RealType> GetCorner(int Index) const
+	{
+		check(Index >= 0 && Index <= 3);
+		RealType DX = RealType(int((Index == 1) | (Index == 2)) * 2 - 1) * Extents.X; // X positive at indices 1, 2
+		RealType DY = RealType((Index & 2) - 1) * Extents.Y; // Y positive at indices 2, 3
+		return Origin + DX * UnitAxisX + DY * AxisY();
+	}
+
+	/**
+	 * Call CornerPointFunc(TVector2<RealType>) for each of the 4 box corners. Order is the same as GetCorner(X).
+	 */
+	template<typename PointFuncType>
+	void EnumerateCorners(PointFuncType CornerPointFunc) const
+	{
+		TVector2<RealType> X = AxisX() * Extents.X, Y = AxisY() * Extents.Y;
+
+		CornerPointFunc(Origin - X - Y);
+		CornerPointFunc(Origin + X - Y);
+		CornerPointFunc(Origin + X + Y);
+		CornerPointFunc(Origin - X + Y);
+	}
+
+	/**
+	 * Call CornerPointPredicate(TVector2<RealType>) for each of the 4 box corners, with early-out if any call returns false
+	 * @return true if all tests pass
+	 */
+	template<typename PointPredicateType>
+	bool TestCorners(PointPredicateType CornerPointPredicate) const
+	{
+		TVector2<RealType> X = AxisX() * Extents.X, Y = AxisY() * Extents.Y;
+
+		return
+			CornerPointFunc(Origin - X - Y) &&
+			CornerPointFunc(Origin + X - Y) &&
+			CornerPointFunc(Origin + X + Y) &&
+			CornerPointFunc(Origin - X + Y);
+	}
+
+	/**
+	 * Get whether the corner at Index (see GetCorner documentation comment) is in the negative or positive direction for each axis
+	 * @param Index corner index in range 0-4
+	 * @return Index2i with 0 or 1 for each axis, 0 if corner is in the negative direction for that axis, 1 if in the positive direction
+	 */
+	static FIndex2i GetCornerSide(int Index)
+	{
+		check(Index >= 0 && Index <= 3);
+		return FIndex2i(
+			int((Index == 1) | (Index == 2)),
+			(Index & 2) >> 1
+		);
+	}
+
+
+
+	/**
+	 * Find squared distance to box.
+	 * @param Point input point
+	 * @return Squared distance from point to box, or 0 if point is inside box
+	 */
+	RealType DistanceSquared(const TVector2<RealType>& Point) const
+	{
+		TVector2<RealType> BeyondExtents = ToLocalSpace(Point).GetAbs() - Extents;
+		TVector2<RealType> ClampedBeyond( // clamp negative (inside) to zero
+			FMath::Max((RealType)0, BeyondExtents.X),
+			FMath::Max((RealType)0, BeyondExtents.Y));
+		return ClampedBeyond.SquaredLength();
+	}
+
+	/**
+	 * Find closest point on box
+	 * @param Point input point
+	 * @return Closest point on box. Input point is returned if it is inside box.
+	 */
+	TVector2<RealType> ClosestPoint(const TVector2<RealType>& Point) const
+	{
+		TVector2<RealType> Local = ToLocalSpace(Point);
+		Local.X = FMath::Clamp(Local.X, -Extents.X, Extents.X);
+		Local.Y = FMath::Clamp(Local.Y, -Extents.Y, Extents.Y);
+		return FromLocalSpace(Local);
+	}
+
+};
+
+typedef TOrientedBox2<float> FOrientedBox2f;
+typedef TOrientedBox2<double> FOrientedBox2d;
 typedef TOrientedBox3<float> FOrientedBox3f;
 typedef TOrientedBox3<double> FOrientedBox3d;
 
