GeometryCore: DynamicVertex Attributes.  fixing type promotion warnings.

diff --git a/Engine/Source/Runtime/GeometryCore/Public/DynamicMesh/DynamicVertexAttribute.h b/Engine/Source/Runtime/GeometryCore/Public/DynamicMesh/DynamicVertexAttribute.h
--- a/Engine/Source/Runtime/GeometryCore/Public/DynamicMesh/DynamicVertexAttribute.h
+++ b/Engine/Source/Runtime/GeometryCore/Public/DynamicMesh/DynamicVertexAttribute.h
@@ -513,10 +513,10 @@ protected:
 		int IndexSet = AttribDimension * SetAttribute;
 		int IndexA = AttribDimension * AttributeA;
 		int IndexB = AttribDimension * AttributeB;
-		AttribValueType Beta = ((AttribValueType)1 - Alpha);
+		double Beta = (1. - Alpha);
 		for (int i = 0; i < AttribDimension; ++i)
 		{
-			AttribValues[IndexSet+i] = Beta*AttribValues[IndexA+i] + Alpha*AttribValues[IndexB+i];
+			AttribValues[IndexSet+i] = AttribValueType(Beta*AttribValues[IndexA+i] + Alpha*AttribValues[IndexB+i]);
 		}
 	}
 
@@ -529,8 +529,8 @@ protected:
 		int IndexC = AttribDimension * AttributeC;
 		for (int i = 0; i < AttribDimension; ++i)
 		{
-			AttribValues[IndexSet + i] = 
-				BaryCoords.X*AttribValues[IndexA+i] + BaryCoords.Y*AttribValues[IndexB+i] + BaryCoords.Z*AttribValues[IndexC+i];
+			AttribValues[IndexSet + i] = AttribValueType(
+				BaryCoords.X*AttribValues[IndexA+i] + BaryCoords.Y*AttribValues[IndexB+i] + BaryCoords.Z*AttribValues[IndexC+i]);
 		}
 	}
 
diff --git a/Engine/Source/Runtime/GeometryCore/Public/DynamicMesh/DynamicVertexSkinWeightsAttribute.h b/Engine/Source/Runtime/GeometryCore/Public/DynamicMesh/DynamicVertexSkinWeightsAttribute.h
--- a/Engine/Source/Runtime/GeometryCore/Public/DynamicMesh/DynamicVertexSkinWeightsAttribute.h
+++ b/Engine/Source/Runtime/GeometryCore/Public/DynamicMesh/DynamicVertexSkinWeightsAttribute.h
@@ -483,7 +483,7 @@ protected:
 	void SetBoneWeightsFromLerp(int SetAttribute, int AttributeA, int AttributeB, double Alpha)
 	{
 		Alpha = FMath::Clamp(Alpha, 0.0, 1.0);
-		VertexBoneWeights[SetAttribute] = FBoneWeights::Blend(VertexBoneWeights[AttributeA], VertexBoneWeights[AttributeB], Alpha);
+		VertexBoneWeights[SetAttribute] = FBoneWeights::Blend(VertexBoneWeights[AttributeA], VertexBoneWeights[AttributeB], (float)Alpha);
 	}
 
 	/** Set the value at an Attribute to be a barycentric interpolation of three other Attributes */
@@ -494,9 +494,9 @@ protected:
 		if (!FMath::IsNearlyZero(BaryCoords.Y + BaryCoords.Z))
 		{
 			const double BCW = BaryCoords.Y / (BaryCoords.Y + BaryCoords.Z);
-			const FBoneWeights BC = FBoneWeights::Blend(VertexBoneWeights[AttributeB], VertexBoneWeights[AttributeC], BCW);
+			const FBoneWeights BC = FBoneWeights::Blend(VertexBoneWeights[AttributeB], VertexBoneWeights[AttributeC], (float)BCW);
 
-			VertexBoneWeights[SetAttribute] = FBoneWeights::Blend(VertexBoneWeights[AttributeA], BC, BaryCoords.X);
+			VertexBoneWeights[SetAttribute] = FBoneWeights::Blend(VertexBoneWeights[AttributeA], BC, (float)BaryCoords.X);
 		}
 		else if (SetAttribute != AttributeA)
 		{
