GeometryCore:
Add FDynamicMesh3::EnableMatchingAttributes(), to expose AttributeSet::EnableMatchingAttributes() at the Mesh level
Add FDynamicMesh3::GetTriEdgeIDFromEdgeID(), helper to simplify construction of FMeshTriEdgeID identifiers
Add FRotator conversions for TQuaternion
Add TTransformSRT3::GetRotator()
Add UE::Geometry::Lerp variant for TFrame3

diff --git a/Engine/Source/Runtime/GeometryCore/Private/DynamicMesh/DynamicMesh3.cpp b/Engine/Source/Runtime/GeometryCore/Private/DynamicMesh/DynamicMesh3.cpp
--- a/Engine/Source/Runtime/GeometryCore/Private/DynamicMesh/DynamicMesh3.cpp
+++ b/Engine/Source/Runtime/GeometryCore/Private/DynamicMesh/DynamicMesh3.cpp
@@ -335,6 +335,68 @@ void FDynamicMesh3::Clear()
 }
 
 
+void FDynamicMesh3::EnableMatchingAttributes(const FDynamicMesh3& ToMatch, bool bClearExisting)
+{
+	bool bWantVertexNormals = ToMatch.HasVertexNormals() || (this->HasVertexNormals() && bClearExisting == false);
+	if (bClearExisting || bWantVertexNormals == false)
+	{
+		DiscardVertexNormals();
+	}
+	if (bWantVertexNormals)
+	{
+		EnableVertexNormals(FVector3f::UnitZ());
+	}
+
+	bool bWantVertexColors = ToMatch.HasVertexColors() || (this->HasVertexColors() && bClearExisting == false);
+	if (bClearExisting || bWantVertexColors == false)
+	{
+		DiscardVertexColors();
+	}
+	if (bWantVertexColors)
+	{
+		EnableVertexColors(FVector3f::Zero());
+	}
+
+
+	bool bWantVertexUVs = ToMatch.HasVertexUVs() || (this->HasVertexUVs() && bClearExisting == false);
+	if (bClearExisting || bWantVertexUVs == false)
+	{
+		DiscardVertexUVs();
+	}
+	if (bWantVertexUVs)
+	{
+		EnableVertexUVs(FVector2f::Zero());
+	}
+
+
+	bool bWantTriangleGroups = ToMatch.HasTriangleGroups() || (this->HasTriangleGroups() && bClearExisting == false);
+	if (bClearExisting || bWantTriangleGroups == false)
+	{
+		DiscardTriangleGroups();
+	}
+	if (bWantTriangleGroups)
+	{
+		EnableTriangleGroups();
+	}
+
+
+	bool bWantAttributes = ToMatch.HasAttributes() || (this->HasAttributes() && bClearExisting == false);
+	if (bClearExisting || bWantAttributes == false)
+	{
+		DiscardAttributes();
+	}
+	if (bWantAttributes)
+	{
+		EnableAttributes();
+	}
+	if (HasAttributes() && ToMatch.HasAttributes())
+	{
+		Attributes()->EnableMatchingAttributes(*ToMatch.Attributes(), bClearExisting);
+	}
+
+}
+
+
 
 int FDynamicMesh3::GetComponentsFlags() const
 {
diff --git a/Engine/Source/Runtime/GeometryCore/Public/DynamicMesh/DynamicMesh3.h b/Engine/Source/Runtime/GeometryCore/Public/DynamicMesh/DynamicMesh3.h
--- a/Engine/Source/Runtime/GeometryCore/Public/DynamicMesh/DynamicMesh3.h
+++ b/Engine/Source/Runtime/GeometryCore/Public/DynamicMesh/DynamicMesh3.h
@@ -242,6 +242,13 @@ public:
 	/** Discard all data */
 	void Clear();
 
+	/**
+	 * Ensure that all the same extended attributes available in ToMatch are also enabled.
+	 * By default, clears existing attributes, so that there will be an exact match
+	 * If bClearExisting is passed as false, existing attributes are not removed/cleared.
+	 */
+	void EnableMatchingAttributes(const FDynamicMesh3& ToMatch, bool bClearExisting = true);
+
 	/**
 	 * Serialization operator for FDynamicMesh3.
 	 *
@@ -743,6 +750,21 @@ public:
 	/** Return edge vertex indices, but oriented based on attached triangle (rather than min-sorted) */
 	FIndex2i GetOrientedBoundaryEdgeV(int EdgeID) const;
 
+	/** Return (triangle, edge_index) representation for given Edge ID */
+	inline FMeshTriEdgeID GetTriEdgeIDFromEdgeID(int EdgeID) const
+	{
+		checkSlow(IsEdge(EdgeID));
+		int32 TriIndex = Edges[EdgeID].Tri.A;
+		FIndex3i TriEdges = TriangleEdges[TriIndex];
+		if (TriEdges.A == EdgeID)
+		{
+			return FMeshTriEdgeID(TriIndex, 0);
+		}
+		{ 
+			return FMeshTriEdgeID(TriIndex, ( TriEdges.B == EdgeID ) ? 1 : 2 );
+		}
+	}
+
 	//
 	// Vertex and Triangle attribute arrays
 	//
diff --git a/Engine/Source/Runtime/GeometryCore/Public/FrameTypes.h b/Engine/Source/Runtime/GeometryCore/Public/FrameTypes.h
--- a/Engine/Source/Runtime/GeometryCore/Public/FrameTypes.h
+++ b/Engine/Source/Runtime/GeometryCore/Public/FrameTypes.h
@@ -464,6 +464,16 @@ struct TFrame3
 
 };
 
+template <typename T>
+TFrame3<T> Lerp(const TFrame3<T>& A, const TFrame3<T>& B, T Alpha)
+{
+	return TFrame3<T>(
+		UE::Geometry::Lerp(A.Origin, B.Origin, Alpha),
+		TQuaternion<T>(A.Rotation, B.Rotation, Alpha));
+}
+
+
+
 typedef TFrame3<float> FFrame3f;
 typedef TFrame3<double> FFrame3d;
 
diff --git a/Engine/Source/Runtime/GeometryCore/Public/Quaternion.h b/Engine/Source/Runtime/GeometryCore/Public/Quaternion.h
--- a/Engine/Source/Runtime/GeometryCore/Public/Quaternion.h
+++ b/Engine/Source/Runtime/GeometryCore/Public/Quaternion.h
@@ -116,6 +116,10 @@ PRAGMA_ENABLE_DEPRECATION_WARNINGS
 		Quat.W = (double)W;
 		return Quat;
 	}
+	explicit inline operator FRotator() const
+	{
+		return ((FQuat)*this).Rotator();
+	}
 	explicit inline TQuaternion(const FQuat4f& Quat)
 	{
 		X = (RealType)Quat.X;
@@ -130,6 +134,14 @@ PRAGMA_ENABLE_DEPRECATION_WARNINGS
 		Z = (RealType)Quat.Z;
 		W = (RealType)Quat.W;
 	}
+	explicit inline TQuaternion(const FRotator& Rotator)
+	{
+		FQuat4d Quat(Rotator);
+		X = (RealType)Quat.X;
+		Y = (RealType)Quat.Y;
+		Z = (RealType)Quat.Z;
+		W = (RealType)Quat.W;
+	}
 
 };
 
diff --git a/Engine/Source/Runtime/GeometryCore/Public/TransformTypes.h b/Engine/Source/Runtime/GeometryCore/Public/TransformTypes.h
--- a/Engine/Source/Runtime/GeometryCore/Public/TransformTypes.h
+++ b/Engine/Source/Runtime/GeometryCore/Public/TransformTypes.h
@@ -109,6 +109,14 @@ public:
 		return Rotation; 
 	}
 
+	/**
+	* @return Rotation portion of Transform, as FRotator
+	*/
+	FRotator GetRotator() const 
+	{ 
+		return (FRotator)Rotation; 
+	}
+
 	/** 
 	 * Set Rotation portion of Transform 
 	 */
