GeometryCore : fixing implicit conversion warnings, ParallelTransformReduce

diff --git a/Engine/Source/Runtime/GeometryCore/Public/Async/ParallelTransformReduce.h b/Engine/Source/Runtime/GeometryCore/Public/Async/ParallelTransformReduce.h
--- a/Engine/Source/Runtime/GeometryCore/Public/Async/ParallelTransformReduce.h
+++ b/Engine/Source/Runtime/GeometryCore/Public/Async/ParallelTransformReduce.h
@@ -10,11 +10,11 @@ namespace Geometry {
 
 	//
 	// Index interface (similar to ParallelFor)
-	// 	   TransformFuncT should be a function pointer-like object with signature: T(int)
+	// 	   TransformFuncT should be a function pointer-like object with signature: T(IntType )
 	// 	   ReduceFuncT should be a function pointer-like object with signature: T(T,T)
 
-	template<typename T, typename TransformFuncT, typename ReduceFuncT>
-	T ParallelTransformReduce(int64 Num,
+	template<typename IntType, typename T, typename TransformFuncT, typename ReduceFuncT>
+	T ParallelTransformReduce(IntType Num,
 							  const T& Init,
 							  TransformFuncT Transform,
 							  ReduceFuncT Reduce,
@@ -24,7 +24,7 @@ namespace Geometry {
 		// ParallelFor doesn't yet support int64 NumTasks, so cap it
 		const int32 NumTasks = (int32) FMath::Min(InNumTasks, MAX_int32);
 
-		const int64 NumPerTask = (int64)FMathf::Ceil((float)Num / NumTasks);
+		const IntType NumPerTask = (IntType)FMathf::Ceil((float)Num / NumTasks);
 
 		TArray<T> PerTaskResults;
 		PerTaskResults.SetNum(NumTasks);
@@ -33,8 +33,8 @@ namespace Geometry {
 		{
 			T LocalResult{ Init };
 
-			int64 End = FMath::Min((TaskIndex + 1) * NumPerTask, Num);
-			for (int64 Index = TaskIndex * NumPerTask; Index < End; ++Index)
+			IntType End = FMath::Min((TaskIndex + 1) * NumPerTask, Num);
+			for (IntType Index = TaskIndex * NumPerTask; Index < End; ++Index)
 			{
 				T Transformed{ Transform(Index) };
 				LocalResult = Reduce(Transformed, LocalResult);
@@ -52,14 +52,15 @@ namespace Geometry {
 	}
 
 
+
 	//
 	// Index interface for non-copyable types
 	// 	   InitFuncT should be a function-like object with signature void(T&)
-	// 	   TransformFuncT should be a function-like object with signature void(int,T&)
+	// 	   TransformFuncT should be a function-like object with signature void(IntType,T&)
 	// 	   ReduceFuncT should be a function-like object with signature void(T,T&)
 
-	template<typename T, typename InitFuncT, typename TransformFuncT, typename ReduceFuncT>
-	void ParallelTransformReduce(int64 Num,
+	template<typename IntType, typename T, typename InitFuncT, typename TransformFuncT, typename ReduceFuncT>
+	void ParallelTransformReduce(IntType Num,
 								 InitFuncT InitFunc,
 								 TransformFuncT Transform,
 								 ReduceFuncT Reduce,
@@ -71,7 +72,7 @@ namespace Geometry {
 		// ParallelFor doesn't yet support int64 NumTasks, so cap it
 		const int32 NumTasks = (int32)FMath::Min(InNumTasks, MAX_int32);
 
-		int64 NumPerTask = (int64)FMathf::Ceil((float)Num / NumTasks);
+		IntType NumPerTask = (IntType)FMathf::Ceil((float)Num / NumTasks);
 
 		TArray<T> PerTaskResults;
 		PerTaskResults.SetNum(NumTasks);
@@ -80,8 +81,8 @@ namespace Geometry {
 		{
 			InitFunc(PerTaskResults[TaskIndex]);
 
-			int64 End = FMath::Min((TaskIndex + 1) * NumPerTask, Num);
-			for (int64 Index = TaskIndex * NumPerTask; Index < End; ++Index)
+			IntType End = FMath::Min((TaskIndex + 1) * NumPerTask, Num);
+			for (IntType Index = TaskIndex * NumPerTask; Index < End; ++Index)
 			{
 				T Transformed;
 				Transform(Index, Transformed);
@@ -136,7 +137,7 @@ namespace Geometry {
 	template<typename ContainerType, typename ElementType, typename SizeType>
 	TIndexedContainerIterator<ContainerType, ElementType, SizeType>
 		AdvanceIterator(TIndexedContainerIterator<ContainerType, ElementType, SizeType> Begin,
-						int64 N)
+						SizeType N)
 	{
 		return Begin + N;
 	}
@@ -165,12 +166,12 @@ namespace Geometry {
 		TArray<T> PerTaskResults;
 		PerTaskResults.SetNum(NumTasks);
 
-		ParallelFor(NumTasks, [&](int64 TaskIndex)
+		ParallelFor(NumTasks, [&](int32 TaskIndex)
 		{
 			T LocalResult{ Init };
 
-			IterT LocalIter = AdvanceIterator(BeginIterator, FMath::Min(Num, TaskIndex * NumPerTask));
-			IterT LocalEndIter = AdvanceIterator(BeginIterator, FMath::Min(Num, (TaskIndex + 1) * NumPerTask));
+			IterT LocalIter = AdvanceIterator(BeginIterator, FMath::Min(Num, int64(TaskIndex) * NumPerTask));
+			IterT LocalEndIter = AdvanceIterator(BeginIterator, FMath::Min(Num, (int64(TaskIndex) + 1) * NumPerTask));
 
 			while (LocalIter != LocalEndIter)
 			{
@@ -217,12 +218,12 @@ namespace Geometry {
 		TArray<T> PerTaskResults;
 		PerTaskResults.SetNum(NumTasks);
 
-		ParallelFor(NumTasks, [&](int64 TaskIndex)
+		ParallelFor(NumTasks, [&](int32 TaskIndex)
 		{
 			InitFunc(PerTaskResults[TaskIndex]);
 
-			IterT LocalIter = AdvanceIterator(BeginIterator, FMath::Min(Num, TaskIndex * NumPerTask));
-			IterT LocalEndIter = AdvanceIterator(BeginIterator, FMath::Min(Num, (TaskIndex + 1) * NumPerTask));
+			IterT LocalIter = AdvanceIterator(BeginIterator, FMath::Min(Num, int64(TaskIndex) * NumPerTask));
+			IterT LocalEndIter = AdvanceIterator(BeginIterator, FMath::Min(Num, (int64(TaskIndex) + 1) * NumPerTask));
 
 			while (LocalIter != LocalEndIter)
 			{
