Add functions to fit an oriented box2 to points, a simple polygon, or a convex hull
Add function to find a 2d convex hull of a simple polygon (faster than finding a convex hull of a point set)

diff --git a/Engine/Source/Runtime/GeometryCore/Private/CompGeom/ConvexHull2.cpp b/Engine/Source/Runtime/GeometryCore/Private/CompGeom/ConvexHull2.cpp
--- a/Engine/Source/Runtime/GeometryCore/Private/CompGeom/ConvexHull2.cpp
+++ b/Engine/Source/Runtime/GeometryCore/Private/CompGeom/ConvexHull2.cpp
@@ -7,12 +7,112 @@
 
 #include "CompGeom/ExactPredicates.h"
 #include "Algo/Unique.h"
+#include "Algo/Reverse.h"
 
 namespace UE
 {
 namespace Geometry
 {
 
+namespace
+{
+	inline int32 IncModM(int32 Idx, int32 Mod)
+	{
+		int32 Inc = Idx + 1;
+		return Inc < Mod ? Inc : 0;
+	}
+}
+
+template<typename RealType>
+bool TConvexHull2<RealType>::SolveSimplePolygon(int32 N, TFunctionRef<TVector2<RealType>(int32)> GetPointFunc, bool bIsKnownCCW)
+{
+	Dimension = 0;
+	NumUniquePoints = 0;
+
+	if (N == 0)
+	{
+		Hull.Empty();
+		return false;
+	}
+	Hull.Reset(N);
+
+
+	// Find the (lowest) leftmost point (guaranteed to be on the hull)
+	int32 LeftmostPtIdx = 0;
+	TVector2<RealType> Leftmost = GetPointFunc(0);
+	// If orientation is not already known to be CCW, we use area sign to decide orientation
+	RealType OrientSign = 1;
+
+	{
+		auto AddEdgeArea2 = [&GetPointFunc](TVector2<RealType> V1, TVector2<RealType> V2)
+		{
+			return V1.X * V2.Y - V1.Y * V2.X;
+		};
+		RealType AreaSum2 = bIsKnownCCW ? 0 : AddEdgeArea2(GetPointFunc(N - 1), Leftmost);
+		TVector2<RealType> LastPt = Leftmost; // LastPt only for computing AreaSum to check winding direction
+		for (int32 Idx = 1; Idx < N; ++Idx)
+		{
+			TVector2<RealType> Pt = GetPointFunc(Idx);
+			if (Pt.X < Leftmost.X || (Pt.X == Leftmost.X && Pt.Y < Leftmost.Y))
+			{
+				Leftmost = Pt;
+				LeftmostPtIdx = Idx;
+			}
+			if (!bIsKnownCCW)
+			{
+				AreaSum2 += AddEdgeArea2(LastPt, Pt);
+				LastPt = Pt;
+			}
+		}
+		if (!bIsKnownCCW && AreaSum2 < 0)
+		{
+			OrientSign = -1;
+		}
+	}
+
+	Hull.Add(LeftmostPtIdx);
+	TVector2<RealType> PrevPt = Leftmost;
+	int32 CurIdx = IncModM(LeftmostPtIdx, N);
+	int32 EndIdx = LeftmostPtIdx;
+	TVector2<RealType> CurPt = GetPointFunc(CurIdx);
+	for (int32 NextIdx = IncModM(CurIdx, N); CurIdx != EndIdx; NextIdx = IncModM(NextIdx, N))
+	{
+		TVector2<RealType> NextPt = GetPointFunc(NextIdx);
+		if (ExactPredicates::Orient2<RealType>(PrevPt, CurPt, NextPt)*OrientSign <= 0)
+		{
+			// Go backwards until we're out of points or we find a locally-convex point
+			while (Hull.Num() > 1)
+			{
+				CurPt = PrevPt;
+				CurIdx = Hull.Pop(false);
+				int32 PrevIdx = Hull.Last();
+				PrevPt = GetPointFunc(PrevIdx);
+				
+				if (ExactPredicates::Orient2<RealType>(PrevPt, CurPt, NextPt)*OrientSign > 0)
+				{
+					Hull.Push(CurIdx);
+					PrevPt = CurPt;
+					break;
+				}
+			}
+		}
+		else // CurPt is locally convex with Prev,Next; go ahead and add it
+		{
+			Hull.Push(CurIdx);
+			PrevPt = CurPt;
+		}
+		CurIdx = NextIdx;
+		CurPt = NextPt;
+	}
+	if (OrientSign < 0)
+	{
+		Algo::Reverse(Hull);
+	}
+	bool bFoundValid = Hull.Num() >= 3;
+	Dimension = bFoundValid ? 2 : 0;
+	return bFoundValid;
+}
+
 
 template<class RealType>
 bool TConvexHull2<RealType>::Solve(int32 NumPoints, TFunctionRef<TVector2<RealType>(int32)> GetPointFunc, TFunctionRef<bool(int32)> FilterFunc)
diff --git a/Engine/Source/Runtime/GeometryCore/Private/CompGeom/FitOrientedBox2.cpp b/Engine/Source/Runtime/GeometryCore/Private/CompGeom/FitOrientedBox2.cpp
new file mode 100644
--- /dev/null
+++ b/Engine/Source/Runtime/GeometryCore/Private/CompGeom/FitOrientedBox2.cpp
@@ -0,0 +1,261 @@
+// Copyright Epic Games, Inc. All Rights Reserved.
+
+#include "CompGeom/FitOrientedBox2.h"
+#include "Containers/Array.h"
+#include "Math/Vector.h"
+#include "OrientedBoxTypes.h"
+#include "VectorTypes.h"
+#include "IndexTypes.h"
+#include "MathUtil.h"
+#include "CompGeom/ConvexHull2.h"
+#include "CompGeom/ExactPredicates.h"
+
+namespace UE
+{
+namespace Geometry
+{
+
+template<typename RealType>
+TOrientedBox2<RealType> FitOrientedBox2ConvexHull(int32 NumPts, TFunctionRef<TVector2<RealType>(int32)> GetHullPt, EBox2FitCriteria FitMethod)
+{
+	// fallback to using the AABB in failure cases
+	auto ComputeFallback = [NumPts, &GetHullPt]
+	{
+		TAxisAlignedBox2<RealType> AABB;
+		for (int32 PtIdx = 0; PtIdx < NumPts; ++PtIdx)
+		{
+			AABB.Contain(GetHullPt(PtIdx));
+		}
+		return TOrientedBox2<RealType>(AABB);
+	};
+
+	if (NumPts < 3) // input has too few points to be a convex hull
+	{
+		return ComputeFallback();
+	}
+
+	// Extreme Indices ordered as ymin, xmax, ymax, xmin
+	FIndex4i InitialExtreme(0, 0, 0, 0);
+
+	// Function to walk to the farthest point in a given direct
+	// For first pass, don't need looping or step count
+	auto FindExtremePtInitial = [NumPts, &GetHullPt](TVector2<RealType> BasisVec, int Start, RealType& OutExtremeVal)
+	{
+		OutExtremeVal = BasisVec.Dot(GetHullPt(Start));
+		int32 ExtremeIdx = Start;
+		for (int32 Idx = Start + 1; Idx < NumPts; ++Idx)
+		{
+			RealType Val = BasisVec.Dot(GetHullPt(Idx));
+			if (Val <= OutExtremeVal)
+			{
+				return ExtremeIdx;
+			}
+			OutExtremeVal = Val;
+			ExtremeIdx = Idx;
+		}
+
+		// For first pass, we should only need to test the first point again
+		RealType LastVal = BasisVec.Dot(GetHullPt(0));
+		if (LastVal <= OutExtremeVal)
+		{
+			return ExtremeIdx;
+		}
+		OutExtremeVal = LastVal;
+		return 0;
+	};
+
+	// Function to walk to the farthest point in a given direct, and report the point + number of steps taken to get there
+	auto FindExtremePt = [NumPts, &GetHullPt](TVector2<RealType> BasisVec, int Start, RealType& OutExtremeVal, int32& OutSteps)
+	{
+		OutExtremeVal = BasisVec.Dot(GetHullPt(Start));
+		OutSteps = 0;
+		int32 ExtremeIdx = Start;
+		for (int32 Idx = Start + 1; Idx < NumPts; ++Idx)
+		{
+			RealType Val = BasisVec.Dot(GetHullPt(Idx));
+			if (Val <= OutExtremeVal)
+			{
+				return ExtremeIdx;
+			}
+			OutExtremeVal = Val;
+			ExtremeIdx = Idx;
+			++OutSteps;
+		}
+		// same loop as above, but covering [0, Start)
+		for (int32 Idx = 0; Idx < Start; ++Idx)
+		{
+			RealType Val = BasisVec.Dot(GetHullPt(Idx));
+			if (Val <= OutExtremeVal)
+			{
+				return ExtremeIdx;
+			}
+			OutExtremeVal = Val;
+			ExtremeIdx = Idx;
+			++OutSteps;
+		}
+
+		return ExtremeIdx;
+	};
+
+	auto GetEdge = [&GetHullPt, NumPts](int32 Idx)
+	{
+		const TVector2<RealType>& A = GetHullPt(Idx);
+		int32 NextIdx = Idx + 1;
+		const TVector2<RealType>& B = GetHullPt(NextIdx < NumPts ? NextIdx : 0);
+		return B - A;
+	};
+
+	// Aligned with the first edge of the hull, find the 'extreme' points
+	// in each direction, and use that to create an initial bounding box
+	TVector2<RealType> ExtremesMax, ExtremesMin;
+	InitialExtreme.A = 0;
+	TVector2<RealType> Base = GetEdge(0);
+	if (!ensure(Base.Normalize()))
+	{
+		return ComputeFallback(); // duplicate points -> not a valid hull, return the fallback bounds
+	}
+	ExtremesMin.Y = PerpCW(Base).Dot(GetHullPt(0)); // Extreme down point is the start point
+	InitialExtreme.B = FindExtremePtInitial(Base, 0, ExtremesMax.X); // Extreme right point (relative to initial edge)
+	InitialExtreme.C = FindExtremePtInitial(-PerpCW(Base), InitialExtreme.B, ExtremesMax.Y); // Extreme up point
+	InitialExtreme.D = FindExtremePtInitial(-Base, InitialExtreme.C, ExtremesMin.X); // Extreme left point
+	TAxisAlignedBox2<RealType> Range(-ExtremesMin, ExtremesMax);
+
+	auto EvalBox = FitMethod == EBox2FitCriteria::Area ?
+		[](const TAxisAlignedBox2<RealType>& Range) -> RealType
+	{
+		return Range.Area();
+	} :
+		[](const TAxisAlignedBox2<RealType>& Range) -> RealType
+	{
+		return Range.Perimeter();
+	};
+	RealType BestScore = EvalBox(Range);
+	TVector2<RealType> BestBase = Base;
+	TAxisAlignedBox2<RealType> BestRange = Range;
+
+	// Now consider each edge in sequence, update the 'extreme' points in each direction to compute a new bounding box
+	// Note this is a variant of rotating calipers algorithm, but rather than advance by the smallest rotation at each step, it instead
+	// always advances the bottom edge one step, and advances the extreme points forward from their saved indices as needed
+	FIndex4i LastOffsets = InitialExtreme;
+	FIndex4i LastExtreme = InitialExtreme;
+	for (int32 Idx = 1; Idx < NumPts; ++Idx)
+	{
+		Base = GetEdge(Idx);
+		if (!ensure(Base.Normalize()))
+		{
+			return ComputeFallback(); // duplicate points -> not a valid hull, return the fallback bounds
+		}
+
+		TVector2<RealType> CurDir = PerpCW(Base);
+		// first extreme point is the starting point
+		ExtremesMin.Y = CurDir.Dot(GetHullPt(Idx));
+		int32 CurSteps = 0;
+		int32 CurIdx = Idx;
+		RealType* RangeValues[3]{ &ExtremesMax.X, &ExtremesMax.Y, &ExtremesMin.X };
+		for (int32 Turn = 0; Turn < 3; ++Turn)
+		{
+			CurDir = -PerpCW(CurDir);
+			int32 LastOffset = LastOffsets[Turn + 1] - 1;
+			if (LastOffset > CurSteps) // jump to the saved extreme point if we've not already passed it
+			{
+				CurSteps = LastOffset;
+				CurIdx = LastExtreme[Turn + 1];
+			}
+			int32 TakenSteps = 0;
+			CurIdx = LastExtreme[Turn + 1] = FindExtremePt(CurDir, CurIdx, *RangeValues[Turn], TakenSteps);
+			CurSteps += TakenSteps;
+			LastOffsets[Turn + 1] = CurSteps;
+		}
+		// Use the distances travelled in each direction to evaluate the bounding box for this edge
+		Range = TAxisAlignedBox2<RealType>(-ExtremesMin, ExtremesMax);
+		RealType Score = EvalBox(Range);
+		if (Score < BestScore)
+		{
+			BestScore = Score;
+			BestBase = Base;
+			BestRange = Range;
+		}
+	}
+
+	TVector2<RealType> RangeCenter = BestRange.Center();
+	TVector2<RealType> Center = RangeCenter.X * BestBase - RangeCenter.Y * PerpCW(BestBase);
+	TOrientedBox2<RealType> ToRet(Center, BestBase, BestRange.Extents());
+	return ToRet;
+}
+
+// Compute the best fit box for a simple polygon. If the polygon is already known to be a convex hull, computes the bounds faster.
+template <typename RealType>
+TOrientedBox2<RealType> FitOrientedBox2SimplePolygon(TArrayView<const TVector2<RealType>> Polygon, EBox2FitCriteria FitMethod)
+{
+	int32 Num = Polygon.Num();
+	if (!ensure(Num >= 3)) // fall back to an AABB if input has too few points to be a simple polygon
+	{
+		return TOrientedBox2<RealType>(TAxisAlignedBox2<RealType>(Polygon));
+	}
+
+	TConvexHull2<RealType> Hull;
+	if (!Hull.SolveSimplePolygon(Num, [&Polygon](int32 Idx)
+		{
+			return Polygon[Idx];
+		}, false))
+	{
+		// fall back to an AABB to gracefully fail when hull cannot be found
+		return TOrientedBox2<RealType>(TAxisAlignedBox2<RealType>(Polygon));
+	}
+
+	const TArray<int32>& Indices = Hull.GetPolygonIndices();
+	return FitOrientedBox2ConvexHull<RealType>(Indices.Num(), [&Indices, &Polygon](int32 Idx)
+		{
+			return Polygon[Indices[Idx]];
+		}, FitMethod);
+}
+
+template <typename RealType>
+TOrientedBox2<RealType> FitOrientedBox2Points(TArrayView<const TVector2<RealType>> Points, EBox2FitCriteria FitMethod)
+{
+	if (Points.Num() < 2)
+	{
+		return TOrientedBox2<RealType>(TAxisAlignedBox2<RealType>(Points));
+	}
+
+	TConvexHull2<RealType> Hull;
+	bool bHas2DHull = Hull.Solve(Points);
+	if (!bHas2DHull)
+	{
+		if (Hull.GetDimension() == 0)
+		{
+			return TOrientedBox2<RealType>(Points[0], TVector2<RealType>(1, 0), TVector2<RealType>::ZeroVector);
+		}
+		else // Dimensions == 1
+		{
+			TLine2<RealType> Line = Hull.GetLine(Points);
+			
+			TInterval1<RealType> XInterval;
+			for (const TVector2<RealType>& Pt : Points)
+			{
+				XInterval.Contain(Line.Project(Pt));
+			}
+			TVector2<RealType> Center = Line.PointAt(XInterval.Center());
+			TVector2<RealType> Extents(XInterval.Extent(), TMathUtil<RealType>::ZeroTolerance);
+			TOrientedBox2<RealType> LineBox(Center, Line.Direction, Extents);
+			return LineBox;
+		}
+	}
+
+	const TArray<int32>& Indices = Hull.GetPolygonIndices();
+	return FitOrientedBox2ConvexHull<RealType>(Indices.Num(), [&Indices, &Points](int32 Idx)
+		{
+			return Points[Indices[Idx]];
+		}, FitMethod);
+}
+
+// explicit instantiations
+template TOrientedBox2<float> GEOMETRYCORE_API FitOrientedBox2Points<float>(TArrayView<const TVector2<float>> Points, EBox2FitCriteria FitMethod);
+template TOrientedBox2<double> GEOMETRYCORE_API FitOrientedBox2Points<double>(TArrayView<const TVector2<double>> Points, EBox2FitCriteria FitMethod);
+template TOrientedBox2<float> GEOMETRYCORE_API FitOrientedBox2SimplePolygon<float>(TArrayView<const TVector2<float>> Points, EBox2FitCriteria FitMethod);
+template TOrientedBox2<double> GEOMETRYCORE_API FitOrientedBox2SimplePolygon<double>(TArrayView<const TVector2<double>> Points, EBox2FitCriteria FitMethod);
+template TOrientedBox2<float> GEOMETRYCORE_API FitOrientedBox2ConvexHull<float>(int32 NumPts, TFunctionRef<TVector2<float>(int32)> GetHullPt, EBox2FitCriteria FitMethod);
+template TOrientedBox2<double> GEOMETRYCORE_API FitOrientedBox2ConvexHull<double>(int32 NumPts, TFunctionRef<TVector2<double>(int32)> GetHullPt, EBox2FitCriteria FitMethod);
+	
+}
+}
\ No newline at end of file
diff --git a/Engine/Source/Runtime/GeometryCore/Public/BoxTypes.h b/Engine/Source/Runtime/GeometryCore/Public/BoxTypes.h
--- a/Engine/Source/Runtime/GeometryCore/Public/BoxTypes.h
+++ b/Engine/Source/Runtime/GeometryCore/Public/BoxTypes.h
@@ -598,7 +598,15 @@ struct TAxisAlignedBox2
 	{
 	}
 
-	TAxisAlignedBox2(const TArray<TVector2<RealType>>& Pts)
+	TAxisAlignedBox2(const TArray<TVector2<RealType>>& Pts) :
+		TAxisAlignedBox2(Empty())
+	{
+		*this = Empty();
+		Contain(Pts);
+	}
+
+	TAxisAlignedBox2(TArrayView<const TVector2<RealType>> Pts) :
+		TAxisAlignedBox2(Empty())
 	{
 		*this = Empty();
 		Contain(Pts);
@@ -690,6 +698,14 @@ struct TAxisAlignedBox2
 		}
 	}
 
+	void Contain(TArrayView<const TVector2<RealType>> Pts)
+	{
+		for (const TVector2<RealType>& Pt : Pts)
+		{
+			Contain(Pt);
+		}
+	}
+
 	bool Contains(const TVector2<RealType>& V) const
 	{
 		return (Min.X <= V.X) && (Min.Y <= V.Y) && (Max.X >= V.X) && (Max.Y >= V.Y);
@@ -742,6 +758,11 @@ struct TAxisAlignedBox2
 		return Width() * Height();
 	}
 
+	inline RealType Perimeter() const
+	{
+		return (Width() + Height()) * 2;
+	}
+
 	RealType DiagonalLength() const
 	{
 		return (RealType)TMathUtil<RealType>::Sqrt((Max.X - Min.X) * (Max.X - Min.X) + (Max.Y - Min.Y) * (Max.Y - Min.Y));
diff --git a/Engine/Source/Runtime/GeometryCore/Public/CompGeom/ConvexHull2.h b/Engine/Source/Runtime/GeometryCore/Public/CompGeom/ConvexHull2.h
--- a/Engine/Source/Runtime/GeometryCore/Public/CompGeom/ConvexHull2.h
+++ b/Engine/Source/Runtime/GeometryCore/Public/CompGeom/ConvexHull2.h
@@ -24,6 +24,18 @@ class GEOMETRYCORE_API TConvexHull2
 {
 public:
 
+	/**
+	 * Generate convex hull for a simple polygon
+	 * If input has fewer than 3 points, will return false
+	 * If input is not a simple polygon (e.g. if it has self-intersections), may not find a correct hull
+	 * 
+	 * @param NumPolygonPoints	Number of points in the polygon
+	 * @param GetPointFunc		Function providing array-style access into the polygon points
+	 * @param bIsKnownCCW		If true, will save some processing time by assuming input is wound counter-clockwise
+	 * @return true if a hull was generated
+	 */
+	bool SolveSimplePolygon(int32 NumPolygonPoints, TFunctionRef<TVector2<RealType>(int32)> GetPointFunc, bool bIsKnownCCW = false);
+
 	/**
 	 * Generate convex hull as long as input is not degenerate
 	 * If input is degenerate, this will return false, and caller can call GetDimension()
@@ -106,6 +118,17 @@ public:
 		return Hull;
 	}
 
+	/** @return the line spanned by a collinear input.  Only meaningful if GetDimension() == 1 */
+	TLine2<RealType> GetLine(TArrayView<const TVector2<RealType>> Points) const
+	{
+		ensure(Dimension == 1);
+		if (ensure(Hull.Num() > 1))
+		{
+			return TLine2<RealType>::FromPoints(Points[Hull[0]], Points[Hull[1]]);
+		}
+		return TLine2<RealType>(); // not enough points to compute a valid line, return default
+	}
+
 protected:
 
 	// divide-and-conquer algorithm
diff --git a/Engine/Source/Runtime/GeometryCore/Public/CompGeom/FitOrientedBox2.h b/Engine/Source/Runtime/GeometryCore/Public/CompGeom/FitOrientedBox2.h
new file mode 100644
--- /dev/null
+++ b/Engine/Source/Runtime/GeometryCore/Public/CompGeom/FitOrientedBox2.h
@@ -0,0 +1,58 @@
+// Copyright Epic Games, Inc. All Rights Reserved.
+
+#pragma once
+
+#include "CoreMinimal.h"
+#include "Math/Vector.h"
+#include "VectorTypes.h"
+#include "OrientedBoxTypes.h"
+
+namespace UE
+{
+namespace Geometry
+{
+
+	enum EBox2FitCriteria
+	{
+		Area,
+		Perimeter
+	};
+
+	/**
+	* Use the rotating calipers algorithm to find a best-fit oriented bounding box of a point set.
+	* Note it internally computes a 2D convex hull of a point set, and is currently O(n log n) in the number of points.
+	* 
+	* @param Points		The points to fit
+	* @param FitMethod	What criteria to optimize
+	* @return			A best-fit TOrientedBox2 that contains the points
+	*/
+	template <typename RealType>
+	TOrientedBox2<RealType>  GEOMETRYCORE_API FitOrientedBox2Points(TArrayView<const TVector2<RealType>> Points, EBox2FitCriteria FitMethod = EBox2FitCriteria::Area);
+
+	/**
+	* Use the rotating calipers algorithm to find a best-fit oriented bounding box of a simple polygon.
+	* Note the polygon should not be self-intersecting.  O(n) in the number of points.
+	*
+	* @param Polygon	The vertices of the simple polygon to fit
+	* @param FitMethod	What criteria to optimize
+	* @return			A best-fit TOrientedBox2 that contains the points
+	*/
+	template <typename RealType>
+	TOrientedBox2<RealType>  GEOMETRYCORE_API FitOrientedBox2SimplePolygon(TArrayView<const TVector2<RealType>> Polygon, EBox2FitCriteria FitMethod = EBox2FitCriteria::Area);
+
+	/**
+	* Use the rotating calipers algorithm to find a best-fit oriented bounding box of a convex hull.
+	*
+	* @param NumPts		Number of points in the convex hull
+	* @param GetHullPt	Function(hull point index) -> hull point
+	* @param FitMethod	What criteria to optimize
+	* @return			A best-fit TOrientedBox2 that contains the points
+	*/
+	template <typename RealType>
+	TOrientedBox2<RealType>  GEOMETRYCORE_API FitOrientedBox2ConvexHull(int32 NumPts, TFunctionRef<TVector2<RealType>(int32)> GetHullPt, EBox2FitCriteria FitMethod = EBox2FitCriteria::Area);
+
+	// TODO: also have versions for a convex polygon input (+ a near-convex polygon version, that tolerates/fixes 'small' defects?)
+
+} // end namespace Geometry
+}// end namespace UE
+
diff --git a/Engine/Source/Runtime/GeometryCore/Public/OrientedBoxTypes.h b/Engine/Source/Runtime/GeometryCore/Public/OrientedBoxTypes.h
--- a/Engine/Source/Runtime/GeometryCore/Public/OrientedBoxTypes.h
+++ b/Engine/Source/Runtime/GeometryCore/Public/OrientedBoxTypes.h
@@ -79,18 +79,6 @@ struct TOrientedBox3
 	/** @return an axis of the box */
 	TVector<RealType> GetAxis(int AxisIndex) const { return Frame.GetAxis(AxisIndex); }
 
-	/** @return maximum extent of box */
-	inline RealType MaxExtent() const
-	{
-		return Extents.MaxAbs();
-	}
-
-	/** @return minimum extent of box */
-	inline RealType MinExtent() const
-	{
-		return Extents.MinAbs();
-	}
-
 	/** @return vector from minimum-corner to maximum-corner of box */
 	inline TVector<RealType> Diagonal() const
 	{
@@ -100,13 +88,14 @@ struct TOrientedBox3
 	/** @return volume of box */
 	inline RealType Volume() const
 	{
-		return (RealType)8 * (Extents.X) * (Extents.Y) * (Extents.Z);
+		return (RealType)8 * FMath::Max(0, Extents.X) * FMath::Max(0, Extents.Y) * FMath::Max(0, Extents.Z);
 	}
 
 	/** @return surface area of box */
 	inline RealType SurfaceArea() const
 	{
-		return (RealType)8 * ( Extents.X * Extents.Y + Extents.X * Extents.Z + Extents.Y * Extents.Z);
+		TVector<RealType> ClampExtents(FMath::Max(0, Extents.X), FMath::Max(0, Extents.Y), FMath::Max(0, Extents.Z));
+		return (RealType)8 * (ClampExtents.X * ClampExtents.Y + ClampExtents.X * ClampExtents.Z + ClampExtents.Y * ClampExtents.Z);
 	}
 
 	/** @return true if box contains point */
@@ -313,18 +302,6 @@ struct TOrientedBox2
 	/** @return an axis of the box */
 	inline TVector2<RealType> GetAxis(int AxisIndex) const { return AxisIndex == 0 ? AxisX() : AxisY(); }
 
-	/** @return maximum extent of box */
-	inline RealType MaxExtent() const
-	{
-		return Extents.MaxAbs();
-	}
-
-	/** @return minimum extent of box */
-	inline RealType MinExtent() const
-	{
-		return Extents.MinAbs();
-	}
-
 	/** @return vector from minimum-corner to maximum-corner of box */
 	inline TVector2<RealType> Diagonal() const
 	{
@@ -334,13 +311,13 @@ struct TOrientedBox2
 	/** @return area of box */
 	inline RealType Area() const
 	{
-		return (RealType)4 * (Extents.X) * (Extents.Y);
+		return (RealType)4 * FMath::Max(0, Extents.X) * FMath::Max(0, Extents.Y);
 	}
 
 	/** @return perimeter of box */
 	inline RealType Perimeter() const
 	{
-		return (RealType)4 * (Extents.X + Extents.Y);
+		return (RealType)4 * (FMath::Max(0, Extents.X) + FMath::Max(0, Extents.Y));
 	}
 
 	/** @return Point transformed to the local space of the box (Origin at box center, Axes aligned to box) */
@@ -357,12 +334,11 @@ struct TOrientedBox2
 	}
 
 	/** @return true if box contains point */
-	inline bool Contains(const TVector<RealType>& Point) const
+	inline bool Contains(const TVector2<RealType>& Point) const
 	{
-		TVector<RealType> LocalPoint = ToLocalPoint(Point);
+		TVector2<RealType> LocalPoint = ToLocalSpace(Point);
 		return (TMathUtil<RealType>::Abs(LocalPoint.X) <= Extents.X) &&
-			(TMathUtil<RealType>::Abs(LocalPoint.Y) <= Extents.Y) &&
-			(TMathUtil<RealType>::Abs(LocalPoint.Z) <= Extents.Z);
+			(TMathUtil<RealType>::Abs(LocalPoint.Y) <= Extents.Y);
 	}
 
 
