CIS fix: type mismatch between local functions in FitOrientedBox2.cpp

diff --git a/Engine/Source/Runtime/GeometryCore/Private/CompGeom/FitOrientedBox2.cpp b/Engine/Source/Runtime/GeometryCore/Private/CompGeom/FitOrientedBox2.cpp
--- a/Engine/Source/Runtime/GeometryCore/Private/CompGeom/FitOrientedBox2.cpp
+++ b/Engine/Source/Runtime/GeometryCore/Private/CompGeom/FitOrientedBox2.cpp
@@ -120,15 +120,16 @@ TOrientedBox2<RealType> FitOrientedBox2ConvexHull(int32 NumPts, TFunctionRef<TVe
 	InitialExtreme.D = FindExtremePtInitial(-Base, InitialExtreme.C, ExtremesMin.X); // Extreme left point
 	TAxisAlignedBox2<RealType> Range(-ExtremesMin, ExtremesMax);
 
-	auto EvalBox = FitMethod == EBox2FitCriteria::Area ?
-		[](const TAxisAlignedBox2<RealType>& Range) -> RealType
-	{
-		return Range.Area();
-	} :
-		[](const TAxisAlignedBox2<RealType>& Range) -> RealType
-	{
-		return Range.Perimeter();
-	};
+	using FEvalFn = TFunctionRef<RealType(const TAxisAlignedBox2<RealType>&)>;
+	FEvalFn EvalBox = FitMethod == EBox2FitCriteria::Area ?
+		FEvalFn([](const TAxisAlignedBox2<RealType>& Range) -> RealType
+			{
+				return Range.Area();
+			}) :
+		FEvalFn([](const TAxisAlignedBox2<RealType>& Range) -> RealType
+			{
+				return Range.Perimeter();
+			});
 	RealType BestScore = EvalBox(Range);
 	TVector2<RealType> BestBase = Base;
 	TAxisAlignedBox2<RealType> BestRange = Range;
