Geometry Core, Geometry Processing.  Fix implicit cast warnings encountered during clang build.

GeometryCore: StairGenerator.cpp, SweepGenerator.{h,.cpp}, UVPacking.cpp, ParallelTransformReduce.h, CurvewUtil.h, SweepingMeshSDF.h FastWinding.h

GeometryProcessing:  Remesher.cpp, AdaptiveTessellate.cpp, UniformTessellate.cpp, GeometrySelectionUtil.cpp, MeshSimpleApproximation.cpp

diff --git a/Engine/Plugins/Runtime/GeometryProcessing/Source/DynamicMesh/Private/Operations/AdaptiveTessellate.cpp b/Engine/Plugins/Runtime/GeometryProcessing/Source/DynamicMesh/Private/Operations/AdaptiveTessellate.cpp
--- a/Engine/Plugins/Runtime/GeometryProcessing/Source/DynamicMesh/Private/Operations/AdaptiveTessellate.cpp
+++ b/Engine/Plugins/Runtime/GeometryProcessing/Source/DynamicMesh/Private/Operations/AdaptiveTessellate.cpp
@@ -1171,7 +1171,7 @@ namespace AdaptiveTessellateLocals
 	{	
 		auto LerpElements = [](const RealType* Element1, const RealType* Element2, RealType* OutElement, RealType Alpha)
 		{
-			Alpha = FMath::Clamp(Alpha, 0.0, 1.0);
+			Alpha = FMath::Clamp(Alpha, RealType(0), RealType(1));
 			RealType OneMinusAlpha = (RealType)1 - Alpha;
 
 			for (int Idx = 0; Idx < ElementSize; ++Idx)
@@ -1228,7 +1228,7 @@ namespace AdaptiveTessellateLocals
 
 				for (int Index = 0; Index < IDBlock.Num(); ++Index)
 				{	
-					const double Alpha = CoordBlock[Index];
+					const RealType Alpha = (RealType)CoordBlock[Index];
 					const int ElementID = IDBlock[Index]; 
 					
 					LerpElements(Element1, Element2, Out, Alpha);
@@ -1247,7 +1247,7 @@ namespace AdaptiveTessellateLocals
 
 				for (int Index = 0; Index < IDBlock.Num(); ++Index)
 				{	
-					const double Alpha = CoordBlock[Index];
+					const RealType Alpha = (RealType)CoordBlock[Index];
 					const int ElementID = IDBlock[Index]; 
 					
 					LerpElements(Element1, Element2, Out, Alpha);
@@ -1275,7 +1275,7 @@ namespace AdaptiveTessellateLocals
 				const FVector3d Bary = CoordBlock[Index]; 
 				const int ElementID = IDBlock[Index];
 
-				BaryElements(Element1, Element2, Element3, Out, Bary[0], Bary[1], Bary[2]);
+				BaryElements(Element1, Element2, Element3, Out, RealType(Bary[0]), RealType(Bary[1]), RealType(Bary[2]));
 
 				MapE.Add(ElementID, ResultOverlay->AppendElement(Out));
 			}
diff --git a/Engine/Plugins/Runtime/GeometryProcessing/Source/DynamicMesh/Private/Operations/UniformTessellate.cpp b/Engine/Plugins/Runtime/GeometryProcessing/Source/DynamicMesh/Private/Operations/UniformTessellate.cpp
--- a/Engine/Plugins/Runtime/GeometryProcessing/Source/DynamicMesh/Private/Operations/UniformTessellate.cpp
+++ b/Engine/Plugins/Runtime/GeometryProcessing/Source/DynamicMesh/Private/Operations/UniformTessellate.cpp
@@ -317,7 +317,7 @@ namespace UniformTessellateLocals
 				RealType OutElement[ElementSize];
 				for (int VertexOffset = 0; VertexOffset < TessellationNum; ++VertexOffset)
 				{
-					const RealType Alpha = (RealType)(VertexOffset + 1) / (TessellationNum + 1);
+					const RealType Alpha = (RealType)(VertexOffset + 1) / RealType(TessellationNum + 1);
 					LerpElements(Element1, Element2, OutElement, Alpha);
 					SetElementOnEdge(EdgeID, EdgeTri.A, VertexOffset, OutElement); 
 				}
@@ -393,7 +393,7 @@ namespace UniformTessellateLocals
 					RealType OutElement[ElementSize];
 					for (int VertexOffset = 0; VertexOffset < NumNewLevelVertices; ++VertexOffset) 
 					{
-						const RealType Alpha = RealType(VertexOffset + 1) / (NumNewLevelVertices + 1);
+						const RealType Alpha = RealType(VertexOffset + 1) / RealType(NumNewLevelVertices + 1);
 						LerpElements(Element1, Element2, OutElement, Alpha);
 						this->SetElementOnTriangle(TriangleID, ElementIDCounter, OutElement);
 						ElementIDCounter++;
@@ -882,7 +882,7 @@ namespace UniformTessellateLocals
 
 					for (int VertexOffset = 0; VertexOffset < TessellationNum; ++VertexOffset)
 					{
-						const RealType Tau = (RealType)(VertexOffset + 1) / (TessellationNum + 1);
+						const RealType Tau = (RealType)(VertexOffset + 1) / RealType(TessellationNum + 1);
 						this->LerpElements(Element1, Element2, Out, Tau);
 						this->SetElementOnEdge(EdgeID, EdgeTri.A, VertexOffset, Out);	
 					}
@@ -895,7 +895,7 @@ namespace UniformTessellateLocals
 					
 					for (int VertexOffset = 0; VertexOffset < TessellationNum; ++VertexOffset)
 					{
-						const RealType Tau = (RealType)(VertexOffset + 1) / (TessellationNum + 1);
+						const RealType Tau = (RealType)(VertexOffset + 1) / RealType(TessellationNum + 1);
 						this->LerpElements(Element1, Element2, Out, Tau);
 						this->SetElementOnEdge(EdgeID, EdgeTri.B, VertexOffset, Out);
 					}
diff --git a/Engine/Plugins/Runtime/GeometryProcessing/Source/DynamicMesh/Private/Remesher.cpp b/Engine/Plugins/Runtime/GeometryProcessing/Source/DynamicMesh/Private/Remesher.cpp
--- a/Engine/Plugins/Runtime/GeometryProcessing/Source/DynamicMesh/Private/Remesher.cpp
+++ b/Engine/Plugins/Runtime/GeometryProcessing/Source/DynamicMesh/Private/Remesher.cpp
@@ -577,7 +577,7 @@ void FRemesher::ApplyVertexBuffer(bool bParallel)
 	{
 		// TODO: verify that this batching is still necessary, now that timestamps/locking situation has been improved
 		const int BatchSize = 1000;
-		const int NumBatches = FMath::CeilToInt(Mesh->MaxVertexID() / static_cast<float>(BatchSize));
+		const int NumBatches = FMath::CeilToInt32(static_cast<float>(Mesh->MaxVertexID()) / static_cast<float>(BatchSize));
 		ParallelFor(NumBatches, [this, BatchSize](int32 BatchID)
 		{
 			for (int VertexID = BatchID * BatchSize; VertexID < (BatchID + 1) * BatchSize; ++VertexID)
diff --git a/Engine/Plugins/Runtime/GeometryProcessing/Source/DynamicMesh/Private/Selections/GeometrySelectionUtil.cpp b/Engine/Plugins/Runtime/GeometryProcessing/Source/DynamicMesh/Private/Selections/GeometrySelectionUtil.cpp
--- a/Engine/Plugins/Runtime/GeometryProcessing/Source/DynamicMesh/Private/Selections/GeometrySelectionUtil.cpp
+++ b/Engine/Plugins/Runtime/GeometryProcessing/Source/DynamicMesh/Private/Selections/GeometrySelectionUtil.cpp
@@ -223,7 +223,7 @@ bool UE::Geometry::EnumerateTriangleSelectionVertices(
 	{
 		for (uint64 TriangleID : MeshSelection.Selection)
 		{
-			if (Mesh.IsTriangle(TriangleID))
+			if (Mesh.IsTriangle((int32)TriangleID))
 			{
 				FIndex3i Triangle = Mesh.GetTriangle((int32)TriangleID);
 				VertexFunc((uint64)Triangle.A, ApplyTransform.TransformPosition(Mesh.GetVertex(Triangle.A)));
@@ -252,7 +252,7 @@ bool UE::Geometry::EnumerateTriangleSelectionVertices(
 		{
 			if (Mesh.IsVertex((int32)VertexID))
 			{
-				VertexFunc((uint64)VertexID, ApplyTransform.TransformPosition(Mesh.GetVertex(VertexID)));
+				VertexFunc((uint64)VertexID, ApplyTransform.TransformPosition(Mesh.GetVertex((int32)VertexID)));
 			}
 		}
 	}
@@ -369,9 +369,9 @@ bool UE::Geometry::EnumerateTriangleSelectionTriangles(
 	{
 		for (uint64 TriangleID : MeshSelection.Selection)
 		{
-			if (Mesh.IsTriangle(TriangleID))
+			if (Mesh.IsTriangle((int32)TriangleID))
 			{
-				TriangleFunc(TriangleID);
+				TriangleFunc((int32)TriangleID);
 			}
 		}
 	}
@@ -391,7 +391,7 @@ bool UE::Geometry::EnumerateTriangleSelectionTriangles(
 	{
 		for (uint64 VertexID : MeshSelection.Selection)
 		{
-			Mesh.EnumerateVertexTriangles(VertexID, [&](int32 TriangleID)
+			Mesh.EnumerateVertexTriangles((int32)VertexID, [&](int32 TriangleID)
 			{
 				TriangleFunc(TriangleID);
 			});
diff --git a/Engine/Plugins/Runtime/GeometryProcessing/Source/DynamicMesh/Private/ShapeApproximation/MeshSimpleShapeApproximation.cpp b/Engine/Plugins/Runtime/GeometryProcessing/Source/DynamicMesh/Private/ShapeApproximation/MeshSimpleShapeApproximation.cpp
--- a/Engine/Plugins/Runtime/GeometryProcessing/Source/DynamicMesh/Private/ShapeApproximation/MeshSimpleShapeApproximation.cpp
+++ b/Engine/Plugins/Runtime/GeometryProcessing/Source/DynamicMesh/Private/ShapeApproximation/MeshSimpleShapeApproximation.cpp
@@ -325,7 +325,7 @@ void FMeshSimpleShapeApproximation::Generate_ConvexHullDecompositions(FSimpleSha
 		const FDynamicMesh3& SourceMesh = *SourceMeshes[idx];
 		// TODO: if (bSimplifyHulls), also consider simplifying the input?
 		FConvexDecomposition3 Decomposition(SourceMesh);
-		const int32 NumAdditionalSplits = FMath::FloorToInt32(ConvexDecompositionMaxPieces * ConvexDecompositionSearchFactor);
+		const int32 NumAdditionalSplits = FMath::FloorToInt32(float(ConvexDecompositionMaxPieces) * ConvexDecompositionSearchFactor);
 		Decomposition.Compute(ConvexDecompositionMaxPieces, NumAdditionalSplits, ConvexDecompositionErrorTolerance, ConvexDecompositionMinPartThickness);
 
 		for (int32 HullIdx = 0; HullIdx < Decomposition.NumHulls(); HullIdx++)
diff --git a/Engine/Source/Runtime/GeometryCore/Private/Generators/StairGenerator.cpp b/Engine/Source/Runtime/GeometryCore/Private/Generators/StairGenerator.cpp
--- a/Engine/Source/Runtime/GeometryCore/Private/Generators/StairGenerator.cpp
+++ b/Engine/Source/Runtime/GeometryCore/Private/Generators/StairGenerator.cpp
@@ -866,9 +866,9 @@ FVector3d FLinearStairGenerator::GenerateVertex(ESide Side, int VertexColumn, in
 		// Vertices are only generated for Right & Left sides.
 		check(false);
 	}
-	float X = VertexColumn * StepDepth;
+	float X = float(VertexColumn) * StepDepth;
 	float Y = (Side == ESide::Right ? float(0.5 * StepWidth) : -float(0.5 * StepWidth));
-	float Z = VertexRow * StepHeight;
+	float Z = float(VertexRow) * StepHeight;
 	return FVector3d(X, Y, Z);
 }
 
@@ -912,8 +912,8 @@ FVector2f FLinearStairGenerator::GenerateUV(ESide Side, int Step, int VertexId,
 	case ESide::Right:
 	case ESide::Left:
 	{
-		const float UScale = NumSteps * StepDepth * UVScale;
-		const float VScale = NumSteps * StepHeight * UVScale;
+		const float UScale = float(NumSteps) * StepDepth * UVScale;
+		const float VScale = float(NumSteps) * StepHeight * UVScale;
 		UV.X = FMathf::Lerp(-0.5f, 0.5f, float(Col % LeftSideColumnId) / float(NumSteps + 1)) * UScale + 0.5f;
 		UV.Y = FMathf::Lerp(-0.5f, 0.5f, float(Row) / float(NumSteps + 1)) * VScale + 0.5f;
 		break;
@@ -941,7 +941,7 @@ FVector2f FLinearStairGenerator::GenerateUV(ESide Side, int Step, int VertexId,
 	case ESide::Back:
 	{
 		const float UScale = StepWidth * UVScale;
-		const float VScale = NumSteps * StepHeight * UVScale;
+		const float VScale = float(NumSteps) * StepHeight * UVScale;
 		UV.X = Col < LeftSideColumnId ? 0.5f : -0.5f;
 		UV.Y = FMathf::Lerp(-0.5f, 0.5f, float(Row) / float(NumSteps + 1));
 		UV.X = UV.X * UScale + 0.5f;
@@ -951,7 +951,7 @@ FVector2f FLinearStairGenerator::GenerateUV(ESide Side, int Step, int VertexId,
 	case ESide::Bottom:
 	{
 		const float UScale = StepWidth * UVScale;
-		const float VScale = NumSteps * StepDepth * UVScale;
+		const float VScale = float(NumSteps) * StepDepth * UVScale;
 		UV.X = Col < LeftSideColumnId ? 0.5f : -0.5f;
 		UV.Y = FMathf::Lerp(-0.5f, 0.5f, float(Col % LeftSideColumnId) / float(NumSteps + 1));
 		UV.X = UV.X * UScale + 0.5f;
@@ -964,7 +964,7 @@ FVector2f FLinearStairGenerator::GenerateUV(ESide Side, int Step, int VertexId,
 
 float FLinearStairGenerator::GetMaxDimension()
 {
-	return FMathf::Max(FMathf::Max(NumSteps * StepDepth, NumSteps * StepHeight), StepWidth);
+	return FMathf::Max(FMathf::Max(float(NumSteps) * StepDepth, float(NumSteps) * StepHeight), StepWidth);
 }
 
 
@@ -979,9 +979,9 @@ FVector3d FFloatingStairGenerator::GenerateVertex(ESide Side, int VertexColumn,
 		// Vertices are only generated for Right & Left sides.
 		check(false);
 	}
-	float X = VertexColumn * StepDepth;
+	float X = float(VertexColumn) * StepDepth;
 	float Y = (Side == ESide::Right ? float(0.5 * StepWidth) : -float(0.5 * StepWidth));
-	float Z = VertexColumn > 1 ? ((VertexColumn - 2) + VertexRow) * StepHeight : VertexRow * StepHeight;
+	float Z = VertexColumn > 1 ? float((VertexColumn - 2) + VertexRow) * StepHeight : float(VertexRow) * StepHeight;
 	return FVector3d(X, Y, Z);
 }
 
@@ -995,8 +995,8 @@ FVector2f FFloatingStairGenerator::GenerateUV(ESide Side, int Step, int VertexId
 	case ESide::Right:
 	case ESide::Left:
 	{
-		const float UScale = NumSteps * StepDepth * UVScale;
-		const float VScale = NumSteps * StepHeight * UVScale;
+		const float UScale = float(NumSteps) * StepDepth * UVScale;
+		const float VScale = float(NumSteps) * StepHeight * UVScale;
 		float X = float(Col % LeftSideColumnId);
 		float Y = X < 3.0f ? float(Row) : float(Row) + X - 2.0f;
 		X /= float(NumSteps + 1);
@@ -1008,7 +1008,7 @@ FVector2f FFloatingStairGenerator::GenerateUV(ESide Side, int Step, int VertexId
 	case ESide::Back:
 	{
 		const float UScale = StepWidth * UVScale;
-		const float VScale = NumSteps * StepHeight * UVScale;
+		const float VScale = float(NumSteps) * StepHeight * UVScale;
 		float Y = Step < NumSteps - 1 ? float(Row + Step) : float(Row + Step - 1);
 		Y /= float(NumSteps + 1);
 		UV.X = Col < LeftSideColumnId ? 0.5f : -0.5f;
@@ -1037,7 +1037,7 @@ void FCurvedStairGenerator::ResetData()
 
 	bIsClockwise = CurveAngle > 0.0f;
 	CurveRadians = CurveAngle * TMathUtilConstants<float>::DegToRad;
-	CurveRadiansPerStep = CurveRadians / NumSteps;
+	CurveRadiansPerStep = CurveRadians / float(NumSteps);
 	OuterRadius = InnerRadius + StepWidth;
 	RadiusRatio = OuterRadius / InnerRadius;
 	BackNormal = FVector3f::Zero();
@@ -1052,10 +1052,10 @@ FVector3d FCurvedStairGenerator::GenerateVertex(ESide Side, int VertexColumn, in
 	}
 	float X = 0.0f;
 	float Y = 0.0f;
-	float Z = VertexRow * StepHeight;
+	float Z = float(VertexRow) * StepHeight;
 
-	float XCoeff = FMathf::Cos(VertexColumn * CurveRadiansPerStep);
-	float YCoeff = FMathf::Sin(VertexColumn * CurveRadiansPerStep);
+	float XCoeff = FMathf::Cos(float(VertexColumn) * CurveRadiansPerStep);
+	float YCoeff = FMathf::Sin(float(VertexColumn) * CurveRadiansPerStep);
 	if (bIsClockwise)
 	{
 		X = (Side == ESide::Right ? XCoeff * InnerRadius : XCoeff * OuterRadius);
@@ -1080,16 +1080,16 @@ FVector3f FCurvedStairGenerator::GenerateNormal(ESide Side, int VertexId)
 	case ESide::Right:
 	case ESide::Left:
 	{
-		float X = FMathf::Cos((Col % (NumSteps + 1)) * CurveRadiansPerStep);
-		float Y = FMathf::Sin((Col % (NumSteps + 1)) * CurveRadiansPerStep);
+		float X = FMathf::Cos(float(Col % (NumSteps + 1)) * CurveRadiansPerStep);
+		float Y = FMathf::Sin(float(Col % (NumSteps + 1)) * CurveRadiansPerStep);
 		N = (Side == ESide::Right ? FVector3f(-X, -Y, 0.0f) : FVector3f(X, Y, 0.0f));
 		Normalize(N);
 		break;
 	}
 	case ESide::Front:
 	{
-		float X = FMathf::Cos((Col % (NumSteps + 1)) * CurveRadiansPerStep);
-		float Y = FMathf::Sin((Col % (NumSteps + 1)) * CurveRadiansPerStep);
+		float X = FMathf::Cos(float(Col % (NumSteps + 1)) * CurveRadiansPerStep);
+		float Y = FMathf::Sin(float(Col % (NumSteps + 1)) * CurveRadiansPerStep);
 		N = FVector3f(Y, -X, 0.0f);
 		Normalize(N);
 		break;
@@ -1103,8 +1103,8 @@ FVector3f FCurvedStairGenerator::GenerateNormal(ESide Side, int VertexId)
 	{
 		if (BackNormal == FVector3f::Zero())
 		{
-			float X = FMathf::Cos(NumSteps * CurveRadiansPerStep);
-			float Y = FMathf::Sin(NumSteps * CurveRadiansPerStep);
+			float X = FMathf::Cos(float(NumSteps) * CurveRadiansPerStep);
+			float Y = FMathf::Sin(float(NumSteps) * CurveRadiansPerStep);
 			BackNormal = FVector3f(-Y, X, 0.0f);
 		}
 		N = BackNormal;
@@ -1130,7 +1130,7 @@ FVector2f FCurvedStairGenerator::GenerateUV(ESide Side, int Step, int VertexId,
 	case ESide::Left:
 	{
 		const float UScale = OuterRadius * UVScale;
-		const float VScale = NumSteps * StepHeight * UVScale;
+		const float VScale = float(NumSteps) * StepHeight * UVScale;
 		UV.X = FMathf::Lerp(-0.5f, 0.5f, float(Col % LeftSideColumnId) / float(NumSteps + 1));
 		UV.Y = FMathf::Lerp(-0.5f, 0.5f, float(Row) / float(NumSteps + 1));
 
@@ -1161,7 +1161,7 @@ FVector2f FCurvedStairGenerator::GenerateUV(ESide Side, int Step, int VertexId,
 	case ESide::Top:
 	{
 		const float UScale = StepWidth * UVScale;
-		const float VScale = OuterRadius / NumSteps * UVScale;
+		const float VScale = OuterRadius / float(NumSteps) * UVScale;
 		UV.X = Col < LeftSideColumnId ? 0.5f : -0.5f;
 		UV.Y = (Col % LeftSideColumnId) > Step ? -0.5f : 0.5f;
 
@@ -1183,7 +1183,7 @@ FVector2f FCurvedStairGenerator::GenerateUV(ESide Side, int Step, int VertexId,
 	case ESide::Back:
 	{
 		const float UScale = StepWidth * UVScale;
-		const float VScale = NumSteps * StepHeight * UVScale;
+		const float VScale = float(NumSteps) * StepHeight * UVScale;
 		UV.X = Col < LeftSideColumnId ? 0.5f : -0.5f;
 		UV.Y = FMathf::Lerp(-0.5f, 0.5f, float(Row) / float(NumSteps + 1));
 		UV.X = UV.X * UScale + 0.5f;
@@ -1223,7 +1223,7 @@ FVector2f FCurvedStairGenerator::GenerateUV(ESide Side, int Step, int VertexId,
 float FCurvedStairGenerator::GetMaxDimension()
 {
 	float MaxDepth = FMathf::Abs(CurveRadians) * OuterRadius;
-	return FMathf::Max(FMathf::Max(MaxDepth, NumSteps * StepHeight), StepWidth);
+	return FMathf::Max(FMathf::Max(MaxDepth, float(NumSteps) * StepHeight), StepWidth);
 }
 
 
@@ -1240,10 +1240,10 @@ FVector3d FSpiralStairGenerator::GenerateVertex(ESide Side, int VertexColumn, in
 	}
 	float X = 0.0f;
 	float Y = 0.0f;
-	float Z = VertexColumn > 1 ? ((VertexColumn - 2) + VertexRow) * StepHeight : VertexRow * StepHeight;
+	float Z = VertexColumn > 1 ? float((VertexColumn - 2) + VertexRow) * StepHeight : float(VertexRow) * StepHeight;
 
-	float XCoeff = FMathf::Cos(VertexColumn * CurveRadiansPerStep);
-	float YCoeff = FMathf::Sin(VertexColumn * CurveRadiansPerStep);
+	float XCoeff = FMathf::Cos(float(VertexColumn) * CurveRadiansPerStep);
+	float YCoeff = FMathf::Sin(float(VertexColumn) * CurveRadiansPerStep);
 	if (bIsClockwise)
 	{
 		X = (Side == ESide::Right ? XCoeff * InnerRadius : XCoeff * OuterRadius);
@@ -1268,7 +1268,7 @@ FVector2f FSpiralStairGenerator::GenerateUV(ESide Side, int Step, int VertexId,
 	case ESide::Left:
 	{
 		const float UScale = OuterRadius * UVScale;
-		const float VScale = NumSteps * StepHeight * UVScale;
+		const float VScale = float(NumSteps) * StepHeight * UVScale;
 		float X = float(Col % LeftSideColumnId);
 		float Y = X < 3.0f ? float(Row) : float(Row) + X - 2.0f;
 		UV.X = FMathf::Lerp(-0.5f, 0.5f, X / float(NumSteps + 1));
@@ -1291,7 +1291,7 @@ FVector2f FSpiralStairGenerator::GenerateUV(ESide Side, int Step, int VertexId,
 	case ESide::Back:
 	{
 		const float UScale = StepWidth * UVScale;
-		const float VScale = NumSteps * StepHeight * UVScale;
+		const float VScale = float(NumSteps) * StepHeight * UVScale;
 		float Y = Step < NumSteps - 1 ? float(Row + Step) : float(Row + Step - 1);
 		Y /= float(NumSteps + 1);
 		UV.X = Col < LeftSideColumnId ? 0.5f : -0.5f;
diff --git a/Engine/Source/Runtime/GeometryCore/Private/Generators/SweepGenerator.cpp b/Engine/Source/Runtime/GeometryCore/Private/Generators/SweepGenerator.cpp
--- a/Engine/Source/Runtime/GeometryCore/Private/Generators/SweepGenerator.cpp
+++ b/Engine/Source/Runtime/GeometryCore/Private/Generators/SweepGenerator.cpp
@@ -525,7 +525,7 @@ void FProfileSweepGenerator::InitializeUvBuffer(const TArray<int32>& VertPositio
 	{
 		UScale = (UnitUVInWorldCoordinates != 0) ? 
 			// Get an average and convert to UV from world using the scale.
-			UScale / (NumProfilePoints * UnitUVInWorldCoordinates) 
+			UScale / (float(NumProfilePoints) * UnitUVInWorldCoordinates) 
 			: 0;
 	}
 	// Adjust
diff --git a/Engine/Source/Runtime/GeometryCore/Private/Parameterization/UVPacking.cpp b/Engine/Source/Runtime/GeometryCore/Private/Parameterization/UVPacking.cpp
--- a/Engine/Source/Runtime/GeometryCore/Private/Parameterization/UVPacking.cpp
+++ b/Engine/Source/Runtime/GeometryCore/Private/Parameterization/UVPacking.cpp
@@ -281,7 +281,7 @@ void FStandardChartPacker::ScaleCharts( TArray<FUVIsland>& Charts, double UVScal
 			FVector2d ChartSize = Chart.MaxUV - Chart.MinUV;
 			FVector2d ChartSizeScaled = ChartSize * Chart.UVScale * UniformScale;
 
-			const double MaxChartEdge = TextureResolution - 1.0f;
+			const double MaxChartEdge = TextureResolution - 1.0;
 			const double LongestChartEdge = FMathd::Max(ChartSizeScaled.X, ChartSizeScaled.Y);
 
 			const double Epsilon = 0.01f;
diff --git a/Engine/Source/Runtime/GeometryCore/Public/Async/ParallelTransformReduce.h b/Engine/Source/Runtime/GeometryCore/Public/Async/ParallelTransformReduce.h
--- a/Engine/Source/Runtime/GeometryCore/Public/Async/ParallelTransformReduce.h
+++ b/Engine/Source/Runtime/GeometryCore/Public/Async/ParallelTransformReduce.h
@@ -24,7 +24,7 @@ namespace Geometry {
 		// ParallelFor doesn't yet support int64 NumTasks, so cap it
 		const int32 NumTasks = (int32) FMath::Min(InNumTasks, MAX_int32);
 
-		const IntType NumPerTask = (IntType)FMathf::Ceil((float)Num / NumTasks);
+		const IntType NumPerTask = (IntType)FMathf::Ceil((float)Num / (float)NumTasks);
 
 		TArray<T> PerTaskResults;
 		PerTaskResults.SetNum(NumTasks);
@@ -72,7 +72,7 @@ namespace Geometry {
 		// ParallelFor doesn't yet support int64 NumTasks, so cap it
 		const int32 NumTasks = (int32)FMath::Min(InNumTasks, MAX_int32);
 
-		IntType NumPerTask = (IntType)FMathf::Ceil((float)Num / NumTasks);
+		IntType NumPerTask = (IntType)FMathf::Ceil((float)Num / (float)NumTasks);
 
 		TArray<T> PerTaskResults;
 		PerTaskResults.SetNum(NumTasks);
@@ -161,7 +161,7 @@ namespace Geometry {
 		// ParallelFor doesn't yet support int64 NumTasks, so cap it
 		const int32 NumTasks = (int32)FMath::Min(InNumTasks, MAX_int32);
 
-		int64 NumPerTask = (int64)FMathf::Ceil((float)Num / NumTasks);
+		int64 NumPerTask = (int64)FMathf::Ceil((float)Num / (float)NumTasks);
 
 		TArray<T> PerTaskResults;
 		PerTaskResults.SetNum(NumTasks);
@@ -213,7 +213,7 @@ namespace Geometry {
 		// ParallelFor doesn't yet support int64 NumTasks, so cap it
 		const int32 NumTasks = (int32)FMath::Min(InNumTasks, MAX_int32);
 
-		int64 NumPerTask = (int64)FMathf::Ceil((float)Num / NumTasks);
+		int64 NumPerTask = (int64)FMathf::Ceil((float)Num / (float)NumTasks);
 
 		TArray<T> PerTaskResults;
 		PerTaskResults.SetNum(NumTasks);
diff --git a/Engine/Source/Runtime/GeometryCore/Public/Curve/CurveUtil.h b/Engine/Source/Runtime/GeometryCore/Public/Curve/CurveUtil.h
--- a/Engine/Source/Runtime/GeometryCore/Public/Curve/CurveUtil.h
+++ b/Engine/Source/Runtime/GeometryCore/Public/Curve/CurveUtil.h
@@ -257,7 +257,7 @@ using namespace UE::Math;
 					StartCur = 1;
 					StartPrev = 0;
 				}
-				RealType PrevDist = (Vertices[StartPrev][ClipDim] - ClipCoord) * SideSign;
+				RealType PrevDist = (Vertices[StartPrev][ClipDim] - ClipCoord) * RealType(SideSign);
 				if constexpr (!bLoop)
 				{
 					if (PrevDist >= 0)
@@ -267,7 +267,7 @@ using namespace UE::Math;
 				}
 				for (int32 CurIdx = StartCur, PrevIdx = StartPrev; CurIdx < VertNum; PrevIdx = CurIdx++)
 				{
-					RealType CurDist = (Vertices[CurIdx][ClipDim] - ClipCoord) * SideSign;
+					RealType CurDist = (Vertices[CurIdx][ClipDim] - ClipCoord) * RealType(SideSign);
 					if (CurDist >= 0)
 					{
 						if (PrevDist < 0 && CurDist > 0)
diff --git a/Engine/Source/Runtime/GeometryCore/Public/Generators/SweepGenerator.h b/Engine/Source/Runtime/GeometryCore/Public/Generators/SweepGenerator.h
--- a/Engine/Source/Runtime/GeometryCore/Public/Generators/SweepGenerator.h
+++ b/Engine/Source/Runtime/GeometryCore/Public/Generators/SweepGenerator.h
@@ -251,10 +251,10 @@ protected:
 			int32 NextDupVertIdx = UVSection < NumSections ? UVSections[UVSection] : -1;
 			for (int32 VertSubIdx = 0; VertSubIdx < XVerts; UVSubIdx++)
 			{
-				float UVX = bEvenlySpaceUVs ? float(CrossSectionPercentages[VertSubIdx]) : VertSubIdx / float(XVerts);
+				float UVX = bEvenlySpaceUVs ? float(CrossSectionPercentages[VertSubIdx]) : float(VertSubIdx) / float(XVerts);
 				for (int32 XIdx = 0; XIdx < NumCrossSections; XIdx++)
 				{
-					float UVY = bEvenlySpaceUVs ? float(PathPercentages[XIdx]) : XIdx / float(NumCrossSections - 1);
+					float UVY = bEvenlySpaceUVs ? float(PathPercentages[XIdx]) : float(XIdx) / float(NumCrossSections - 1);
 					SetUV(XIdx * XUVs + UVSubIdx, FVector2f(1-UVX, 1-UVY) * SectionsUVScale, (XIdx % CrossSectionsMod) * XVerts + VertSubIdx);
 				}
 
@@ -286,7 +286,7 @@ protected:
 				int32 VertSubIdx = 0;
 				for (int32 XIdx = 0; XIdx < NumCrossSections; XIdx++)
 				{
-					float UVY = bEvenlySpaceUVs ? float(PathPercentages[XIdx]) : XIdx / float(NumCrossSections - 1);
+					float UVY = bEvenlySpaceUVs ? float(PathPercentages[XIdx]) : float(XIdx) / float(NumCrossSections - 1);
 					SetUV(XIdx * XUVs + UVSubIdx, FVector2f(1-UVX, 1-UVY) * SectionsUVScale, (XIdx % CrossSectionsMod) * XVerts + VertSubIdx);
 				}
 			}
@@ -587,7 +587,7 @@ public:
 		float LenAlong = ComputeSegLengths(Radii, Heights, AlongPercents);
 		for (int Idx = 0; Idx < 3; Idx++)
 		{
-			AdditionalLengthSamples[Idx] = (int)(.5f+AlongPercents[Idx + 1] * TargetSamples);
+			AdditionalLengthSamples[Idx] = (int)(.5f+AlongPercents[Idx + 1] * float(TargetSamples));
 		}
 	}
 
@@ -622,7 +622,7 @@ public:
 
 			for (int ExtraSeg = 1, NumExtraSegs = AdditionalLengthSamples[SegIdx] + 1; ExtraSeg < NumExtraSegs; ++ExtraSeg)
 			{
-				const float Along = ExtraSeg / float(NumExtraSegs);
+				const float Along = float(ExtraSeg) / float(NumExtraSegs);
 				SetVert(FMath::Lerp(SrcRadii[SegIdx], SrcRadii[SegIdx + 1], Along),
 						FMath::Lerp(SrcHeights[SegIdx], SrcHeights[SegIdx + 1], Along));
 			}
diff --git a/Engine/Source/Runtime/GeometryCore/Public/Implicit/SweepingMeshSDF.h b/Engine/Source/Runtime/GeometryCore/Public/Implicit/SweepingMeshSDF.h
--- a/Engine/Source/Runtime/GeometryCore/Public/Implicit/SweepingMeshSDF.h
+++ b/Engine/Source/Runtime/GeometryCore/Public/Implicit/SweepingMeshSDF.h
@@ -211,7 +211,7 @@ public:
 			}
 		}
 
-		float fBufferWidth = ExactBandWidth * CellSize;
+		float fBufferWidth = float(ExactBandWidth) * CellSize;
 		if (ComputeMode == EComputeModes::NarrowBand_SpatialFloodFill)
 		{
 			fBufferWidth = (float)FMath::Max(fBufferWidth, float(NarrowBandMaxDistance));
@@ -303,7 +303,7 @@ private:
 
 	float upper_bound(const FDenseGrid3f& GridIn) const
 	{
-		return (float)((GridIn.GetDimensions().X + GridIn.GetDimensions().Y + GridIn.GetDimensions().Z) * CellSize);
+		return (float(GridIn.GetDimensions().X + GridIn.GetDimensions().Y + GridIn.GetDimensions().Z) * CellSize);
 	}
 
 	float cell_tri_dist(const FVector3i& Idx, int TID) const
@@ -744,7 +744,7 @@ private:
 		TArray<FCriticalSection> GridSections;
 		GridSections.SetNum(NumSections);
 		int64 TotalGridCellCount = NI * NJ * NK;
-		int64 SectionSize = FMath::CeilToInt(TotalGridCellCount / (float)NumSections);
+		int64 SectionSize = FMath::CeilToInt(float(TotalGridCellCount) / (float)NumSections);
 		
 		// this returns the FCriticalSection to use for the given span of values
 		auto GetGridSectionLock = [this, SectionSize, &Distances, &GridSections](FVector3i CellGridIndex) -> FCriticalSection*
@@ -1005,7 +1005,7 @@ private:
 			{
 				for (int I = i0; I != i1; I += di)
 				{
-					FVector3d gx(I * DX + Origin[0], J * DX + Origin[1], K * DX + Origin[2]);
+					FVector3d gx(float(I) * DX + Origin[0], float(J) * DX + Origin[1], float(K) * DX + Origin[2]);
 					check_neighbour(phi, closest_tri, gx, I, J, K, I - di, J, K);
 					check_neighbour(phi, closest_tri, gx, I, J, K, I, J - dj, K);
 					check_neighbour(phi, closest_tri, gx, I, J, K, I - di, J - dj, K);
diff --git a/Engine/Source/Runtime/GeometryCore/Public/Spatial/FastWinding.h b/Engine/Source/Runtime/GeometryCore/Public/Spatial/FastWinding.h
--- a/Engine/Source/Runtime/GeometryCore/Public/Spatial/FastWinding.h
+++ b/Engine/Source/Runtime/GeometryCore/Public/Spatial/FastWinding.h
@@ -159,7 +159,7 @@ namespace FastTriWinding
 			double R;
 		};
 
-		auto ExpansionTransform = [&, P = CentroidData.P](int64 TriangleSubsetIndex) -> OrderData
+		auto ExpansionTransform = [&, P = CentroidData.P](int32 TriangleSubsetIndex) -> OrderData
 		{
 			check(TriangleSubsetIndex < TriangleArray.Num());
 			int tid = TriangleArray[(int)TriangleSubsetIndex];
