GeometryCore: FitOrientedBox3 -  fixing implicit conversion warning.

diff --git a/Engine/Source/Runtime/GeometryCore/Private/CompGeom/FitOrientedBox3.cpp b/Engine/Source/Runtime/GeometryCore/Private/CompGeom/FitOrientedBox3.cpp
--- a/Engine/Source/Runtime/GeometryCore/Private/CompGeom/FitOrientedBox3.cpp
+++ b/Engine/Source/Runtime/GeometryCore/Private/CompGeom/FitOrientedBox3.cpp
@@ -179,7 +179,7 @@ TOrientedBox3<RealType> FitOrientedBox3Points(int32 NumPts, TFunctionRef<TVector
 	
 	TArray<TVector<RealType>> Normals;
 	Normals.Reserve(Tris.Num() + 1);
-	Normals.Emplace(0, 0, 1); // Always insert a major axis normal, to fit at least as well as the AABB
+	Normals.Emplace(RealType(0), RealType(0), RealType(1)); // Always insert a major axis normal, to fit at least as well as the AABB
 	if (SameNormalTolerance > 0)
 	{
 		// filter normals for duplicates
