GeometryCore: Update DynamicVector so the iterator can access all the data in a very large vector.  Also fix type conversion warnings.

diff --git a/Engine/Source/Runtime/GeometryCore/Public/Util/DynamicVector.h b/Engine/Source/Runtime/GeometryCore/Public/Util/DynamicVector.h
--- a/Engine/Source/Runtime/GeometryCore/Public/Util/DynamicVector.h
+++ b/Engine/Source/Runtime/GeometryCore/Public/Util/DynamicVector.h
@@ -10,6 +10,7 @@
 #include <UObject/UE5MainStreamObjectVersion.h>
 #include "VectorTypes.h"
 #include "IndexTypes.h"
+#include "Math/NumericLimits.h"
 
 namespace UE
 {
@@ -27,6 +28,8 @@ template <class Type>
 class TDynamicVector
 {
 public:
+	static constexpr uint32 MaxSize = MAX_uint32;
+
 	TDynamicVector()
 	{
 		Blocks.Add(new BlockType());
@@ -56,32 +59,32 @@ public:
 
 	TDynamicVector(const TArray<Type>& Array)
 	{
-		SetNum(Array.Num());
-		for (int Idx = 0; Idx < Array.Num(); Idx++)
+		SetNum((unsigned int) Array.Num());
+		for (int32 Idx = 0; Idx < Array.Num(); Idx++)
 		{
-			(*this)[Idx] = Array[Idx];
+			(*this)[(unsigned int) Idx] = Array[Idx];
 		}
 	}
 	TDynamicVector(TArrayView<const Type> Array)
 	{
-		SetNum(Array.Num());
-		for (int Idx = 0; Idx < Array.Num(); Idx++)
+		SetNum((unsigned int) Array.Num());
+		for (int32 Idx = 0; Idx < Array.Num(); Idx++)
 		{
-			(*this)[Idx] = Array[Idx];
+			(*this)[(unsigned int) Idx] = Array[Idx];
 		}
 	}
 
 	inline void Clear();
 	inline void Fill(const Type& Value);
-	inline void Resize(size_t Count);
-	inline void Resize(size_t Count, const Type& InitValue);
+	inline void Resize(unsigned int Count);
+	inline void Resize(unsigned int Count, const Type& InitValue);
 	/// Resize if Num() is less than Count; returns true if resize occurred
-	inline bool SetMinimumSize(size_t Count, const Type& InitValue);
-	inline void SetNum(size_t Count) { Resize(Count); }
+	inline bool SetMinimumSize(unsigned int Count, const Type& InitValue);
+	inline void SetNum(unsigned int Count) { Resize(Count); }
 
 	inline bool IsEmpty() const { return CurBlock == 0 && CurBlockUsed == 0; }
 	inline size_t GetLength() const { return CurBlock * BlockSize + CurBlockUsed; }
-	inline size_t Num() const { return CurBlock * BlockSize + CurBlockUsed; }
+	inline size_t Num() const { return GetLength(); }
 	static int GetBlockSize() { return BlockSize; }
 	inline size_t GetByteCount() const { return Blocks.Num() * BlockSize * sizeof(Type); }
 
@@ -239,10 +242,10 @@ public:
 
 	private:
 		friend class TDynamicVector;
-		FIterator(TDynamicVector* DVectorIn, int IdxIn)
+		FIterator(TDynamicVector* DVectorIn, unsigned int IdxIn)
 			: DVector(DVectorIn), Idx(IdxIn){}
 		TDynamicVector* DVector{};
-		int Idx{0};
+		unsigned int Idx{0};
 	};
 
 	/** @return iterator at beginning of vector */
@@ -253,7 +256,7 @@ public:
 	/** @return iterator at end of vector */
 	FIterator end()
 	{
-		return FIterator{this, (int)GetLength()};
+		return FIterator{this,  (unsigned int)GetLength()};
 	}
 
 	/*
@@ -288,10 +291,10 @@ public:
 
 	private:
 		friend class TDynamicVector;
-		FConstIterator(const TDynamicVector* DVectorIn, int IdxIn)
+		FConstIterator(const TDynamicVector* DVectorIn, unsigned int IdxIn)
 			: DVector(DVectorIn), Idx(IdxIn){}
 		const TDynamicVector* DVector{};
-		int Idx{0};
+		unsigned int Idx{0};
 	};
 
 	/** @return iterator at beginning of vector */
@@ -302,7 +305,7 @@ public:
 	/** @return iterator at end of vector */
 	FConstIterator end() const
 	{
-		return FConstIterator{this, (int)GetLength()};
+		return FConstIterator{this,  (unsigned int)GetLength()};
 	}
 
 private:
@@ -310,6 +313,7 @@ private:
 	static constexpr int BlockSize = 1 << nShiftBits;
 	static constexpr int BlockIndexBitmask = BlockSize - 1; // low 9 bits
 	static_assert( BlockSize && ((BlockSize & (BlockSize - 1)) == 0), "DynamicVector: BlockSize must be a power of two");
+	static constexpr unsigned int MaxBlockCount = MaxSize / BlockSize + 1;
 
 	unsigned int CurBlock{0};  //< Current block index; always points to the block with the last item in the vector or it is set to zero if the vector is empty. 
 	unsigned int CurBlockUsed{0};  //< Number of used items in the current block.
@@ -599,11 +603,11 @@ public:
 	{
 		Data.Fill(Value);
 	}
-	inline void Resize(size_t Count)
+	inline void Resize(unsigned int Count)
 	{
 		Data.Resize(Count * N);
 	}
-	inline void Resize(size_t Count, const Type& InitValue)
+	inline void Resize(unsigned int Count, const Type& InitValue)
 	{
 		Data.Resize(Count * N, InitValue);
 	}
@@ -759,7 +763,7 @@ void TDynamicVector<Type>::Fill(const Type& Value)
 }
 
 template <class Type>
-void TDynamicVector<Type>::Resize(size_t Count)
+void TDynamicVector<Type>::Resize(unsigned int Count)
 {
 	if (GetLength() == Count)
 	{
@@ -791,18 +795,18 @@ void TDynamicVector<Type>::Resize(size_t Count)
 }
 
 template <class Type>
-void TDynamicVector<Type>::Resize(size_t Count, const Type& InitValue)
+void TDynamicVector<Type>::Resize(unsigned int Count, const Type& InitValue)
 {
 	size_t nCurSize = GetLength();
 	Resize(Count);
-	for (size_t Index = nCurSize; Index < Count; ++Index)
+	for (unsigned int Index = (unsigned int)nCurSize; Index < Count; ++Index)
 	{
 		this->operator[](Index) = InitValue;
 	}
 }
 
 template <class Type>
-bool TDynamicVector<Type>::SetMinimumSize(size_t Count, const Type& InitValue)
+bool TDynamicVector<Type>::SetMinimumSize(unsigned int Count, const Type& InitValue)
 {
 	size_t nCurSize = GetLength();
 	if (Count <= nCurSize)
@@ -810,7 +814,7 @@ bool TDynamicVector<Type>::SetMinimumSize(size_t Count, const Type& InitValue)
 		return false;
 	}
 	Resize(Count);
-	for (size_t Index = nCurSize; Index < Count; ++Index)
+	for (unsigned int Index = (unsigned int)nCurSize; Index < Count; ++Index)
 	{
 		this->operator[](Index) = InitValue;
 	}
@@ -820,6 +824,7 @@ bool TDynamicVector<Type>::SetMinimumSize(size_t Count, const Type& InitValue)
 template <class Type>
 void TDynamicVector<Type>::Add(const Type& Value)
 {
+	checkSlow(size_t(MaxSize) >= GetLength() + 1)
 	if (CurBlockUsed == BlockSize)
 	{
 		if (CurBlock == Blocks.Num() - 1)
@@ -839,7 +844,7 @@ void TDynamicVector<Type>::Add(const TDynamicVector<Type>& AddData)
 {
 	// @todo it could be more efficient to use memcopies here...
 	size_t nSize = AddData.Num();
-	for (unsigned int k = 0; k < nSize; ++k)
+	for (unsigned int k = 0; k < (unsigned int) nSize; ++k)
 	{
 		Add(AddData[k]);
 	}
@@ -901,7 +906,7 @@ void TDynamicVector<Type>::InsertAt(const Type& AddData, unsigned int Index, con
 	size_t nCurSize = GetLength();
 	InsertAt(AddData, Index);
 	// initialize all new values up to (but not including) the inserted index
-	for (size_t i = nCurSize; i < (size_t)Index; ++i)
+	for (unsigned int i = (unsigned int)nCurSize; i < Index; ++i)
 	{
 		this->operator[](i) = InitValue;
 	}
diff --git a/Engine/Source/Runtime/GeometryCore/Public/Util/RefCountVector.h b/Engine/Source/Runtime/GeometryCore/Public/Util/RefCountVector.h
--- a/Engine/Source/Runtime/GeometryCore/Public/Util/RefCountVector.h
+++ b/Engine/Source/Runtime/GeometryCore/Public/Util/RefCountVector.h
@@ -237,7 +237,7 @@ public:
 	 * @param IncrementRefCount Callable for incrementing an item, i.e. an item that has been referenced before.
 	 */
 	template <typename IterateFunc, typename AllocateRefCountFunc, typename IncrementRefCountFunc>
-	void Rebuild(SIZE_T Num, IterateFunc&& Iterate, AllocateRefCountFunc&& AllocateRefCount, IncrementRefCountFunc&& IncrementRefCount)
+	void Rebuild(unsigned int Num, IterateFunc&& Iterate, AllocateRefCountFunc&& AllocateRefCount, IncrementRefCountFunc&& IncrementRefCount)
 	{
 		// Initialize ref counts to the given number of elements.
 		RefCounts.Resize(Num);
@@ -265,10 +265,10 @@ public:
 		Forward<IterateFunc>(Iterate)(UpdateRefCount);
 
 		// Add unused elements to free list.
-		const SIZE_T FreeIndicesNum = Num - UsedCount;
+		const unsigned int FreeIndicesNum = Num - UsedCount;
 		FreeIndices.SetNum(FreeIndicesNum);
-		SIZE_T FreeIndicesIndex = 0;
-		for (SIZE_T Index = 0; (Index < Num) & (FreeIndicesIndex < FreeIndicesNum); ++Index)
+		unsigned int FreeIndicesIndex = 0;
+		for (unsigned int Index = 0; (Index < Num) & (FreeIndicesIndex < FreeIndicesNum); ++Index)
 		{
 			if (RefCounts[Index] == INVALID_REF_COUNT)
 			{
