Add new oriented box fit variant that starts from the fast DiTO box and optimizes the fit iteratively
+ Add ProgressCancel to box fits
+ Change editor code to use the new fit algorithms

diff --git a/Engine/Plugins/Runtime/GeometryProcessing/Source/GeometryAlgorithms/Private/MinVolumeBox3.cpp b/Engine/Plugins/Runtime/GeometryProcessing/Source/GeometryAlgorithms/Private/MinVolumeBox3.cpp
--- a/Engine/Plugins/Runtime/GeometryProcessing/Source/GeometryAlgorithms/Private/MinVolumeBox3.cpp
+++ b/Engine/Plugins/Runtime/GeometryProcessing/Source/GeometryAlgorithms/Private/MinVolumeBox3.cpp
@@ -2,25 +2,10 @@
 
 #include "MinVolumeBox3.h"
 
-#if defined(_MSC_VER) && USING_CODE_ANALYSIS
-#pragma warning(push)
-#pragma warning(disable : 28020)		// disable this warning that occurs in GteMinimumVolumeBox3.h
-#endif
-THIRD_PARTY_INCLUDES_START
-#include "ThirdParty/GTEngine/Mathematics/GteBSNumber.h"
-#include "ThirdParty/GTEngine/Mathematics/GteBSRational.h"
-#include "ThirdParty/GTEngine/Mathematics/GteUIntegerFP32.h"
-#include "ThirdParty/GTEngine/Mathematics/GteUIntegerAP32.h"
-#include "ThirdParty/GTEngine/Mathematics/GteConvexHull3.h"
-#include "ThirdParty/GTEngine/Mathematics/GteMinimumVolumeBox3.h"
-THIRD_PARTY_INCLUDES_END
-#if defined(_MSC_VER) && USING_CODE_ANALYSIS
-#pragma warning(pop)
-#endif
-
 #include "CompGeom/ConvexHull3.h"
+#include "CompGeom/FitOrientedBox3.h"
+#include "CompGeom/DiTOrientedBox.h"
 
-#include "GteUtil.h"
 #include "Util/ProgressCancel.h"
 #include "Util/IteratorUtil.h"
 
@@ -33,78 +18,51 @@ namespace Geometry {
 template <typename RealType>
 struct TMinVolumeBox3Internal
 {
-	using PreciseHullNumberType = gte::BSNumber<gte::UIntegerFP32<197>>;		// 197 is from ConvexHull3 documentation
-	using PreciseBoxNumberType = gte::BSRational<gte::UIntegerAP32>;
-	using DVector3 = gte::Vector3<double>;
-
-	bool bUseExactBox;
-	TArray<DVector3> DoubleInput;
-
 	TOrientedBox3<RealType> Result;
-	bool bSolutionOK;
+	bool bSolutionOK = false;
+	
+	// Settings
 
-	void SetPoint(int32 Index, const TVector<RealType>& Point)
+	RealType SameNormalTolerance = (RealType)0.01;
+	int32 OptimizeIterations = 10;
+	EDiTO DiTODirections = EDiTO::DiTO_26;
+	bool bMostAccurateFit = false;
+	EBox3FitCriteria FitCriteria = EBox3FitCriteria::Volume;
+
+	void SetSolution(TOrientedBox3<RealType> ResultIn)
 	{
-		DoubleInput[Index] = DVector3{ {(double)Point.X, (double)Point.Y, (double)Point.Z} };
+		Result = ResultIn;
+		bSolutionOK = true;
 	}
 
-	bool ComputeResult(FProgressCancel* Progress)
+	bool ComputeResult(int32 NumPoints, TFunctionRef<TVector<RealType>(int32)> GetPointFunc, FProgressCancel* Progress)
 	{
-		gte::OrientedBox3<double> MinimalBox = gte::OrientedBox3<double>();
-
-		FConvexHull3d HullCompute;
-		bSolutionOK = HullCompute.Solve(DoubleInput.Num(),
-										[this](int32 Index) 
-		{
-			return FVector3d{ DoubleInput[Index][0], DoubleInput[Index][1], DoubleInput[Index][2]};
-		});
-
-		if (!bSolutionOK)
-		{
-			return false;
-		}
-
-		int NumIndices = 3 * HullCompute.GetTriangles().Num();
-		if (NumIndices < 1)
-		{
-			bSolutionOK = false;
-			return false;
-		}
-
-		if (Progress && Progress->Cancelled())
-		{
-			return false;
-		}
-
-		const int* Indices = static_cast<const int*>(&(HullCompute.GetTriangles()[0][0]));		// Eww
-
-		if (bUseExactBox)
+		if (bMostAccurateFit)
 		{
-			gte::MinimumVolumeBox3<double, PreciseBoxNumberType> BoxCompute;
-			MinimalBox = BoxCompute(DoubleInput.Num(), &DoubleInput[0], NumIndices, Indices, Progress);
+			Result = FitOrientedBox3Points<RealType>(NumPoints, GetPointFunc, [](int32)->bool {return true;}, FitCriteria, SameNormalTolerance, Progress);
+			if (Progress && Progress->Cancelled())
+			{
+				return false;
+			}
 		}
 		else
 		{
-			gte::MinimumVolumeBox3<double, double> DoubleCompute;
-			MinimalBox = DoubleCompute(DoubleInput.Num(), &DoubleInput[0], NumIndices, Indices, Progress);
+			TOrientedBox3<RealType> InitialBox = ComputeOrientedBBox<RealType>(DiTODirections, NumPoints, GetPointFunc);
+			Result = OptimizeOrientedBox3Points<RealType>(InitialBox, OptimizeIterations, NumPoints, GetPointFunc, [](int32)->bool {return true;}, FitCriteria, Progress);
+			if (Progress && Progress->Cancelled())
+			{
+				return false;
+			}
 		}
-		bSolutionOK = true;
 
 		// if resulting box is not finite, something went wrong, just return an empty box
-		FVector3d Extents = Convert(MinimalBox.extent);
-		if (!FMathd::IsFinite(Extents.SquaredLength()))
+		if (!FMathd::IsFinite(Result.Extents.SquaredLength()))
 		{
 			bSolutionOK = false;
-			MinimalBox = gte::OrientedBox3<double>();
+			return false;
 		}
 
-		Result.Frame = TFrame3<RealType>(
-			TVector<RealType>((RealType)MinimalBox.center[0], (RealType)MinimalBox.center[1], (RealType)MinimalBox.center[2]),
-			TVector<RealType>((RealType)MinimalBox.axis[0][0], (RealType)MinimalBox.axis[0][1], (RealType)MinimalBox.axis[0][2]),
-			TVector<RealType>((RealType)MinimalBox.axis[1][0], (RealType)MinimalBox.axis[1][1], (RealType)MinimalBox.axis[1][2]),
-			TVector<RealType>((RealType)MinimalBox.axis[2][0], (RealType)MinimalBox.axis[2][1], (RealType)MinimalBox.axis[2][2]) );
-		Result.Extents = TVector<RealType>((RealType)MinimalBox.extent[0], (RealType)MinimalBox.extent[1], (RealType)MinimalBox.extent[2]);
-
+		bSolutionOK = true;
 		return true;
 	}
 
@@ -114,43 +72,38 @@ struct TMinVolumeBox3Internal
 } // end namespace UE
 
 template<typename RealType>
-bool TMinVolumeBox3<RealType>::Solve(int32 NumPoints, TFunctionRef<TVector<RealType>(int32)> GetPointFunc, bool bUseExactBox, FProgressCancel* Progress )
+bool TMinVolumeBox3<RealType>::Solve(int32 NumPoints, TFunctionRef<TVector<RealType>(int32)> GetPointFunc, bool bMostAccurateFit, FProgressCancel* Progress )
 {
-	Initialize(NumPoints, bUseExactBox);
+	Initialize(bMostAccurateFit);
 	check(Internal);
 
-	for (int32 k = 0; k < NumPoints; ++k)
-	{
-		TVector<RealType> Point = GetPointFunc(k);
-		Internal->SetPoint(k, Point);
-	}
-	
-	return Internal->ComputeResult(Progress);
+	return Internal->ComputeResult(NumPoints, GetPointFunc, Progress);
 }
 
 
 template<typename RealType>
-bool TMinVolumeBox3<RealType>::SolveSubsample(int32 NumPoints, int32 MaxPoints, TFunctionRef<TVector<RealType>(int32)> GetPointFunc, bool bUseExactBox, FProgressCancel* Progress)
+bool TMinVolumeBox3<RealType>::SolveSubsample(int32 NumPoints, int32 MaxPoints, TFunctionRef<TVector<RealType>(int32)> GetPointFunc, bool bMostAccurateFit, FProgressCancel* Progress)
 {
 	if (NumPoints <= MaxPoints)
 	{
-		return Solve(NumPoints, GetPointFunc, bUseExactBox, Progress);
+		return Solve(NumPoints, GetPointFunc, bMostAccurateFit, Progress);
 	}
 
-	Initialize(MaxPoints, bUseExactBox);
+	Initialize(bMostAccurateFit);
 	check(Internal);
 
 	int32 k = 0;
 	FModuloIteration Iter(NumPoints);
 	int32 Index;
+	TArray<TVector<RealType>> ReducedPoints;
+	ReducedPoints.Reserve(NumPoints);
 	while (Iter.GetNextIndex(Index) && k < MaxPoints)
 	{
 		TVector<RealType> Point = GetPointFunc(Index);
-		Internal->SetPoint(k, Point);
-		k++;
+		ReducedPoints.Add(Point);
 	}
 
-	return Internal->ComputeResult(Progress);
+	return Internal->ComputeResult(ReducedPoints.Num(), [&ReducedPoints](int32 PtIdx) {return ReducedPoints[PtIdx];}, Progress);
 }
 
 template<typename RealType>
@@ -168,11 +121,10 @@ void TMinVolumeBox3<RealType>::GetResult(TOrientedBox3<RealType>& BoxOut)
 
 
 template<typename RealType>
-void TMinVolumeBox3<RealType>::Initialize(int32 NumPoints, bool bUseExactBox)
+void TMinVolumeBox3<RealType>::Initialize(bool bMostAccurateFit)
 {
 	Internal = MakePimpl<TMinVolumeBox3Internal<RealType>>();
-	Internal->bUseExactBox = bUseExactBox;
-	Internal->DoubleInput.SetNum(NumPoints);
+	Internal->bMostAccurateFit = bMostAccurateFit;
 }
 
 
diff --git a/Engine/Plugins/Runtime/GeometryProcessing/Source/GeometryAlgorithms/Public/MinVolumeBox3.h b/Engine/Plugins/Runtime/GeometryProcessing/Source/GeometryAlgorithms/Public/MinVolumeBox3.h
--- a/Engine/Plugins/Runtime/GeometryProcessing/Source/GeometryAlgorithms/Public/MinVolumeBox3.h
+++ b/Engine/Plugins/Runtime/GeometryProcessing/Source/GeometryAlgorithms/Public/MinVolumeBox3.h
@@ -30,20 +30,22 @@ public:
 	 * Calculate the minimal box for the given point set.
 	 * @param NumPoints number of points in the set, ie GetPointFunc can be called for any integer in range [0...NumPoints)
 	 * @param GetPointFunc function that returns a 3D point for a valid Index
-	 * @param bUseExactBox If true, high-precision number types are used for the minimal-box calculation, rather than doubles. This is *much* slower but more accurate (but not recommended).
+	 * @param bMostAccurateFit if true, use the most expensive method to get the best-possible fit 
+	 * @param Progress optionally allow early-cancel of the box fit operation
 	 * @return true if minimal box was found
 	 */
-	bool Solve(int32 NumPoints, TFunctionRef<TVector<RealType>(int32)> GetPointFunc, bool bUseExactBox = false, FProgressCancel* Progress = nullptr);
+	bool Solve(int32 NumPoints, TFunctionRef<TVector<RealType>(int32)> GetPointFunc, bool bMostAccurateFit = false, FProgressCancel* Progress = nullptr);
 
 	/**
 	 * Calculate the minimal box for a Subsampling of MaxPoints points of a point set
 	 * @param NumPoints number of points in the set, ie GetPointFunc can be called for any integer in range [0...NumPoints)
 	 * @param NumSamplePoints maximum number of points to sample from NumPoints. Currently a random-ish subset is selected.
 	 * @param GetPointFunc function that returns a 3D point for a valid Index
-	 * @param bUseExactBox If true, high-precision number types are used for the minimal-box calculation, rather than doubles. This is *much* slower but more accurate (but not recommended).
+	 * @param bMostAccurateFit if true, use the most expensive method to get the best-possible fit
+	 * @param Progress optionally allow early-cancel of the box fit operation
 	 * @return true if minimal box was found
 	 */
-	bool SolveSubsample(int32 NumPoints, int32 NumSamplePoints, TFunctionRef<TVector<RealType>(int32)> GetPointFunc, bool bUseExactBox = false, FProgressCancel* Progress = nullptr);
+	bool SolveSubsample(int32 NumPoints, int32 NumSamplePoints, TFunctionRef<TVector<RealType>(int32)> GetPointFunc, bool bMostAccurateFit = false, FProgressCancel* Progress = nullptr);
 
 
 	/** @return true if minimal box is available */
@@ -53,7 +55,7 @@ public:
 	void GetResult(TOrientedBox3<RealType>& BoxOut);
 
 protected:
-	void Initialize(int32 NumPoints, bool bUseExactBox);
+	void Initialize(bool bMostAccurateFit);
 
 	TPimplPtr<TMinVolumeBox3Internal<RealType>> Internal;
 };
diff --git a/Engine/Source/Runtime/GeometryCore/Private/CompGeom/FitOrientedBox3.cpp b/Engine/Source/Runtime/GeometryCore/Private/CompGeom/FitOrientedBox3.cpp
--- a/Engine/Source/Runtime/GeometryCore/Private/CompGeom/FitOrientedBox3.cpp
+++ b/Engine/Source/Runtime/GeometryCore/Private/CompGeom/FitOrientedBox3.cpp
@@ -12,6 +12,7 @@
 #include "CompGeom/ConvexHull3.h"
 #include "CompGeom/ExactPredicates.h"
 #include "Async/ParallelTransformReduce.h"
+#include "Util/ProgressCancel.h"
 
 namespace UE
 {
@@ -71,13 +72,10 @@ namespace
 
 		return true;
 	}
-}
 
-template <typename RealType>
-TOrientedBox3<RealType> FitOrientedBox3Points(int32 NumPts, TFunctionRef<TVector<RealType>(int32)> GetPtFn, TFunctionRef<bool(int32)> FilterFn, EBox3FitCriteria FitMethod, RealType SameNormalTolerance)
-{
-	// For failure cases, we fall back to an axis-aligned box instead
-	auto GetFallbackBox = [NumPts, &GetPtFn, &FilterFn]() -> TOrientedBox3<RealType>
+	// Helper to compute an oriented box from the axis-aligned box, as a fallback
+	template <typename RealType>
+	TOrientedBox3<RealType> GetAxisAligned(int32 NumPts, TFunctionRef<TVector<RealType>(int32)> GetPtFn, TFunctionRef<bool(int32)> FilterFn)
 	{
 		TAxisAlignedBox3<RealType> AABB;
 		for (int32 Idx = 0; Idx < NumPts; ++Idx)
@@ -88,6 +86,104 @@ TOrientedBox3<RealType> FitOrientedBox3Points(int32 NumPts, TFunctionRef<TVector
 			}
 		}
 		return TOrientedBox3<RealType>(AABB);
+	}
+}
+
+template <typename RealType>
+TOrientedBox3<RealType> OptimizeOrientedBox3Points(const TOrientedBox3<RealType>& InitialBox, int32 NumIterations, int32 NumPoints,
+	TFunctionRef<TVector<RealType>(int32)> GetPointFunc, TFunctionRef<bool(int32)> Filter, EBox3FitCriteria FitMethod, FProgressCancel* ProgressCancel)
+{
+	TArray<TVector2<RealType>> Projection;
+	Projection.Reserve(NumPoints);
+	TOrientedBox3<RealType> BestBox = InitialBox;
+	RealType BestScore = (FitMethod == EBox3FitCriteria::Volume) ? BestBox.Volume() : BestBox.SurfaceArea();
+	int32 BestBoxIter = -1; // iteration that BestBox was updated
+	for (int32 Iter = 0; Iter < NumIterations; ++Iter)
+	{
+		if (ProgressCancel && ProgressCancel->Cancelled())
+		{
+			return BestBox;
+		}
+
+		if (Iter > BestBoxIter + 3)
+		{
+			return BestBox;
+		}
+		int32 Axis = Iter % 3;
+		TVector<RealType> FixedAxis = BestBox.GetAxis(Axis);
+		int32 BasisXIdx = (Axis + 1) % 3;
+		int32 BasisYIdx = (Axis + 2) % 3;
+		TVector<RealType> BasisX = BestBox.GetAxis(BasisXIdx);
+		TVector<RealType> BasisY = BestBox.GetAxis(BasisYIdx);
+
+		Projection.Reset();
+		for (int32 PtIdx = 0; PtIdx < NumPoints; ++PtIdx)
+		{
+			if (Filter(PtIdx))
+			{
+				TVector<RealType> Pt = GetPointFunc(PtIdx);
+				Projection.Emplace(Pt.Dot(BasisX), Pt.Dot(BasisY));
+			}
+		}
+		RealType Depth = BestBox.Extents[Axis] * 2;
+		using TFitFn = TFunctionRef<RealType(RealType, RealType)>;
+		TFitFn Score2 = FitMethod == EBox3FitCriteria::Volume ?
+			TFitFn([](RealType Width, RealType Height) -> RealType
+				{
+					// for volume, we can just use area when finding the best 2D box
+					return Width * Height;
+				}) :
+			TFitFn([Depth](RealType Width, RealType Height) -> RealType
+				{
+					// for surface area, we also need to use the depth in the 2D box search
+					return Width * Height + (Width + Height) * Depth;
+				});
+		TOrientedBox2<RealType> ProjBox = FitOrientedBox2Points<RealType>(Projection, Score2);
+
+		// rotate the solution so we don't arbitrarily swap or flip the axes
+		if (FMath::Abs(ProjBox.UnitAxisX.X) < FMath::Abs(ProjBox.UnitAxisX.Y))
+		{
+			ProjBox.UnitAxisX = PerpCW(ProjBox.UnitAxisX);
+			Swap(ProjBox.Extents.X, ProjBox.Extents.Y);
+		}
+		if (ProjBox.UnitAxisX.X < 0)
+		{
+			ProjBox.UnitAxisX = -ProjBox.UnitAxisX;
+		}
+
+		TVector<RealType> Extents = BestBox.Extents;
+		Extents[BasisXIdx] = ProjBox.Extents.X;
+		Extents[BasisYIdx] = ProjBox.Extents.Y;
+		TAxisAlignedBox3<RealType> SizeBox(-Extents, Extents);
+		RealType Score = FitMethod == EBox3FitCriteria::Volume ? SizeBox.Volume() : SizeBox.SurfaceArea();
+		if (Score < BestScore)
+		{
+			BestScore = Score;
+			BestBoxIter = Iter;
+			RealType CenterOnFixedAxis = FixedAxis.Dot(BestBox.Center());
+			// Convert result to an oriented box
+			TVector2<RealType> ProjC = ProjBox.Center();
+			TVector<RealType> Center = FixedAxis * CenterOnFixedAxis + BasisX * ProjC.X + BasisY * ProjC.Y;
+			TVector<RealType> Axes[3];
+			Axes[Axis] = FixedAxis;
+			TVector2<RealType> ProjAxisX = ProjBox.AxisX(), ProjAxisY = ProjBox.AxisY();
+			Axes[BasisXIdx] = ProjAxisX.X * BasisX + ProjAxisX.Y * BasisY;
+			Axes[BasisYIdx] = ProjAxisY.X * BasisX + ProjAxisY.Y * BasisY;
+			TFrame3<RealType> Frame(Center, Axes[0], Axes[1], Axes[2]);
+			BestBox = TOrientedBox3<RealType>(Frame, Extents);
+		}
+	}
+
+	return BestBox;
+}
+
+template <typename RealType>
+TOrientedBox3<RealType> FitOrientedBox3Points(int32 NumPts, TFunctionRef<TVector<RealType>(int32)> GetPtFn, TFunctionRef<bool(int32)> FilterFn, EBox3FitCriteria FitMethod, RealType SameNormalTolerance, FProgressCancel* ProgressCancel)
+{
+	// For failure cases, we fall back to an axis-aligned box instead
+	auto GetFallbackBox = [NumPts, &GetPtFn, &FilterFn]() -> TOrientedBox3<RealType>
+	{
+		return GetAxisAligned(NumPts, GetPtFn, FilterFn);
 	};
 
 	TConvexHull3<RealType> Hull;
@@ -229,8 +325,18 @@ TOrientedBox3<RealType> FitOrientedBox3Points(int32 NumPts, TFunctionRef<TVector
 	bool bSingleThread = Tris.Num() < 1000; // TODO: test to find the threshold at which threading is helpful
 	int64 TasksToUse = bSingleThread ? 1 : 64;
 
-	auto FitBoxForDirection = [&GetPtFn, &Normals, &Tris, &TriNbrs, &UsedVerts, &FitMethod, bSingleThread](int32 NormalIdx) -> FResult
+	if (ProgressCancel && ProgressCancel->Cancelled())
 	{
+		return TOrientedBox3<RealType>();
+	}
+
+	auto FitBoxForDirection = [&GetPtFn, &Normals, &Tris, &TriNbrs, &UsedVerts, &FitMethod, bSingleThread, &ProgressCancel](int32 NormalIdx) -> FResult
+	{
+		if (ProgressCancel && ProgressCancel->Cancelled())
+		{
+			return FResult { TMathUtil<RealType>::MaxReal, TOrientedBox3<RealType>() };
+		}
+
 		TVector<RealType> N = Normals[NormalIdx];
 		TArray<bool> Facing;
 		Facing.SetNumUninitialized(Tris.Num());
@@ -305,6 +411,11 @@ TOrientedBox3<RealType> FitOrientedBox3Points(int32 NumPts, TFunctionRef<TVector
 		return B;
 	};
 
+	if (ProgressCancel && ProgressCancel->Cancelled())
+	{
+		return TOrientedBox3<RealType>();
+	}
+
 	FResult BestResult = ParallelTransformReduce<int32, FResult, TFunctionRef<FResult(int32)>, TFunctionRef<FResult(FResult, FResult)>>
 		(Normals.Num(), DefaultInit, FitBoxForDirection, TakeMinBox, TasksToUse);
 
@@ -312,8 +423,16 @@ TOrientedBox3<RealType> FitOrientedBox3Points(int32 NumPts, TFunctionRef<TVector
 }
 
 // explicit instantiations
-template TOrientedBox3<float> GEOMETRYCORE_API FitOrientedBox3Points<float>(int32 NumPts, TFunctionRef<TVector<float>(int32)> GetPtFn, TFunctionRef<bool(int32)> FilterFn, EBox3FitCriteria FitMethod, float SameNormalTolerance);
-template TOrientedBox3<double> GEOMETRYCORE_API FitOrientedBox3Points<double>(int32 NumPts, TFunctionRef<TVector<double>(int32)> GetPtFn, TFunctionRef<bool(int32)> FilterFn, EBox3FitCriteria FitMethod, double SameNormalTolerance);
+template TOrientedBox3<float> GEOMETRYCORE_API FitOrientedBox3Points<float>(int32 NumPts, TFunctionRef<TVector<float>(int32)> GetPtFn, 
+	TFunctionRef<bool(int32)> FilterFn, EBox3FitCriteria FitMethod, float SameNormalTolerance, FProgressCancel* ProgressCancel);
+template TOrientedBox3<double> GEOMETRYCORE_API FitOrientedBox3Points<double>(int32 NumPts, TFunctionRef<TVector<double>(int32)> GetPtFn,
+	TFunctionRef<bool(int32)> FilterFn, EBox3FitCriteria FitMethod, double SameNormalTolerance, FProgressCancel* ProgressCancel);
+
+template TOrientedBox3<float> GEOMETRYCORE_API OptimizeOrientedBox3Points(const TOrientedBox3<float>& InitialBox, int32 NumIterations, int32 NumPoints,
+	TFunctionRef<TVector<float>(int32)> GetPointFunc, TFunctionRef<bool(int32)> Filter, EBox3FitCriteria FitMethod, FProgressCancel* ProgressCancel);
+template TOrientedBox3<double> GEOMETRYCORE_API OptimizeOrientedBox3Points(const TOrientedBox3<double>& InitialBox, int32 NumIterations, int32 NumPoints,
+	TFunctionRef<TVector<double>(int32)> GetPointFunc, TFunctionRef<bool(int32)> Filter, EBox3FitCriteria FitMethod, FProgressCancel* ProgressCancel);
+
 
 }
 }
\ No newline at end of file
diff --git a/Engine/Source/Runtime/GeometryCore/Public/CompGeom/FitOrientedBox3.h b/Engine/Source/Runtime/GeometryCore/Public/CompGeom/FitOrientedBox3.h
--- a/Engine/Source/Runtime/GeometryCore/Public/CompGeom/FitOrientedBox3.h
+++ b/Engine/Source/Runtime/GeometryCore/Public/CompGeom/FitOrientedBox3.h
@@ -7,6 +7,8 @@
 #include "VectorTypes.h"
 #include "OrientedBoxTypes.h"
 
+class FProgressCancel;
+
 namespace UE
 {
 namespace Geometry
@@ -31,7 +33,7 @@ namespace Geometry
 	 */
 	template <typename RealType>
 	TOrientedBox3<RealType> GEOMETRYCORE_API FitOrientedBox3Points(int32 NumPoints, TFunctionRef<TVector<RealType>(int32)> GetPointFunc, TFunctionRef<bool(int32)> Filter,
-		EBox3FitCriteria FitMethod = EBox3FitCriteria::Volume, RealType SameNormalTolerance = 0);
+		EBox3FitCriteria FitMethod = EBox3FitCriteria::Volume, RealType SameNormalTolerance = 0, FProgressCancel* ProgressCancel = nullptr);
 
 	/**
 	 * Compute a close-fitting oriented bounding box to the given points.
@@ -43,16 +45,31 @@ namespace Geometry
 	 * @return						A best-fit TOrientedBox3 that contains the points
 	 */
 	template <typename RealType>
-	TOrientedBox3<RealType> FitOrientedBox3Points(TArrayView<const TVector<RealType>> Points, EBox3FitCriteria FitMethod = EBox3FitCriteria::Volume, RealType SameNormalTolerance = 0)
+	TOrientedBox3<RealType> FitOrientedBox3Points(TArrayView<const TVector<RealType>> Points, EBox3FitCriteria FitMethod = EBox3FitCriteria::Volume, RealType SameNormalTolerance = 0, FProgressCancel* ProgressCancel = nullptr)
 	{
 		TFunctionRef<TVector<RealType>(int32)> GetPtFn = [&Points](int32 Idx)
 		{
 			return Points[Idx];
 		};
-		return FitOrientedBox3Points(Points.Num(), GetPtFn, [](int32 Idx) {return true;}, FitMethod, SameNormalTolerance);
+		return FitOrientedBox3Points(Points.Num(), GetPtFn, [](int32 Idx) {return true;}, FitMethod, SameNormalTolerance, ProgressCancel);
 	}
 
 
+	/**
+	 * Improve an oriented bounding box fit iteratively
+	 *
+	 * @param InitialBox		An initial guess at an oriented bounding box fitting the points
+	 * @param NumIterations		How many passes to take to improve the bounding 
+	 * @param NumPoints			Number of points to consider
+	 * @param GetPointFunc		Function providing array-style access into points
+	 * @param Filter			Optional filter to include only a subset of the points in the output hull
+	 * @param FitMethod			What criteria to optimize
+	 * @return					A best-fit TOrientedBox3 that contains the points
+	 */
+	template <typename RealType>
+	TOrientedBox3<RealType> GEOMETRYCORE_API OptimizeOrientedBox3Points(const TOrientedBox3<RealType>& InitialBox, int32 NumIterations, int32 NumPoints, 
+		TFunctionRef<TVector<RealType>(int32)> GetPointFunc, TFunctionRef<bool(int32)> Filter, EBox3FitCriteria FitMethod = EBox3FitCriteria::Volume, FProgressCancel* ProgressCancel = nullptr);
+
 } // end namespace Geometry
 }// end namespace UE
 
