CurveUtil explicit type fix

diff --git a/Engine/Source/Runtime/GeometryCore/Public/Curve/CurveUtil.h b/Engine/Source/Runtime/GeometryCore/Public/Curve/CurveUtil.h
--- a/Engine/Source/Runtime/GeometryCore/Public/Curve/CurveUtil.h
+++ b/Engine/Source/Runtime/GeometryCore/Public/Curve/CurveUtil.h
@@ -126,7 +126,7 @@ using namespace UE::Math;
 			const TVector2<RealType>& V2 = Vertices[Idx];
 			Area += V1.X * V2.Y - V1.Y * V2.X;
 		}
-		return Area * 0.5;
+		return static_cast<RealType>(Area * 0.5);
 	}
 
 	/**
