fix rotating calipers failing to find extreme point due to two points on hull being close enough that they appear equal w.r.t. the extreme direction

fixes oriented bounding box not containing the output points in GeometryFlowTests

diff --git a/Engine/Source/Runtime/GeometryCore/Private/CompGeom/FitOrientedBox2.cpp b/Engine/Source/Runtime/GeometryCore/Private/CompGeom/FitOrientedBox2.cpp
--- a/Engine/Source/Runtime/GeometryCore/Private/CompGeom/FitOrientedBox2.cpp
+++ b/Engine/Source/Runtime/GeometryCore/Private/CompGeom/FitOrientedBox2.cpp
@@ -46,7 +46,7 @@ TOrientedBox2<RealType> FitOrientedBox2ConvexHull(int32 NumPts, TFunctionRef<TVe
 		for (int32 Idx = Start + 1; Idx < NumPts; ++Idx)
 		{
 			RealType Val = BasisVec.Dot(GetHullPt(Idx));
-			if (Val <= OutExtremeVal)
+			if (Val < OutExtremeVal)
 			{
 				return ExtremeIdx;
 			}
@@ -56,7 +56,7 @@ TOrientedBox2<RealType> FitOrientedBox2ConvexHull(int32 NumPts, TFunctionRef<TVe
 
 		// For first pass, we should only need to test the first point again
 		RealType LastVal = BasisVec.Dot(GetHullPt(0));
-		if (LastVal <= OutExtremeVal)
+		if (LastVal < OutExtremeVal)
 		{
 			return ExtremeIdx;
 		}
@@ -73,7 +73,7 @@ TOrientedBox2<RealType> FitOrientedBox2ConvexHull(int32 NumPts, TFunctionRef<TVe
 		for (int32 Idx = Start + 1; Idx < NumPts; ++Idx)
 		{
 			RealType Val = BasisVec.Dot(GetHullPt(Idx));
-			if (Val <= OutExtremeVal)
+			if (Val < OutExtremeVal)
 			{
 				return ExtremeIdx;
 			}
@@ -85,7 +85,7 @@ TOrientedBox2<RealType> FitOrientedBox2ConvexHull(int32 NumPts, TFunctionRef<TVe
 		for (int32 Idx = 0; Idx < Start; ++Idx)
 		{
 			RealType Val = BasisVec.Dot(GetHullPt(Idx));
-			if (Val <= OutExtremeVal)
+			if (Val < OutExtremeVal)
 			{
 				return ExtremeIdx;
 			}
