GeometryCore: add support for custom per-path-point frames and scaling to FGeneralizedCylinderGenerator
GeometryScript: add new AppendSweepPolygon function, similar to AppendSimpleSweptPolygon but takes the sweep path as a list of FTransform instead of FVector, giving the caller more control over rotation and scaling along the sweep.

diff --git a/Engine/Plugins/Experimental/GeometryScripting/Source/GeometryScriptingCore/Private/MeshPrimitiveFunctions.cpp b/Engine/Plugins/Experimental/GeometryScripting/Source/GeometryScriptingCore/Private/MeshPrimitiveFunctions.cpp
--- a/Engine/Plugins/Experimental/GeometryScripting/Source/GeometryScriptingCore/Private/MeshPrimitiveFunctions.cpp
+++ b/Engine/Plugins/Experimental/GeometryScripting/Source/GeometryScriptingCore/Private/MeshPrimitiveFunctions.cpp
@@ -651,6 +651,68 @@ UDynamicMesh* UGeometryScriptLibrary_MeshPrimitiveFunctions::AppendSimpleSweptPo
 
 
 
+UDynamicMesh* UGeometryScriptLibrary_MeshPrimitiveFunctions::AppendSweepPolygon(
+	UDynamicMesh* TargetMesh,
+	FGeometryScriptPrimitiveOptions PrimitiveOptions,
+	FTransform Transform,
+	const TArray<FVector2D>& PolygonVertices,
+	const TArray<FTransform>& SweepPath,
+	bool bLoop,
+	bool bCapped,
+	float StartScale,
+	float EndScale,
+	float RotationAngleDeg,
+	UGeometryScriptDebug* Debug)
+{
+	if (TargetMesh == nullptr)
+	{
+		UE::Geometry::AppendError(Debug, EGeometryScriptErrorType::InvalidInputs, LOCTEXT("AppendSweepPolygon_NullMesh", "AppendSweepPolygon: TargetMesh is Null"));
+		return TargetMesh;
+	}
+	if (PolygonVertices.Num() < 3)
+	{
+		UE::Geometry::AppendError(Debug, EGeometryScriptErrorType::InvalidInputs, LOCTEXT("AppendSweepPolygon_InvalidPolygon", "AppendSweepPolygon: PolygonVertices array requires at least 3 positions"));
+		return TargetMesh;
+	}
+	if (SweepPath.Num() < 2)
+	{
+		UE::Geometry::AppendError(Debug, EGeometryScriptErrorType::InvalidInputs, LOCTEXT("AppendSweepPolygon_InvalidSweepPath", "AppendSweepPolygon: SweepPath array requires at least 2 positions"));
+		return TargetMesh;
+	}
+
+	FMatrix2d Rotation2D = FMatrix2d::RotationDeg(-RotationAngleDeg);
+	FGeneralizedCylinderGenerator SweepGen;
+	for (const FVector2D& Point : PolygonVertices)
+	{
+		SweepGen.CrossSection.AppendVertex( Rotation2D * FVector2d(Point.X, Point.Y) );
+	}
+	for (const FTransform& SweepXForm : SweepPath)
+	{
+		SweepGen.Path.Add(SweepXForm.GetLocation());
+		FQuaterniond Rotation(SweepXForm.GetRotation());
+		SweepGen.PathFrames.Add(
+			FFrame3d(SweepXForm.GetLocation(), Rotation.AxisY(), Rotation.AxisZ(), Rotation.AxisX())
+		);
+		FVector3d Scale = SweepXForm.GetScale3D();
+		SweepGen.PathScales.Add(FVector2d(Scale.Y, Scale.Z));
+	}
+
+	SweepGen.bLoop = bLoop;
+	SweepGen.bCapped = bCapped;
+	SweepGen.bPolygroupPerQuad = (PrimitiveOptions.PolygroupMode == EGeometryScriptPrimitivePolygroupMode::PerQuad);
+	SweepGen.InitialFrame = FFrame3d(SweepGen.Path[0]);
+	SweepGen.StartScale = StartScale;
+	SweepGen.EndScale = EndScale;
+
+	SweepGen.Generate();
+
+	AppendPrimitive(TargetMesh, &SweepGen, Transform, PrimitiveOptions);
+	return TargetMesh;
+}
+
+
+
+
 UDynamicMesh* UGeometryScriptLibrary_MeshPrimitiveFunctions::AppendRectangleXY(
 	UDynamicMesh* TargetMesh,
 	FGeometryScriptPrimitiveOptions PrimitiveOptions,
diff --git a/Engine/Plugins/Experimental/GeometryScripting/Source/GeometryScriptingCore/Public/GeometryScript/MeshPrimitiveFunctions.h b/Engine/Plugins/Experimental/GeometryScripting/Source/GeometryScriptingCore/Public/GeometryScript/MeshPrimitiveFunctions.h
--- a/Engine/Plugins/Experimental/GeometryScripting/Source/GeometryScriptingCore/Public/GeometryScript/MeshPrimitiveFunctions.h
+++ b/Engine/Plugins/Experimental/GeometryScripting/Source/GeometryScriptingCore/Public/GeometryScript/MeshPrimitiveFunctions.h
@@ -260,6 +260,32 @@ public:
 		float EndScale = 1.0f,
 		UGeometryScriptDebug* Debug = nullptr);
 
+	/**
+	 * Sweep the given 2D PolygonVertices along the SweepPath specified as a set of FTransforms
+	 * If the 2D vertices are (U,V), then in the coordinate space of the FTransform, X points "along" the path,
+	 * Y points "right" (U) and Z points "up" (V).
+	 * @param PolygonVertices vertices of the closed 2D polyon that will be swept along the SweepPath
+	 * @param SweepPath defines the 3D sweep path curve as a 3D poly-path, with rotation and scaling at each polypath vertex taken from the Transform
+	 * @param bLoop if true, SweepPath is considered to be a Loop and a section connecting the end and start of the path is added (bCapped is ignored)
+	 * @param bCapped if true the open ends of the swept generalized cylinder are triangulated
+	 * @param StartScale uniform scaling applied to the 2D polygon at the start of the path. Interpolated via arc length to EndScale at the end of the path.
+	 * @param EndScale uniform scaling applied to the 2D polygon at the end of the path
+	 * @param RotationAngleDeg Rotation applied to the 2D Polygon. Positive rotation rotates clockwise, ie Up/+Z/+V towards Right/+Y/+U
+	 */
+	UFUNCTION(BlueprintCallable, Category = "GeometryScript|Primitives", meta=(ScriptMethod))
+	static UPARAM(DisplayName = "Target Mesh") UDynamicMesh* 
+	AppendSweepPolygon( 
+		UDynamicMesh* TargetMesh, 
+		FGeometryScriptPrimitiveOptions PrimitiveOptions,
+		FTransform Transform,
+		const TArray<FVector2D>& PolygonVertices,
+		const TArray<FTransform>& SweepPath,
+		bool bLoop = false,
+		bool bCapped = true,
+		float StartScale = 1.0f,
+		float EndScale = 1.0f,
+		float RotationAngleDeg = 0.0f,
+		UGeometryScriptDebug* Debug = nullptr);
 
 
 	UFUNCTION(BlueprintCallable, Category = "GeometryScript|Primitives", meta=(ScriptMethod))
diff --git a/Engine/Source/Runtime/GeometryCore/Public/Generators/SweepGenerator.h b/Engine/Source/Runtime/GeometryCore/Public/Generators/SweepGenerator.h
--- a/Engine/Source/Runtime/GeometryCore/Public/Generators/SweepGenerator.h
+++ b/Engine/Source/Runtime/GeometryCore/Public/Generators/SweepGenerator.h
@@ -663,11 +663,16 @@ public:
 	TArray<FVector3d> Path;
 
 	FFrame3d InitialFrame;
+	// If PathFrames.Num == Path.Num, then PathFrames[k] is used for each step instead of the propagated InitialFrame.
+	TArray<FFrame3d> PathFrames;
+	// If PathScales.Num == Path.Num, then PathScales[k] is applied to the CrossSection at each step (this is combined with StartScale/EndScale, but ignored if bLoop=true)
+	TArray<FVector2d> PathScales;
 
 	bool bCapped = false;
 	bool bLoop = false;
 	ECapType CapType = ECapType::FlatTriangulation;
 
+	// 2D uniform scale of the CrossSection, interpolated along the Path (via arc length) from StartScale to EndScale
 	double StartScale = 1.0;
 	double EndScale = 1.0;
 
@@ -694,7 +699,8 @@ public:
 		}
 		int PathNum = Path.Num();
 		
-		bool bApplyScaling = (StartScale != 1.0 || EndScale != 1.0) && (bLoop == false);
+		bool bHavePathScaling = (PathScales.Num() == PathNum);
+		bool bApplyScaling = (bHavePathScaling|| (StartScale != 1.0) || (EndScale != 1.0)) && (bLoop == false);
 		bool bNeedArcLength = (bApplyScaling || bUVScaleRelativeWorld);
 		double TotalPathArcLength = (bNeedArcLength) ? UE::Geometry::CurveUtil::ArcLength<double, FVector3d>(Path, bLoop) : 1.0;
 
@@ -719,22 +725,32 @@ public:
 
 		double AccumArcLength = 0;
 		FFrame3d CrossSectionFrame = InitialFrame;
+		bool bHaveExplicitFrames = (PathFrames.Num() == Path.Num());
 		for (int PathIdx = 0; PathIdx < PathNum; ++PathIdx)
 		{
-			FVector3d Tangent = UE::Geometry::CurveUtil::Tangent<double, FVector3d>(Path, PathIdx, bLoop);
-			CrossSectionFrame.AlignAxis(2, Tangent);
 			FVector3d C = Path[PathIdx];
-			FVector3d X = CrossSectionFrame.X();
-			FVector3d Y = CrossSectionFrame.Y();
-			for (int SubIdx = 0; SubIdx < XNum; SubIdx++)
+			FVector3d X, Y;
+			if (bHaveExplicitFrames == false)
 			{
-				FVector2d XP = CrossSection[SubIdx];
-				if (bApplyScaling)
-				{
-					double T = FMathd::Clamp((AccumArcLength / TotalPathArcLength), 0.0, 1.0);
-					XP *= FMathd::Lerp(StartScale, EndScale, T);
-				}
+				FVector3d Tangent = UE::Geometry::CurveUtil::Tangent<double, FVector3d>(Path, PathIdx, bLoop);
+				CrossSectionFrame.AlignAxis(2, Tangent);
+				X = CrossSectionFrame.X();
+				Y = CrossSectionFrame.Y();
+			}
+			else
+			{
+				C = PathFrames[PathIdx].Origin;
+				X = PathFrames[PathIdx].X();
+				Y = PathFrames[PathIdx].Y();
+			}
+
+			double T = FMathd::Clamp((AccumArcLength / TotalPathArcLength), 0.0, 1.0);
+			double UniformScale = (bApplyScaling) ? FMathd::Lerp(StartScale, EndScale, T) : 1.0;
+			FVector2d PathScaling = (bHavePathScaling) ? PathScales[PathIdx] : FVector2d::One();
 
+			for (int SubIdx = 0; SubIdx < XNum; SubIdx++)
+			{
+				FVector2d XP = UniformScale * PathScaling * CrossSection[SubIdx];
 				FVector2d XN = XNormals[SubIdx];
 				Vertices[SubIdx + PathIdx * XNum] = C + X * XP.X + Y * XP.Y;
 				Normals[SubIdx + PathIdx * XNum] = (FVector3f)(X * XN.X + Y * XN.Y);
