GeometryCore: Add UE::Geometry::VectorUtil::PlaneAngleSignedR implentation, same as PlaneAngleSignedD but returns radians instead of degrees.

diff --git a/Engine/Source/Runtime/GeometryCore/Public/VectorUtil.h b/Engine/Source/Runtime/GeometryCore/Public/VectorUtil.h
--- a/Engine/Source/Runtime/GeometryCore/Public/VectorUtil.h
+++ b/Engine/Source/Runtime/GeometryCore/Public/VectorUtil.h
@@ -272,6 +272,26 @@ namespace VectorUtil
 		return (RealType)(fSign * AngleD(vFrom, vTo));
 	}
 
+	/**
+	 * Calculates angle between VFrom and VTo after projection onto plane with normal defined by PlaneN
+	 * @return angle in radians
+	 */
+	template <typename RealType>
+	inline RealType PlaneAngleSignedR(const TVector<RealType>& VFrom, const TVector<RealType>& VTo, const TVector<RealType>& PlaneN)
+	{
+		TVector<RealType> vFrom = VFrom - VFrom.Dot(PlaneN) * PlaneN;
+		TVector<RealType> vTo = VTo - VTo.Dot(PlaneN) * PlaneN;
+		Normalize(vFrom);
+		Normalize(vTo);
+		TVector<RealType> C = vFrom.Cross(vTo);
+		if (C.SquaredLength() < TMathUtil<RealType>::ZeroTolerance)
+		{ // vectors are parallel
+			return vFrom.Dot(vTo) < 0 ? TMathUtil<RealType>::Pi : (RealType)0;
+		}
+		RealType fSign = C.Dot(PlaneN) < 0 ? (RealType)-1 : (RealType)1;
+		return (RealType)(fSign * AngleR(vFrom, vTo));
+	}
+
 	/**
 	 * tan(theta/2) = +/- sqrt( (1-cos(theta)) / (1+cos(theta)) )
 	 * @return positive value of tan(theta/2) where theta is angle between normalized vectors A and B
