make more mesh creation tools generate tangents for their created meshes
(PolyExt, PathExt, PathRev, and BdryRev)

diff --git a/Engine/Plugins/Experimental/MeshModelingToolsetExp/Source/MeshModelingToolsExp/Private/DrawPolyPathTool.cpp b/Engine/Plugins/Experimental/MeshModelingToolsetExp/Source/MeshModelingToolsExp/Private/DrawPolyPathTool.cpp
--- a/Engine/Plugins/Experimental/MeshModelingToolsetExp/Source/MeshModelingToolsExp/Private/DrawPolyPathTool.cpp
+++ b/Engine/Plugins/Experimental/MeshModelingToolsetExp/Source/MeshModelingToolsExp/Private/DrawPolyPathTool.cpp
@@ -18,6 +18,7 @@
 #include "Selection/ToolSelectionUtil.h"
 #include "Operations/ExtrudeMesh.h"
 #include "DynamicMesh/MeshNormals.h"
+#include "DynamicMesh/MeshTangents.h"
 #include "MeshBoundaryLoops.h"
 #include "ToolDataVisualizer.h"
 
@@ -1012,6 +1013,7 @@ void UDrawPolyPathTool::EmitNewObject()
 	FVector3d Center = PathMesh.GetBounds().Center();
 	MeshTransform.Origin = MeshTransform.ToPlane(Center, 2);
 	MeshTransforms::WorldToFrameCoords(PathMesh, MeshTransform);
+	UE::Geometry::FMeshTangentsf::ComputeDefaultOverlayTangents(PathMesh);
 
 	GetToolManager()->BeginUndoTransaction(LOCTEXT("CreatePolyPathTransactionName", "Create PolyPath"));
 
diff --git a/Engine/Plugins/Experimental/MeshModelingToolsetExp/Source/MeshModelingToolsExp/Private/RevolveBoundaryTool.cpp b/Engine/Plugins/Experimental/MeshModelingToolsetExp/Source/MeshModelingToolsExp/Private/RevolveBoundaryTool.cpp
--- a/Engine/Plugins/Experimental/MeshModelingToolsetExp/Source/MeshModelingToolsExp/Private/RevolveBoundaryTool.cpp
+++ b/Engine/Plugins/Experimental/MeshModelingToolsetExp/Source/MeshModelingToolsExp/Private/RevolveBoundaryTool.cpp
@@ -283,6 +283,7 @@ void URevolveBoundaryTool::OnShutdown(EToolShutdownType ShutdownType)
 	{
 		if (ShutdownType == EToolShutdownType::Accept)
 		{
+			Preview->PreviewMesh->CalculateTangents();
 			GenerateAsset(Preview->Shutdown());
 		}
 		else
diff --git a/Engine/Plugins/Runtime/MeshModelingToolset/Source/MeshModelingTools/Private/DrawAndRevolveTool.cpp b/Engine/Plugins/Runtime/MeshModelingToolset/Source/MeshModelingTools/Private/DrawAndRevolveTool.cpp
--- a/Engine/Plugins/Runtime/MeshModelingToolset/Source/MeshModelingTools/Private/DrawAndRevolveTool.cpp
+++ b/Engine/Plugins/Runtime/MeshModelingToolset/Source/MeshModelingTools/Private/DrawAndRevolveTool.cpp
@@ -224,6 +224,7 @@ void UDrawAndRevolveTool::Shutdown(EToolShutdownType ShutdownType)
 	{
 		if (ShutdownType == EToolShutdownType::Accept)
 		{
+			Preview->PreviewMesh->CalculateTangents(); // Copy tangents from the PreviewMesh
 			GenerateAsset(Preview->Shutdown());
 		}
 		else
diff --git a/Engine/Plugins/Runtime/MeshModelingToolset/Source/MeshModelingTools/Private/DrawPolygonTool.cpp b/Engine/Plugins/Runtime/MeshModelingToolset/Source/MeshModelingTools/Private/DrawPolygonTool.cpp
--- a/Engine/Plugins/Runtime/MeshModelingToolset/Source/MeshModelingTools/Private/DrawPolygonTool.cpp
+++ b/Engine/Plugins/Runtime/MeshModelingToolset/Source/MeshModelingTools/Private/DrawPolygonTool.cpp
@@ -24,6 +24,7 @@
 #include "SceneQueries/SceneSnappingManager.h"
 #include "ConstrainedDelaunay2.h"
 #include "Arrangement2d.h"
+#include "DynamicMesh/MeshTangents.h"
 
 #include "DynamicMeshEditor.h"
 
@@ -1116,6 +1117,7 @@ void UDrawPolygonTool::EmitCurrentPolygon()
 		ResetPolygon();
 		return;
 	}
+	UE::Geometry::FMeshTangentsf::ComputeDefaultOverlayTangents(Mesh);
 
 	GetToolManager()->BeginUndoTransaction(LOCTEXT("CreatePolygon", "Create Polygon"));
 
diff --git a/Engine/Source/Runtime/GeometryCore/Public/DynamicMesh/MeshTangents.h b/Engine/Source/Runtime/GeometryCore/Public/DynamicMesh/MeshTangents.h
--- a/Engine/Source/Runtime/GeometryCore/Public/DynamicMesh/MeshTangents.h
+++ b/Engine/Source/Runtime/GeometryCore/Public/DynamicMesh/MeshTangents.h
@@ -250,6 +250,32 @@ public:
 	 */
 	bool CopyToOverlays(FDynamicMesh3& MeshToSet) const;
 
+
+
+	/**
+	 * Convenience function to compute tangents from the mesh attribute set's primary UVs and normals
+	 * @return true if tangents were computed, false otherwise
+	 */
+	static bool ComputeDefaultOverlayTangents(FDynamicMesh3& Mesh)
+	{
+		const FDynamicMeshAttributeSet* Attributes = Mesh.Attributes();
+		if (!Attributes)
+		{
+			return false;
+		}
+		const FDynamicMeshUVOverlay* UVOverlay = Attributes->PrimaryUV();
+		const FDynamicMeshNormalOverlay* NormalOverlay = Attributes->PrimaryNormals();
+		if (UVOverlay && NormalOverlay)
+		{
+			TMeshTangents Tangents;
+			Tangents.SetMesh(&Mesh);
+			Tangents.ComputeTriVertexTangents(NormalOverlay, UVOverlay, FComputeTangentsOptions());
+			Mesh.Attributes()->EnableTangents();
+			return Tangents.CopyToOverlays(Mesh);
+		}
+		return false;
+	}
+
 protected:
 	/**
 	 * Set the size of the Tangents array to Count, and optionally clear all values to (0,0,0)
