Fix FTransformSRT3d::InverseTransformNormal applying the scaling and rotation in the wrong order

diff --git a/Engine/Source/Runtime/GeometryCore/Public/TransformTypes.h b/Engine/Source/Runtime/GeometryCore/Public/TransformTypes.h
--- a/Engine/Source/Runtime/GeometryCore/Public/TransformTypes.h
+++ b/Engine/Source/Runtime/GeometryCore/Public/TransformTypes.h
@@ -271,7 +271,7 @@ public:
 	 */
 	TVector<RealType> InverseTransformNormal(const UE::Math::TVector<RealType>& Normal) const
 	{
-		return InverseTransformVectorNoScale( Normalized(Scale3D*Normal) );
+		return Normalized( Scale3D * InverseTransformVectorNoScale(Normal) );
 	}
 
 
