Fix AABB -> FBox/FBox2D casts for LWC (don't cast to float)

diff --git a/Engine/Source/Runtime/GeometryCore/Public/BoxTypes.h b/Engine/Source/Runtime/GeometryCore/Public/BoxTypes.h
--- a/Engine/Source/Runtime/GeometryCore/Public/BoxTypes.h
+++ b/Engine/Source/Runtime/GeometryCore/Public/BoxTypes.h
@@ -331,9 +331,7 @@ struct TAxisAlignedBox3
 
 	explicit operator FBox() const
 	{
-		FVector MinV((float)Min.X, (float)Min.Y, (float)Min.Z);
-		FVector MaxV((float)Max.X, (float)Max.Y, (float)Max.Z);
-		return FBox(MinV, MaxV);
+		return FBox((FVector)Min, (FVector)Max);
 	}
 	TAxisAlignedBox3(const FBox& Box)
 	{
@@ -646,9 +644,7 @@ struct TAxisAlignedBox2
 
 	explicit operator FBox2D() const
 	{
-		FVector2D MinV((float)Min.X, (float)Min.Y);
-		FVector2D MaxV((float)Max.X, (float)Max.Y);
-		return FBox2D(MinV, MaxV);
+		return FBox2D((FVector2D)Min, (FVector2D)Max);
 	}
 	TAxisAlignedBox2(const FBox2D& Box)
 	{
