Extend FSweepGeneratorBase and FGeneralizedCylinderGenerator to support open 2D profile curves, and custom UV for the profile and path.

#rnx

diff --git a/Engine/Source/Runtime/GeometryCore/Public/Generators/SweepGenerator.h b/Engine/Source/Runtime/GeometryCore/Public/Generators/SweepGenerator.h
--- a/Engine/Source/Runtime/GeometryCore/Public/Generators/SweepGenerator.h
+++ b/Engine/Source/Runtime/GeometryCore/Public/Generators/SweepGenerator.h
@@ -55,6 +55,9 @@ public:
 	/** If true, each quad gets a separate polygroup */
 	bool bPolygroupPerQuad = false;
 
+	/** If true, the last point of the profile curve is considered to be connected to the first. */
+	bool bProfileCurveIsClosed = true;
+
 protected:
 	int32 CapVertStart[2], CapNormalStart[2], CapUVStart[2], CapTriangleStart[2], CapPolygonStart[2];
 
@@ -71,49 +74,108 @@ protected:
 							   int32 NumCrossSections,
 							   bool bLoop,
 							   const ECapType Caps[2],
-							   FVector2f SectionsUVScale, FVector2f CapUVScale, FVector2f CapUVOffset)
+							   FVector2f SectionsUVScale, 
+							   FVector2f CapUVScale, 
+							   FVector2f CapUVOffset,
+							   const TArrayView<const float>& CustomCrossSectionTexCoord = {}, // if specified and valid, we use custom UVs instead of automatically generating them
+							   const TArrayView<const float>& CustomPathTexCoord = {})
 	{
 		// per cross section
-		int32 XVerts = CrossSection.VertexCount();
-		int32 XNormals = XVerts + NormalSections.Num();
-		int32 XUVs = XVerts + UVSections.Num() + 1;
+		const int32 XVerts = CrossSection.VertexCount();
+		const int32 XSegments = bProfileCurveIsClosed ? XVerts : XVerts - 1;
+		const int32 XNormals = XVerts + NormalSections.Num();
+		const int32 XUVs = XSegments + UVSections.Num() + 1;
+
+		float TotalPerimeter = 0.0f, TotalPathLength = 0.0f;
+		TArray<float> CrossSectionTexCoord, PathTexCoord;
 
-		double TotalPerimeter = 0, TotalPathLength = 0;
-		TArray<double> CrossSectionPercentages, PathPercentages;
-		if (bEvenlySpaceUVs)
+		// Compute texture coordinates along the cross section (U coordinates)
+		if (bProfileCurveIsClosed && CustomCrossSectionTexCoord.Num() >= XVerts + 1)
+		{
+			CrossSectionTexCoord = CustomCrossSectionTexCoord;
+		}
+		else if (bProfileCurveIsClosed && CustomCrossSectionTexCoord.Num() == XVerts)
+		{
+			CrossSectionTexCoord = CustomCrossSectionTexCoord;
+
+			// If the cross section curve is closed and we are missing texture coordinate for the 
+			// last element we wrap araound and use the coordinate of the first element
+			CrossSectionTexCoord.Add(CrossSectionTexCoord[0]);
+		}
+		else if (bProfileCurveIsClosed == false && CustomCrossSectionTexCoord.Num() >= XVerts) 
 		{
-			CrossSectionPercentages.Add(0);
-			PathPercentages.Add(0);
-			for (int XIdx = 0, XNum = CrossSection.VertexCount(); XIdx < XNum; XIdx++)
+			CrossSectionTexCoord = CustomCrossSectionTexCoord;
+		}
+		else if (bEvenlySpaceUVs)
+		{
+			CrossSectionTexCoord.Add(0.0f);
+			for (int Idx = 0; Idx < XSegments; Idx++)
 			{
-				double SegLen = Distance(CrossSection[XIdx], CrossSection[(XIdx + 1) % XNum]);
+				float SegLen = float(Distance(CrossSection[Idx], CrossSection[(Idx + 1) % XSegments]));
 				TotalPerimeter += SegLen;
-				CrossSectionPercentages.Add(TotalPerimeter);
+				CrossSectionTexCoord.Add(TotalPerimeter);
+			}
+			TotalPerimeter = FMath::Max(TotalPerimeter, FMathf::ZeroTolerance);
+			for (int Idx = 0; Idx < CrossSectionTexCoord.Num(); Idx++)
+			{
+				CrossSectionTexCoord[Idx] /= TotalPerimeter;
+				CrossSectionTexCoord[Idx] = 1.0f - CrossSectionTexCoord[Idx];
 			}
-			TotalPerimeter = FMath::Max(TotalPerimeter, FMathd::ZeroTolerance);
-			for (int Idx = 1; Idx < CrossSectionPercentages.Num(); Idx++)
+		}
+		else 
+		{	
+			for (int Idx = 0; Idx < XSegments; Idx++)
 			{
-				CrossSectionPercentages[Idx] /= TotalPerimeter;
+				float U = float(Idx) / float(XSegments);
+				CrossSectionTexCoord.Add(1.0f - U);
 			}
+			CrossSectionTexCoord.Add(0.0f);
+		}
+		
+		// Compute texture coordinates along the path (V coordinates)
+		if (CustomPathTexCoord.Num() >= NumCrossSections)
+		{
+			PathTexCoord = CustomPathTexCoord;
+		}
+		else if (bLoop && CustomPathTexCoord.Num() == NumCrossSections - 1)
+		{
+			PathTexCoord = CustomPathTexCoord;
+			
+			// If the path curve is closed and we are missing texture coordinate for the 
+			// last element we wrap araound and use the coordinate of the first element
+			PathTexCoord.Add(PathTexCoord[0]);
+		}
+		else if (bEvenlySpaceUVs) 
+		{
+			PathTexCoord.Add(0.0f);
 			int NumPathSegs = bLoop ? Path.Num() : Path.Num() - 1;
-			for (int PIdx = 0; PIdx < NumPathSegs; PIdx++)
+			for (int Idx = 0; Idx < NumPathSegs; Idx++)
 			{
-				double SegLen = Distance(Path[PIdx], Path[(PIdx + 1) % Path.Num()]);
+				float SegLen = float(Distance(Path[Idx], Path[(Idx + 1) % Path.Num()]));
 				TotalPathLength += SegLen;
-				PathPercentages.Add(TotalPathLength);
+				PathTexCoord.Add(TotalPathLength);
 			}
-			TotalPathLength = FMath::Max(TotalPathLength, FMathd::ZeroTolerance);
-			for (int Idx = 1; Idx < PathPercentages.Num(); Idx++)
+			TotalPathLength = FMath::Max(TotalPathLength, FMathf::ZeroTolerance);
+			for (int Idx = 0; Idx < PathTexCoord.Num(); Idx++)
 			{
-				PathPercentages[Idx] /= TotalPathLength;
+				PathTexCoord[Idx] /= TotalPathLength;
+				PathTexCoord[Idx] = 1.0f - PathTexCoord[Idx];
 			}
 		}
-
+		else 
+		{
+			for (int32 Idx = 0; Idx < NumCrossSections; Idx++)
+			{
+				float V = float(Idx) / float(NumCrossSections - 1);
+				PathTexCoord.Add(1.0f - V);
+			}
+		}
+		
 		int32 NumVerts = XVerts * NumCrossSections - (bLoop ? XVerts : 0);
 		int32 NumNormals = NumCrossSections > 1 ? (XNormals * NumCrossSections - (bLoop ? XNormals : 0)) : 0;
 		NumNormals += XNormals * SharpNormalsAlongLength.Num();
 		int32 NumUVs = NumCrossSections > 1 ? XUVs * NumCrossSections : 0;
-		int32 NumPolygons = (NumCrossSections - 1) * XVerts;
+		int32 NumPolygons = (NumCrossSections - 1) * XSegments;
 		int32 NumTriangles = NumPolygons * 2;
 
 		TArray<FIndex3i> OutTriangles;
@@ -140,7 +202,7 @@ protected:
 				}
 				else if (Caps[CapIdx] == ECapType::FlatMidpointFan)
 				{
-					NumTriangles += XVerts;
+					NumTriangles += XSegments;
 					NumPolygons++;
 					NumUVs += XVerts + 1;
 					NumNormals += XVerts + 1;
@@ -206,7 +268,7 @@ protected:
 					int32 CapVertStartIdx = CapVertStart[CapIdx];
 					int32 TriIdx = CapTriangleStart[CapIdx];
 					int32 PolyIdx = CapPolygonStart[CapIdx];
-					for (int32 VertIdx = 0; VertIdx < XVerts; VertIdx++)
+					for (int32 VertIdx = 0; VertIdx < XSegments; VertIdx++)
 					{
 						bool Flipped = CapIdx == 0;
 						SetTriangle(TriIdx,
@@ -260,13 +322,12 @@ protected:
 
 			int32 NumSections = UVSections.Num();
 			int32 NextDupVertIdx = UVSection < NumSections ? UVSections[UVSection] : -1;
-			for (int32 VertSubIdx = 0; VertSubIdx < XVerts; UVSubIdx++)
+			for (int32 VertSubIdx = 0; VertSubIdx < XSegments; UVSubIdx++)
 			{
-				float UVX = bEvenlySpaceUVs ? float(CrossSectionPercentages[VertSubIdx]) : float(VertSubIdx) / float(XVerts);
 				for (int32 XIdx = 0; XIdx < NumCrossSections; XIdx++)
 				{
-					float UVY = bEvenlySpaceUVs ? float(PathPercentages[XIdx]) : float(XIdx) / float(NumCrossSections - 1);
-					SetUV(XIdx * XUVs + UVSubIdx, FVector2f(1-UVX, 1-UVY) * SectionsUVScale, (XIdx % CrossSectionsMod) * XVerts + VertSubIdx);
+					FVector2f UV = FVector2f(CrossSectionTexCoord[VertSubIdx], PathTexCoord[XIdx]);
+					SetUV(XIdx * XUVs + UVSubIdx, UV * SectionsUVScale, (XIdx % CrossSectionsMod) * XVerts + VertSubIdx);
 				}
 
 				if (VertSubIdx == NextDupVertIdx)
@@ -278,12 +339,12 @@ protected:
 					for (int32 XIdx = 0; XIdx + 1 < NumCrossSections; XIdx++)
 					{
 						SetTriangleUVs(
-							XVerts * 2 * XIdx + 2 * VertSubIdx,
+							XSegments * 2 * XIdx + 2 * VertSubIdx,
 							XIdx * XUVs + UVSubIdx,
 							XIdx * XUVs + UVSubIdx + 1,
 							(XIdx + 1) * XUVs + UVSubIdx, true);
 						SetTriangleUVs(
-							XVerts * 2 * XIdx + 2 * VertSubIdx + 1,
+							XSegments * 2 * XIdx + 2 * VertSubIdx + 1,
 							(XIdx + 1) * XUVs + UVSubIdx + 1,
 							(XIdx + 1) * XUVs + UVSubIdx,
 							XIdx * XUVs + UVSubIdx + 1, true);
@@ -293,12 +354,11 @@ protected:
 			}
 			{
 				// final UV
-				float UVX = 1.0f;
-				int32 VertSubIdx = 0;
+				int32 VertSubIdx = bProfileCurveIsClosed ? 0 : XSegments;
 				for (int32 XIdx = 0; XIdx < NumCrossSections; XIdx++)
 				{
-					float UVY = bEvenlySpaceUVs ? float(PathPercentages[XIdx]) : float(XIdx) / float(NumCrossSections - 1);
-					SetUV(XIdx * XUVs + UVSubIdx, FVector2f(1-UVX, 1-UVY) * SectionsUVScale, (XIdx % CrossSectionsMod) * XVerts + VertSubIdx);
+					FVector2f UV = FVector2f(CrossSectionTexCoord.Last(), PathTexCoord[XIdx]);
+					SetUV(XIdx * XUVs + UVSubIdx, UV * SectionsUVScale, (XIdx % CrossSectionsMod) * XVerts + VertSubIdx);
 				}
 			}
 			NumSections = NormalSections.Num();
@@ -328,40 +388,43 @@ protected:
 				}
 				else
 				{
-					int32 WrappedNextNormalSubIdx = (NormalSubIdx + 1) % XNormals;
-					int32 WrappedNextVertexSubIdx = (VertSubIdx + 1) % XVerts;
-					SharpNormalIdx = 0;
-					for (int32 XIdx = 0, NXIdx = 0; XIdx + 1 < NumCrossSections; XIdx++, NXIdx++)
+					if (VertSubIdx < XSegments) // if bProfileCurveIsClosed == false skip the last triangle strip generation
 					{
-						int32 T0Idx = XVerts * 2 * XIdx + 2 * VertSubIdx;
-						int32 T1Idx = T0Idx + 1;
-						int32 PIdx = XVerts * XIdx + VertSubIdx;
-						int32 NextXIdx = (XIdx + 1) % CrossSectionsMod;
-						int32 NextNXIdx = (NXIdx + 1) % NormalCrossSectionsMod;
-						SetTrianglePolygon(T0Idx, (bPolygroupPerQuad) ? PIdx : (CurFaceGroupIndex+XIdx) );
-						SetTrianglePolygon(T1Idx, (bPolygroupPerQuad) ? PIdx : (CurFaceGroupIndex+XIdx) );
-						SetTriangle(T0Idx,
-									XIdx * XVerts + VertSubIdx,
-									XIdx * XVerts + WrappedNextVertexSubIdx,
-									NextXIdx * XVerts + VertSubIdx, true);
-						SetTriangle(T1Idx,
-									NextXIdx * XVerts + WrappedNextVertexSubIdx,
-									NextXIdx * XVerts + VertSubIdx,
-									XIdx * XVerts + WrappedNextVertexSubIdx, true);
-						SetTriangleNormals(
-							T0Idx,
-							NXIdx * XNormals + NormalSubIdx,
-							NXIdx * XNormals + WrappedNextNormalSubIdx,
-							NextNXIdx * XNormals + NormalSubIdx, true);
-						SetTriangleNormals(
-							T1Idx,
-							NextNXIdx * XNormals + WrappedNextNormalSubIdx,
-							NextNXIdx * XNormals + NormalSubIdx,
-							NXIdx * XNormals + WrappedNextNormalSubIdx, true);
-						if (SharpNormalIdx < SharpNormalsAlongLength.Num() && XIdx+1 == SharpNormalsAlongLength[SharpNormalIdx])
+						int32 WrappedNextNormalSubIdx = (NormalSubIdx + 1) % XNormals;
+						int32 WrappedNextVertexSubIdx = (VertSubIdx + 1) % XVerts;
+						SharpNormalIdx = 0;
+						for (int32 XIdx = 0, NXIdx = 0; XIdx + 1 < NumCrossSections; XIdx++, NXIdx++)
 						{
-							NXIdx++;
-							SharpNormalIdx++;
+							int32 T0Idx = XSegments * 2 * XIdx + 2 * VertSubIdx;
+							int32 T1Idx = T0Idx + 1;
+							int32 PIdx = XSegments * XIdx + VertSubIdx;
+							int32 NextXIdx = (XIdx + 1) % CrossSectionsMod;
+							int32 NextNXIdx = (NXIdx + 1) % NormalCrossSectionsMod;
+							SetTrianglePolygon(T0Idx, (bPolygroupPerQuad) ? PIdx : (CurFaceGroupIndex + XIdx));
+							SetTrianglePolygon(T1Idx, (bPolygroupPerQuad) ? PIdx : (CurFaceGroupIndex + XIdx));
+							SetTriangle(T0Idx,
+								XIdx * XVerts + VertSubIdx,
+								XIdx * XVerts + WrappedNextVertexSubIdx,
+								NextXIdx * XVerts + VertSubIdx, true);
+							SetTriangle(T1Idx,
+								NextXIdx * XVerts + WrappedNextVertexSubIdx,
+								NextXIdx * XVerts + VertSubIdx,
+								XIdx * XVerts + WrappedNextVertexSubIdx, true);
+							SetTriangleNormals(
+								T0Idx,
+								NXIdx * XNormals + NormalSubIdx,
+								NXIdx * XNormals + WrappedNextNormalSubIdx,
+								NextNXIdx * XNormals + NormalSubIdx, true);
+							SetTriangleNormals(
+								T1Idx,
+								NextNXIdx * XNormals + WrappedNextNormalSubIdx,
+								NextNXIdx * XNormals + NormalSubIdx,
+								NXIdx * XNormals + WrappedNextNormalSubIdx, true);
+							if (SharpNormalIdx < SharpNormalsAlongLength.Num() && XIdx + 1 == SharpNormalsAlongLength[SharpNormalIdx])
+							{
+								NXIdx++;
+								SharpNormalIdx++;
+							}
 						}
 					}
 					VertSubIdx++;
@@ -685,6 +748,32 @@ public:
 	// Only relevant if bUVScaleRelativeWorld is true (see that description)
 	float UnitUVInWorldCoordinates = 100;
 
+    // Optional custom UV values:
+	
+    // -If FSweepGeneratorBase::bProfileCurveIsClosed == true and CrossSectionTexCoord.Num() >= CrossSection.VertexCount() + 1 
+	// then the first CrossSection.VertexCount() + 1 values will be used as U coordinates.
+    //
+    // -If FSweepGeneratorBase::bProfileCurveIsClosed == true and CrossSectionTexCoord.Num() == CrossSection.VertexCount() 
+	// then the CrossSectionTexCoord[0] will be used as the value for the last element.
+    //
+    // -If FSweepGeneratorBase::bProfileCurveIsClosed == false and CrossSectionTexCoord.Num() >= CrossSection.VertexCount() 
+	// then the first CrossSection.VertexCount() values will be used as U coordinates.
+    // 
+    // -Otherwise, the U coordinates will be automatically generated.
+    TArray<float> CrossSectionTexCoord;
+    
+    // -If bLoop == true and PathTexCoord.Num() >= Path.Num() + 1 then the first Path.Num() + 1 values will be used as
+    //  V coordinates.
+    //
+    // -If bLoop == true and PathTexCoord.Num() == Path.Num() then the PathTexCoord[0] will be used as the value for the 
+    //  last element.
+    //
+    // -If bLoop == false and PathTexCoord.Num() >= Path.Num() then the first Path.Num() values will be used as V 
+    // coordinates.
+	//
+    // -Otherwise, the V coordinates will be automatically generated.
+    TArray<float> PathTexCoord;
+
 public:
 	/** Generate the mesh */
 	virtual FMeshShapeGenerator& Generate() override
@@ -714,7 +803,7 @@ public:
 			SectionScale.Y = float( TotalPathArcLength / UnitUVInWorldCoordinates);
 			CapScale.X = CapScale.Y = 1.0f / UnitUVInWorldCoordinates;
 		}
-		ConstructMeshTopology(CrossSection, {}, {}, {}, true, Path, PathNum + (bLoop ? 1 : 0), bLoop, Caps, SectionScale, CapScale, Bounds.Center());
+		ConstructMeshTopology(CrossSection, {}, {}, {}, true, Path, PathNum + (bLoop ? 1 : 0), bLoop, Caps, SectionScale, CapScale, Bounds.Center(), CrossSectionTexCoord, PathTexCoord);
 
 		int XNum = CrossSection.VertexCount();
 		TArray<FVector2d> XNormals; XNormals.SetNum(XNum);
