FindNearestCurveToRay: when processing segments in a curve, balance the distance-to-ray-origin against the ray-segment distance to find the optimal segment

diff --git a/Engine/Source/Runtime/GeometryCore/Private/Spatial/GeometrySet3.cpp b/Engine/Source/Runtime/GeometryCore/Private/Spatial/GeometrySet3.cpp
--- a/Engine/Source/Runtime/GeometryCore/Private/Spatial/GeometrySet3.cpp
+++ b/Engine/Source/Runtime/GeometryCore/Private/Spatial/GeometrySet3.cpp
@@ -186,7 +186,7 @@ bool FGeometrySet3::FindNearestCurveToRay(const FRay3d& Ray, FNearest& ResultOut
 	int NearestIndex = -1;
 	int NearestSegmentIdx = -1;
 	double NearestSegmentParam = 0;
-	double NearestMetric = TNumericLimits<double>::Max();
+	double NearestCurveMetric = TNumericLimits<double>::Max();
 
 	FCriticalSection Critical;
 
@@ -197,44 +197,47 @@ bool FGeometrySet3::FindNearestCurveToRay(const FRay3d& Ray, FNearest& ResultOut
 		const FPolyline3d& Polyline = Curve.Geometry;
 		int NumSegments = Polyline.SegmentCount();
 
-		double CurveMinRayParamT = TNumericLimits<double>::Max();
+		double NearestSegmentMetric = TNumericLimits<double>::Max();
+		double CurveNearestSegmentRayParamT = TNumericLimits<double>::Max();
 		int CurveNearestSegmentIdx = -1;
 		double CurveNearestSegmentParam = 0;
-		double CurveNearestDist = TNumericLimits<double>::Max();
 		FVector3d CurveNearestPosition;
+
+
 		for (int si = 0; si < NumSegments; ++si)
 		{
 			double SegRayParam; double SegSegParam;
-			double SegDistSqr = FDistRay3Segment3d::SquaredDistance(Ray, Polyline.GetSegment(si), SegRayParam, SegSegParam);
-			if (SegRayParam < CurveMinRayParamT)
+			FDistRay3Segment3d::SquaredDistance(Ray, Polyline.GetSegment(si), SegRayParam, SegSegParam);
+
+			FVector3d RayPosition = Ray.PointAt(SegRayParam);
+			FVector3d CurvePosition = Polyline.GetSegmentPoint(si, SegSegParam);
+			if (PointWithinToleranceTest(RayPosition, CurvePosition))
 			{
-				FVector3d RayPosition = Ray.PointAt(SegRayParam);
-				FVector3d CurvePosition = Polyline.GetSegmentPoint(si, SegSegParam);
-				if (PointWithinToleranceTest(RayPosition, CurvePosition))
+				// if multiple segments have a point within the tolerance, try to balance preference for closer-to-ray-origin and closer-to-segment
+				const double SegmentDistanceMetric = FMathd::Abs(VectorUtil::Area(Ray.Origin, CurvePosition, Ray.PointAt(SegRayParam)));
+				if (SegmentDistanceMetric < NearestSegmentMetric)
 				{
-					CurveMinRayParamT = SegRayParam;
+					CurveNearestSegmentRayParamT = SegRayParam;
 					CurveNearestSegmentIdx = si;
 					CurveNearestSegmentParam = SegSegParam;
-					CurveNearestDist = FMathd::Sqrt(SegDistSqr);
 					CurveNearestPosition = CurvePosition;
+					NearestSegmentMetric = SegmentDistanceMetric;
 				}
 			}
 		}
 
 		// if we found a possible point, lock outer structures and merge in
-		if (CurveMinRayParamT < TNumericLimits<double>::Max())
+		if (NearestSegmentMetric < TNumericLimits<double>::Max())
 		{
 			Critical.Lock();
-			// try to balance preference for closer-to-ray-origin and closer-to-segment
-			double DistanceMetric = FMathd::Abs(VectorUtil::Area(Ray.Origin, CurveNearestPosition, Ray.PointAt(CurveMinRayParamT)));
-			if (DistanceMetric < NearestMetric)
+			if (NearestSegmentMetric < NearestCurveMetric)
 			{
-				MinRayParamT = CurveMinRayParamT;
+				MinRayParamT = CurveNearestSegmentRayParamT;
 				NearestIndex = ci;
 				NearestID = Curve.ID;
 				NearestSegmentIdx = CurveNearestSegmentIdx;
 				NearestSegmentParam = CurveNearestSegmentParam;
-				NearestMetric = DistanceMetric;
+				NearestCurveMetric = NearestSegmentMetric;
 			}
 			Critical.Unlock();
 		}
