fix for edge collapse leaving isolated elements with parents, which breaks overlay validity

diff --git a/Engine/Source/Runtime/GeometryCore/Private/DynamicMesh/DynamicMeshOverlay.cpp b/Engine/Source/Runtime/GeometryCore/Private/DynamicMesh/DynamicMeshOverlay.cpp
--- a/Engine/Source/Runtime/GeometryCore/Private/DynamicMesh/DynamicMeshOverlay.cpp
+++ b/Engine/Source/Runtime/GeometryCore/Private/DynamicMesh/DynamicMeshOverlay.cpp
@@ -1270,19 +1270,13 @@ void TDynamicMeshOverlay<RealType, ElementSize>::OnCollapseEdge(const FDynamicMe
 	}
 
 	// clear the two triangles we removed
-	for (int j = 0; j < 3; ++j) 
+	if (bT0Set)
 	{
-		if (bT0Set)
-		{
-
-			ElementsRefCounts.Decrement(Triangle0[j]);
-			ElementTriangles[3 * tid_removed0 + j] = FDynamicMesh3::InvalidID;
-		}
-		if (collapseInfo.bIsBoundary == false && bT1Set) 
-		{
-			ElementsRefCounts.Decrement(Triangle1[j]);
-			ElementTriangles[3*tid_removed1 + j] = FDynamicMesh3::InvalidID;
-		}
+		UnsetTriangle(tid_removed0, true);
+	}
+	if (collapseInfo.bIsBoundary == false && bT1Set)
+	{
+		UnsetTriangle(tid_removed1, true);
 	}
 
 	// if the edge was split, but still shared one element, this should be protected against in the calling code
@@ -1291,17 +1285,17 @@ void TDynamicMeshOverlay<RealType, ElementSize>::OnCollapseEdge(const FDynamicMe
 		removed_elemid[1] = FDynamicMesh3::InvalidID;
 	}
 
-	// remove the elements associated with the removed vertex
+	// Note: the elements associated with the removed vertex should already have been removed when the triangles were un-set above
+#if UE_BUILD_DEBUG
 	for (int k = 0; k < 2; ++k)
 	{
 		if (removed_elemid[k] != FDynamicMesh3::InvalidID)
 	{
 			int rc = ElementsRefCounts.GetRefCount(removed_elemid[k]);
-			ensure(rc == 1);
-			ElementsRefCounts.Decrement(removed_elemid[k]);
-			ParentVertices[removed_elemid[k]] = FDynamicMesh3::InvalidID;
+			checkSlow(rc == 0);
 		}
 	}
+#endif
 
 }
 
@@ -1507,13 +1501,16 @@ bool TDynamicMeshOverlay<RealType, ElementSize>::CheckValidity(bool bAllowNonMan
 	for (int tid : ParentMesh->TriangleIndicesItr())
 	{
 		FIndex3i Tri = GetTriangle(tid);
+		int ValidCount = 0;
 		for (int j = 0; j < 3; ++j)
 		{
 			if (Tri[j] != FDynamicMesh3::InvalidID)
 			{
+				++ValidCount;
 				RealRefCounts[Tri[j]] += 1;
 			}
 		}
+		CheckOrFailF(ValidCount == 3 || ValidCount == 0); // element tri should be fully set or fully unset
 	}
 	// verify that refcount list counts are same as actual reference counts
 	for (int32 ElementID = 0; ElementID < RealRefCounts.Num(); ++ElementID)
