Add a GetFace method to TConvexHull3 and use it to get polygonal faces for convex hulls when building a convex hull for Chaos. This makes us less likely to add near-degenerate faces for which we cannot compute the normals.

diff --git a/Engine/Source/Runtime/Experimental/Chaos/Private/Chaos/CollisionConvexMesh.cpp b/Engine/Source/Runtime/Experimental/Chaos/Private/Chaos/CollisionConvexMesh.cpp
--- a/Engine/Source/Runtime/Experimental/Chaos/Private/Chaos/CollisionConvexMesh.cpp
+++ b/Engine/Source/Runtime/Experimental/Chaos/Private/Chaos/CollisionConvexMesh.cpp
@@ -137,37 +137,39 @@ namespace Chaos
 			if (UseConvexHull3(BuildMethod)) // Use the newer Geometry Tools code path for generating convex hulls.
 			{
 				UE::Geometry::TConvexHull3<FRealType> HullCompute;
+				HullCompute.bSaveTriangleNeighbors = true;
 				if (HullCompute.Solve<FVec3Type>(*VerticesToUse))
 				{
 					// Get Planes, FaceIndices, Vertices, and LocalBoundingBox
 					TMap<int32, int32> HullVertMap;
-					HullCompute.GetTriangles([&VerticesToUse, &HullVertMap, &OutPlanes, &OutFaceIndices, &OutVertices, &OutLocalBounds](UE::Geometry::FIndex3i Triangle)
+					HullCompute.GetFaces([&VerticesToUse, &HullVertMap, &OutPlanes, &OutFaceIndices, &OutVertices, &OutLocalBounds](TArray<int32>& FaceIndices, FVector3f FaceNormal)
 						{
-							for (int32 j = 0; j < 3; ++j)		// From FMeshConvexHull::Compute_FullMesh
+							TArray<int32>& OutFace = OutFaceIndices.Emplace_GetRef();
+							OutFace.SetNumUninitialized(FaceIndices.Num());
+
+							// Note winding is backwards from what Chaos expects, so this fills OutFace in reverse order
+							for (int32 j = 0, n = FaceIndices.Num(); j < n; ++j)
 							{
-								int32 Index = Triangle[j];
-								if (HullVertMap.Contains(Index) == false)
+								int32 Index = FaceIndices[j];
+								int32* MappedVert = HullVertMap.Find(Index);
+								if (!MappedVert)
 								{
 									const FVec3Type& OrigPos = (*VerticesToUse)[Index];
 									int32 NewVID = OutVertices.Num();
 									OutVertices.Add(OrigPos);
 									HullVertMap.Add(Index, NewVID);
-									Triangle[j] = NewVID;
+									OutFace[n - j - 1] = NewVID;
 								}
 								else
 								{
-									Triangle[j] = HullVertMap[Index];
+									OutFace[n - j - 1] = *MappedVert;
 								}
 							}
 
-							// Winding is backwards from what Chaos expects
-							OutFaceIndices.Add({ Triangle[0], Triangle[2], Triangle[1] });
-
-							UE::Geometry::THalfspace3<FRealType> Halfspace(OutVertices[Triangle[0]], OutVertices[Triangle[1]], OutVertices[Triangle[2]]);
-							const FVec3Type& Normal = Halfspace.Normal;
-							const FVec3Type& Point = OutVertices[Triangle[0]];
+							const FVec3Type& Normal = FaceNormal;
+							const FVec3Type& Point = OutVertices[OutFace[0]];
 							OutPlanes.Add(FPlaneType{ Point, Normal });
-						});
+						}, [&VerticesToUse](int32 Idx) { return (*VerticesToUse)[Idx]; });
 
 					OutLocalBounds = FAABB3Type(OutVertices[0], OutVertices[0]);
 					for (int32 VertexIndex = 1; VertexIndex < OutVertices.Num(); ++VertexIndex)
diff --git a/Engine/Source/Runtime/Experimental/Chaos/Private/Chaos/MassProperties.cpp b/Engine/Source/Runtime/Experimental/Chaos/Private/Chaos/MassProperties.cpp
--- a/Engine/Source/Runtime/Experimental/Chaos/Private/Chaos/MassProperties.cpp
+++ b/Engine/Source/Runtime/Experimental/Chaos/Private/Chaos/MassProperties.cpp
@@ -149,29 +149,31 @@ namespace Chaos
 		TVec Center = Vertices[Surfaces[0][0]];
 		for (const auto& Element : Surfaces)
 		{
-			// For now we only support triangular elements
-			check(Element.Num() == 3);
-
-			PMatrix<T,3,3> DeltaMatrix;
-			TVec PerElementSize;
-			for (int32 i = 0; i < Element.Num(); ++i)
+			// implicitly triangulate with a fan
+			for (int32 FanIdx = 0; FanIdx + 2 < Element.Num(); ++FanIdx)
 			{
-				TVec DeltaVector = Vertices[Element[i]] - Center;
-				DeltaMatrix.M[0][i] = DeltaVector[0];
-				DeltaMatrix.M[1][i] = DeltaVector[1];
-				DeltaMatrix.M[2][i] = DeltaVector[2];
-			}
-			
-			PerElementSize[0] = DeltaMatrix.M[0][0] + DeltaMatrix.M[0][1] + DeltaMatrix.M[0][2];
-			PerElementSize[1] = DeltaMatrix.M[1][0] + DeltaMatrix.M[1][1] + DeltaMatrix.M[1][2];
-			PerElementSize[2] = DeltaMatrix.M[2][0] + DeltaMatrix.M[2][1] + DeltaMatrix.M[2][2];
+				int32 FanSubInds[]{ 0, FanIdx + 1, FanIdx + 2 };
+				PMatrix<T, 3, 3> DeltaMatrix;
+				TVec PerElementSize;
+				for (int32 i = 0; i < 3; ++i)
+				{
+					TVec DeltaVector = Vertices[Element[FanSubInds[i]]] - Center;
+					DeltaMatrix.M[0][i] = DeltaVector[0];
+					DeltaMatrix.M[1][i] = DeltaVector[1];
+					DeltaMatrix.M[2][i] = DeltaVector[2];
+				}
 
-			T Det = DeltaMatrix.M[0][0] * (DeltaMatrix.M[1][1] * DeltaMatrix.M[2][2] - DeltaMatrix.M[1][2] * DeltaMatrix.M[2][1]) -
-				DeltaMatrix.M[0][1] * (DeltaMatrix.M[1][0] * DeltaMatrix.M[2][2] - DeltaMatrix.M[1][2] * DeltaMatrix.M[2][0]) +
-				DeltaMatrix.M[0][2] * (DeltaMatrix.M[1][0] * DeltaMatrix.M[2][1] - DeltaMatrix.M[1][1] * DeltaMatrix.M[2][0]);
-			
-			Volume += Det;
-			VolumeTimesSum += Det * PerElementSize;
+				PerElementSize[0] = DeltaMatrix.M[0][0] + DeltaMatrix.M[0][1] + DeltaMatrix.M[0][2];
+				PerElementSize[1] = DeltaMatrix.M[1][0] + DeltaMatrix.M[1][1] + DeltaMatrix.M[1][2];
+				PerElementSize[2] = DeltaMatrix.M[2][0] + DeltaMatrix.M[2][1] + DeltaMatrix.M[2][2];
+
+				T Det = DeltaMatrix.M[0][0] * (DeltaMatrix.M[1][1] * DeltaMatrix.M[2][2] - DeltaMatrix.M[1][2] * DeltaMatrix.M[2][1]) -
+					DeltaMatrix.M[0][1] * (DeltaMatrix.M[1][0] * DeltaMatrix.M[2][2] - DeltaMatrix.M[1][2] * DeltaMatrix.M[2][0]) +
+					DeltaMatrix.M[0][2] * (DeltaMatrix.M[1][0] * DeltaMatrix.M[2][1] - DeltaMatrix.M[1][1] * DeltaMatrix.M[2][0]);
+
+				Volume += Det;
+				VolumeTimesSum += Det * PerElementSize;
+			}
 		}
 		// @todo(mlentine): Should add suppoert for thin shell mass properties
 		if (Volume < UE_KINDA_SMALL_NUMBER)	//handle negative volume using fallback for now. Need to investigate cases where this happens
diff --git a/Engine/Source/Runtime/GeometryCore/Private/CompGeom/ConvexHull3.cpp b/Engine/Source/Runtime/GeometryCore/Private/CompGeom/ConvexHull3.cpp
--- a/Engine/Source/Runtime/GeometryCore/Private/CompGeom/ConvexHull3.cpp
+++ b/Engine/Source/Runtime/GeometryCore/Private/CompGeom/ConvexHull3.cpp
@@ -661,6 +661,145 @@ bool TConvexHull3<RealType>::Solve(int32 NumPoints, TFunctionRef<TVector<RealTyp
 	return true;
 }
 
+template<class RealType>
+void TConvexHull3<RealType>::GetFaces(TFunctionRef<void(TArray<int32>&, TVector<RealType>)> PolygonFunc, TFunctionRef<TVector<RealType>(int32)> GetPointFunc) const
+{
+	if (!ensureMsgf(bSaveTriangleNeighbors && HullNeighbors.Num() == Hull.Num(), TEXT("To extract faces, set bSaveTriangleNeighbors = true before calling Solve()")))
+	{
+		return;
+	}
+
+	TArray<int32> CurFaceVertIDs; // Used for face vertex IDs
+	TArray<int32> ToProcess; // Used as a stack for hull indices
+
+	// IDs to indicate which triangles can be grouped into the same convex polygon
+	TArray<int32> GroupIDs;
+	GroupIDs.Init(-1, Hull.Num());
+
+	int32 NextGroupID = 0;
+
+	// Greedily flood-fill to find coplanar triangles in each group
+	for (int32 TriIdx = 0; TriIdx < Hull.Num(); ++TriIdx)
+	{
+		if (GroupIDs[TriIdx] >= 0)
+		{
+			continue;
+		}
+
+		int32 CurGroupID = NextGroupID++;
+		GroupIDs[TriIdx] = CurGroupID;
+
+		ToProcess.Reset();
+
+		const FIndex3i& Tri = Hull[TriIdx];
+		const FIndex3i& NbrInds = HullNeighbors[TriIdx];
+		TVector<RealType> TriPosns[3]{ GetPointFunc(Tri.A), GetPointFunc(Tri.B), GetPointFunc(Tri.C) };
+		TVector<RealType> FaceNormal = VectorUtil::Normal(TriPosns[0], TriPosns[1], TriPosns[2]);
+		bool bValidNormal = FaceNormal.SquaredLength() > 0;
+		ToProcess.Add(NbrInds[0]);
+		ToProcess.Add(NbrInds[1]);
+		ToProcess.Add(NbrInds[2]);
+		while (ToProcess.Num() > 0)
+		{
+			int32 NbrTriIdx = ToProcess.Pop(false);
+			if (GroupIDs[NbrTriIdx] >= 0 || GroupIDs[NbrTriIdx] < -CurGroupID - 1)
+			{
+				continue;
+			}
+
+			const FIndex3i& NbrTri = Hull[NbrTriIdx];
+			int32 OtherV = -1;
+			for (int32 SubIdx = 0; SubIdx < 3; ++SubIdx)
+			{
+				int32 V = NbrTri[SubIdx];
+				if (V != Tri.A && V != Tri.B && V != Tri.C)
+				{
+					OtherV = V;
+				}
+			}
+			check(OtherV >= 0);
+			TVector<RealType> OtherPos = GetPointFunc(OtherV);
+			RealType O3DVal = ExactPredicates::Orient3<RealType>(TriPosns[0], TriPosns[1], TriPosns[2], OtherPos);
+			if (O3DVal == 0)
+			{
+				GroupIDs[NbrTriIdx] = CurGroupID;
+
+				if (!bValidNormal)
+				{
+					FaceNormal = VectorUtil::Normal(GetPointFunc(NbrTri.A), GetPointFunc(NbrTri.B), GetPointFunc(NbrTri.C));
+					bValidNormal = FaceNormal.SquaredLength() > 0;
+				}
+
+				const FIndex3i& NbrNbrInds = HullNeighbors[NbrTriIdx];
+				for (int32 SubIdx = 0; SubIdx < 3; ++SubIdx)
+				{
+					int32 NbrNbrTriIdx = NbrNbrInds[SubIdx];
+					if (GroupIDs[NbrNbrTriIdx] >= 0 || GroupIDs[NbrNbrTriIdx] < -CurGroupID - 1)
+					{
+						ToProcess.Add(NbrNbrTriIdx);
+					}
+				}
+			}
+			else
+			{
+				// Flag the unset triangles as already-visited, to avoid repeat tests
+				GroupIDs[NbrTriIdx] = -CurGroupID - 2;
+			}
+		}
+
+		CurFaceVertIDs.Reset();
+		WalkBorder(Hull, HullNeighbors, [&](int32 Idx) { return GroupIDs[Idx] == CurGroupID; }, TriIdx, CurFaceVertIDs);
+		PolygonFunc(CurFaceVertIDs, FaceNormal);
+	}
+}
+
+template<typename RealType>
+void TConvexHull3<RealType>::WalkBorder(const TArray<FIndex3i>& Triangles, const TArray<FIndex3i>& TriangleNeighbors, TFunctionRef<bool(int32)> InGroupFunc, int32 StartIdx, TArray<int32>& OutBorderVertexIndices)
+{
+	check(Triangles.Num() == TriangleNeighbors.Num());
+
+	struct FVisit
+	{
+		int32 Tri;
+		int8 Edge;
+	};
+
+	auto CrossEdge = [&Triangles, &TriangleNeighbors](int32 Tri, int32 EdgeNum) -> FVisit
+	{
+		int32 NbrTri = TriangleNeighbors[Tri][EdgeNum];
+		int32 SecondVert = Triangles[Tri][EdgeNum + 1 < 3 ? EdgeNum + 1 : 0];
+		int8 NbrEdge = (int8)Triangles[NbrTri].IndexOf(SecondVert);
+
+		return FVisit{ NbrTri, NbrEdge };
+	};
+
+	TArray<bool> Visited;
+	Visited.SetNumZeroed(Triangles.Num());
+	TArray<FVisit> VisitStack;
+	VisitStack.Reserve(Triangles.Num() / 3);
+	Visited[StartIdx] = true;
+	VisitStack.Add(CrossEdge(StartIdx, 2));
+	VisitStack.Add(CrossEdge(StartIdx, 1));
+	VisitStack.Add(CrossEdge(StartIdx, 0));
+
+	while (VisitStack.Num())
+	{
+		FVisit Visit = VisitStack.Pop(false);
+		if (!InGroupFunc(Visit.Tri))
+		{
+			OutBorderVertexIndices.Add(Triangles[Visit.Tri][Visit.Edge]);
+			continue;
+		}
+		if (Visited[Visit.Tri])
+		{
+			continue;
+		}
+		Visited[Visit.Tri] = true;
+		VisitStack.Add(CrossEdge(Visit.Tri, (Visit.Edge + 2) % 3));
+		VisitStack.Add(CrossEdge(Visit.Tri, (Visit.Edge + 1) % 3));
+	}
+}
+
 template struct TExtremePoints3<float>;
 template struct TExtremePoints3<double>;
 template class TConvexHull3<float>;
diff --git a/Engine/Source/Runtime/GeometryCore/Public/CompGeom/ConvexHull3.h b/Engine/Source/Runtime/GeometryCore/Public/CompGeom/ConvexHull3.h
--- a/Engine/Source/Runtime/GeometryCore/Public/CompGeom/ConvexHull3.h
+++ b/Engine/Source/Runtime/GeometryCore/Public/CompGeom/ConvexHull3.h
@@ -135,6 +135,14 @@ public:
 		return Hull;
 	}
 
+	/**
+	 * Get faces of the convex hull as convex polygons
+	 * @param PolygonFunc	Callback to be called for each polygon, with the array of vertex indices and the face normal
+	 * @param GetPointFunc	Function providing array-style access into points
+	 * TODO: Provide an optional Epsilon parameter to simplify the hull within some tolerance
+	 */
+	void GetFaces(TFunctionRef<void(TArray<int32>&, TVector<RealType>)> PolygonFunc, TFunctionRef<TVector<RealType>(int32)> GetPointFunc) const;
+
 	/**
 	 * Only valid if bSaveTriangleNeighbors was true when Solve() was called
 	 * @return Neighbors of each hull triangle, in edge order -- i.e., Nbr.A is the triangle across edge(Tri.A, Tri.B)
@@ -230,6 +238,10 @@ public:
 	/** Set this to be able to cancel running operation */
 	FProgressCancel* Progress = nullptr;
 
+	// Useful helper for walking the border of a convex region, or of a QHull horizon. Will also work for extracting a group boundary for any closed triangle mesh.
+	// Note: Assumes TriangleNeighbors are listed such that the first neighbor is the tri across from the (Tri.A, Tri.B) edge
+	static void WalkBorder(const TArray<FIndex3i>& Triangles, const TArray<FIndex3i>& TriangleNeighbors, TFunctionRef<bool(int32)> InGroupFunc, int32 StartIdx, TArray<int32>& OutBorderVertexIndices);
+
 protected:
 
 	int32 Dimension;
