fix broken tet mesh test by making tets consistently generate by default in the UE orientation (orient3d gives negative sign), and switching the Validate function back to expecting that orientation

diff --git a/Engine/Source/Runtime/GeometryCore/Private/CompGeom/Delaunay3.cpp b/Engine/Source/Runtime/GeometryCore/Private/CompGeom/Delaunay3.cpp
--- a/Engine/Source/Runtime/GeometryCore/Private/CompGeom/Delaunay3.cpp
+++ b/Engine/Source/Runtime/GeometryCore/Private/CompGeom/Delaunay3.cpp
@@ -582,24 +582,47 @@ bool FDelaunay3::Triangulate(TArrayView<const FVector3f> Vertices)
 	return bSuccess;
 }
 
-TArray<FIndex4i> FDelaunay3::GetTetrahedraAsFIndex4i() const
+TArray<FIndex4i> FDelaunay3::GetTetrahedraAsFIndex4i(bool bReverseOrientation) const
 {
+	TArray<FIndex4i> ToRet;
 	if (Connectivity.IsValid())
 	{
-		return Connectivity->GetTets();
+		Connectivity->EnumerateTrianglePerTet(
+			[&ToRet, bReverseOrientation](FIndex3i Tri, int32 Vert)
+			{
+				if (bReverseOrientation)
+				{
+					ToRet.Emplace(Tri.A, Tri.B, Tri.C, Vert);
+				}
+				else
+				{
+					// Note: UE expects negative-sign orientation, so B,C swapped is 'not reversed'
+					ToRet.Emplace(Tri.A, Tri.C, Tri.B, Vert);
+				}
+				return true;
+			}, true
+		);
 	}
-	return TArray<FIndex4i>();
+	return ToRet;
 }
 
-TArray<FIntVector4> FDelaunay3::GetTetrahedra() const
+TArray<FIntVector4> FDelaunay3::GetTetrahedra(bool bReverseOrientation) const
 {
 	TArray<FIntVector4> ToRet;
 	if (Connectivity.IsValid())
 	{
 		Connectivity->EnumerateTrianglePerTet(
-			[&ToRet](FIndex3i Tri, int32 Vert)
+			[&ToRet, bReverseOrientation](FIndex3i Tri, int32 Vert)
 			{
-				ToRet.Emplace(Tri.A, Tri.B, Tri.C, Vert);
+				if (bReverseOrientation)
+				{
+					ToRet.Emplace(Tri.A, Tri.B, Tri.C, Vert);
+				}
+				else
+				{
+					// Note: UE expects negative-sign orientation, so B,C swapped is 'not reversed'
+					ToRet.Emplace(Tri.A, Tri.C, Tri.B, Vert);
+				}
 				return true;
 			}, true
 		);
diff --git a/Engine/Source/Runtime/GeometryCore/Public/CompGeom/Delaunay3.h b/Engine/Source/Runtime/GeometryCore/Public/CompGeom/Delaunay3.h
--- a/Engine/Source/Runtime/GeometryCore/Public/CompGeom/Delaunay3.h
+++ b/Engine/Source/Runtime/GeometryCore/Public/CompGeom/Delaunay3.h
@@ -41,8 +41,8 @@ public:
 
 	// Get the result as an array of tetrahedra
 	// Note: This creates a new array each call, because the internal data structure does not have a tetrahedra array
-	TArray<FIndex4i> GetTetrahedraAsFIndex4i() const;
-	TArray<FIntVector4> GetTetrahedra() const;
+	TArray<FIndex4i> GetTetrahedraAsFIndex4i(bool bReverseOrientation = false) const;
+	TArray<FIntVector4> GetTetrahedra(bool bReverseOrientation = false) const;
 
 	// @return true if triangulation is Delaunay, useful for validating results (note: likely to be false if edges are constrained)
 	bool IsDelaunay(TArrayView<const FVector3f> Vertices) const;
