fix incorrect direct assignments to TFunctionRef

diff --git a/Engine/Source/Runtime/GeometryCore/Private/CompGeom/ConvexDecomposition3.cpp b/Engine/Source/Runtime/GeometryCore/Private/CompGeom/ConvexDecomposition3.cpp
--- a/Engine/Source/Runtime/GeometryCore/Private/CompGeom/ConvexDecomposition3.cpp
+++ b/Engine/Source/Runtime/GeometryCore/Private/CompGeom/ConvexDecomposition3.cpp
@@ -1199,9 +1199,13 @@ int32 FConvexDecomposition3::MergeBest(int32 TargetNumParts, double MaxErrorTole
 
 	TMap<int32, TUniquePtr<FConvexPart>> ProximityComputedParts;
 
-	TFunctionRef<bool(const FConvexPart&)> IsPartBelowSizeTolerance = MinThicknessTolerance > 0 ? 
-	(TFunctionRef<bool(const FConvexPart&)>)[MinThicknessTolerance](const FConvexPart& Part) -> bool
-	{ // Check the part vs the thickness tolerance
+	auto IsPartBelowSizeTolerance = [MinThicknessTolerance](const FConvexPart& Part) -> bool
+	{ 
+		if (MinThicknessTolerance <= 0) // tolerance is disabled
+		{
+			return false;
+		}
+		// Check the part vs the thickness tolerance
 		// First check the already-computed AABB
 		if (Part.Bounds.MinDim() < MinThicknessTolerance)
 		{
@@ -1236,10 +1240,6 @@ int32 FConvexDecomposition3::MergeBest(int32 TargetNumParts, double MaxErrorTole
 			SmallAxisInterval.Contain(V.Dot(SmallDir));
 		}
 		return SmallAxisInterval.Extent() < MinThicknessTolerance;
-	} : // No tolerance set: Always return false
-	(TFunctionRef<bool(const FConvexPart&)>)[](const FConvexPart& Part) -> bool
-	{
-		return false;
 	};
 
 	int32 MustMergeCount = 0;
diff --git a/Engine/Source/Runtime/GeometryCore/Public/CompGeom/FitOrientedBox3.h b/Engine/Source/Runtime/GeometryCore/Public/CompGeom/FitOrientedBox3.h
--- a/Engine/Source/Runtime/GeometryCore/Public/CompGeom/FitOrientedBox3.h
+++ b/Engine/Source/Runtime/GeometryCore/Public/CompGeom/FitOrientedBox3.h
@@ -47,11 +47,11 @@ namespace Geometry
 	template <typename RealType>
 	TOrientedBox3<RealType> FitOrientedBox3Points(TArrayView<const TVector<RealType>> Points, EBox3FitCriteria FitMethod = EBox3FitCriteria::Volume, RealType SameNormalTolerance = 0, FProgressCancel* ProgressCancel = nullptr)
 	{
-		TFunctionRef<TVector<RealType>(int32)> GetPtFn = [&Points](int32 Idx)
+		auto GetPtFn = [&Points](int32 Idx)
 		{
 			return Points[Idx];
 		};
-		return FitOrientedBox3Points(Points.Num(), GetPtFn, [](int32 Idx) {return true;}, FitMethod, SameNormalTolerance, ProgressCancel);
+		return FitOrientedBox3Points(Points.Num(), (TFunctionRef<TVector<RealType>(int32)>)GetPtFn, [](int32 Idx) {return true;}, FitMethod, SameNormalTolerance, ProgressCancel);
 	}
 
 
