GeometryCore: DynamicMeshAttributeSet now serializes bone weight layer FNames as strings. Conceivably this is a breaking change, however the previous code was completely broken and so no "in the wild" serialized dynamic meshes with skin weight layers should exist.

diff --git a/Engine/Source/Runtime/GeometryCore/Private/DynamicMesh/DynamicMeshAttributeSet.cpp b/Engine/Source/Runtime/GeometryCore/Private/DynamicMesh/DynamicMeshAttributeSet.cpp
--- a/Engine/Source/Runtime/GeometryCore/Private/DynamicMesh/DynamicMeshAttributeSet.cpp
+++ b/Engine/Source/Runtime/GeometryCore/Private/DynamicMesh/DynamicMeshAttributeSet.cpp
@@ -5,6 +5,7 @@
 #include "DynamicMesh/DynamicVertexSkinWeightsAttribute.h"
 #include "IndexTypes.h"
 #include "Async/Async.h"
+#include "Serialization/NameAsStringProxyArchive.h"
 
 
 using namespace UE::Geometry;
@@ -1229,7 +1230,8 @@ void FDynamicMeshAttributeSet::Serialize(FArchive& Ar, const FCompactMaps* Compa
 			for (int32 i = 0; i < NumSkinWeightAttributes; ++i)
 			{
 				FName Key;
-				Ar << Key;
+				FNameAsStringProxyArchive ProxyArchive(Ar);
+				ProxyArchive << Key;
 
 				TUniquePtr<FDynamicMeshVertexSkinWeightsAttribute>& Value = SkinWeightAttributes.Emplace(Key, nullptr);
 				bool IsValid;
@@ -1245,7 +1247,8 @@ void FDynamicMeshAttributeSet::Serialize(FArchive& Ar, const FCompactMaps* Compa
 		{
 			for (SkinWeightAttributesMap::TIterator It(SkinWeightAttributes); It; ++It)
 			{
-				Ar << It.Key();
+				FNameAsStringProxyArchive ProxyArchive(Ar);
+				ProxyArchive << It.Key();
 
 				const TUniquePtr<FDynamicMeshVertexSkinWeightsAttribute>& Value = It.Value();
 				bool bIsValid = Value.IsValid();
