Fixed compile error surfaced when orphaned headers got compiled (they were never included in their own module so compiler warning settings was never applied)
Changes include
* Removed #pragma once in cpp files
* Removed dll export on templated types
* Fixed code in headers that didn't compile (probably never included anywhere
* Added casting for code going from int to float. Header was never compiled in the owning module that had warnings as errors

diff --git a/Engine/Plugins/Experimental/Mutable/Source/MutableRuntime/Public/MuR/Serialisation.h b/Engine/Plugins/Experimental/Mutable/Source/MutableRuntime/Public/MuR/Serialisation.h
--- a/Engine/Plugins/Experimental/Mutable/Source/MutableRuntime/Public/MuR/Serialisation.h
+++ b/Engine/Plugins/Experimental/Mutable/Source/MutableRuntime/Public/MuR/Serialisation.h
@@ -16,7 +16,7 @@ namespace mu
     //! \brief
     //! \ingroup model
     template<class R>
-    class MUTABLERUNTIME_API ResourceProxy : public RefCounted
+    class ResourceProxy : public RefCounted
     {
     public:
 
diff --git a/Engine/Plugins/Runtime/GameplayAbilities/Source/GameplayAbilities/Private/AbilitySystemReplicationProxyInterface.cpp b/Engine/Plugins/Runtime/GameplayAbilities/Source/GameplayAbilities/Private/AbilitySystemReplicationProxyInterface.cpp
--- a/Engine/Plugins/Runtime/GameplayAbilities/Source/GameplayAbilities/Private/AbilitySystemReplicationProxyInterface.cpp
+++ b/Engine/Plugins/Runtime/GameplayAbilities/Source/GameplayAbilities/Private/AbilitySystemReplicationProxyInterface.cpp
@@ -1,7 +1,5 @@
 // Copyright Epic Games, Inc. All Rights Reserved.
 
-#pragma once
-
 #include "AbilitySystemReplicationProxyInterface.h"
 
 #include "GameplayPrediction.h"
diff --git a/Engine/Plugins/Runtime/GeometryCache/Source/GeometryCache/Classes/GeometryCacheStreamingManagerBase.h b/Engine/Plugins/Runtime/GeometryCache/Source/GeometryCache/Classes/GeometryCacheStreamingManagerBase.h
--- a/Engine/Plugins/Runtime/GeometryCache/Source/GeometryCache/Classes/GeometryCacheStreamingManagerBase.h
+++ b/Engine/Plugins/Runtime/GeometryCache/Source/GeometryCache/Classes/GeometryCacheStreamingManagerBase.h
@@ -1,7 +1,10 @@
 // Copyright Epic Games, Inc. All Rights Reserved.
 #pragma once
 
-#include "CoreMinimal.h"
+#include "ContentStreaming.h"
+
+class UGeometryCacheComponent;
+class UGeometryCacheTrackStreamable;
 
 /**
 * Note IStreamingManager is not really anything like an interface it contains code and members and whatnot.
@@ -12,7 +15,7 @@ struct IGeometryCacheStreamingManager : public IStreamingManager
 	IGeometryCacheStreamingManager() {}
 
 	/** Virtual destructor */
-	virtual ~IStreamingManager() {}
+	virtual ~IGeometryCacheStreamingManager() {}
 
 	/** Getter of the singleton */
 	GEOMETRYCACHE_API static struct IGeometryCacheStreamingManager& Get();
@@ -49,4 +52,4 @@ struct IGeometryCacheStreamingManager : public IStreamingManager
 	*/
 	virtual const uint8* GetLoadedChunk(const UGeometryCacheTrackStreamable* Track, uint32 ChunkIndex, uint32* OutChunkSize = NULL) const = 0;
 
-}
\ No newline at end of file
+};
\ No newline at end of file
diff --git a/Engine/Source/Runtime/Core/Public/Math/Halton.h b/Engine/Source/Runtime/Core/Public/Math/Halton.h
--- a/Engine/Source/Runtime/Core/Public/Math/Halton.h
+++ b/Engine/Source/Runtime/Core/Public/Math/Halton.h
@@ -9,11 +9,11 @@
 inline float Halton(int32 Index, int32 Base)
 {
 	float Result = 0.0f;
-	float InvBase = 1.0f / Base;
+	float InvBase = 1.0f / float(Base);
 	float Fraction = InvBase;
 	while (Index > 0)
 	{
-		Result += (Index % Base) * Fraction;
+		Result += float(Index % Base) * Fraction;
 		Index /= Base;
 		Fraction *= InvBase;
 	}
diff --git a/Engine/Source/Runtime/Core/Tests/Containers/PagedArrayTest.cpp b/Engine/Source/Runtime/Core/Tests/Containers/PagedArrayTest.cpp
--- a/Engine/Source/Runtime/Core/Tests/Containers/PagedArrayTest.cpp
+++ b/Engine/Source/Runtime/Core/Tests/Containers/PagedArrayTest.cpp
@@ -21,6 +21,8 @@ struct FSomeForwardDeclaredStruct;
 
 struct FDummyTestStruct
 {
+	FDummyTestStruct();
+	~FDummyTestStruct();
 	TPagedArray<FSomeForwardDeclaredStruct> Data;
 };
 
diff --git a/Engine/Source/Runtime/Engine/Private/Model.cpp b/Engine/Source/Runtime/Engine/Private/Model.cpp
--- a/Engine/Source/Runtime/Engine/Private/Model.cpp
+++ b/Engine/Source/Runtime/Engine/Private/Model.cpp
@@ -13,6 +13,7 @@
 #include "EngineUtils.h"
 #include "Engine/Polys.h"
 #include "DynamicMeshBuilder.h"
+#include "StaticLighting.h"
 
 float UModel::BSPTexelScale = 100.0f;
 
@@ -178,6 +179,8 @@ FArchive& operator<<(FArchive& Ar, FDepecatedModelVertex& V)
 	return Ar;
 }
 
+FNodeGroup::~FNodeGroup() = default;
+
 /*---------------------------------------------------------------------------------------
 	UModel object implementation.
 ---------------------------------------------------------------------------------------*/
diff --git a/Engine/Source/Runtime/Engine/Private/TextureRenderTarget2D.cpp b/Engine/Source/Runtime/Engine/Private/TextureRenderTarget2D.cpp
--- a/Engine/Source/Runtime/Engine/Private/TextureRenderTarget2D.cpp
+++ b/Engine/Source/Runtime/Engine/Private/TextureRenderTarget2D.cpp
@@ -559,6 +559,10 @@ FTextureRenderTarget2DResource::FTextureRenderTarget2DResource(const class UText
 	
 }
 
+PRAGMA_DISABLE_DEPRECATION_WARNINGS
+FTextureRenderTarget2DResource::~FTextureRenderTarget2DResource() = default;
+PRAGMA_ENABLE_DEPRECATION_WARNINGS
+
 /**
  * Clamp size of the render target resource to max values
  *
diff --git a/Engine/Source/Runtime/Engine/Public/Model.h b/Engine/Source/Runtime/Engine/Public/Model.h
--- a/Engine/Source/Runtime/Engine/Public/Model.h
+++ b/Engine/Source/Runtime/Engine/Public/Model.h
@@ -351,6 +351,8 @@ private:
 /** A struct that contains a set of conodes that will be used in one mapping */
 struct FNodeGroup
 {
+	ENGINE_API ~FNodeGroup();
+
 	/** List of nodes in the node group */
 	TArray<int32> Nodes;
 
diff --git a/Engine/Source/Runtime/Engine/Public/TextureResource.h b/Engine/Source/Runtime/Engine/Public/TextureResource.h
--- a/Engine/Source/Runtime/Engine/Public/TextureResource.h
+++ b/Engine/Source/Runtime/Engine/Public/TextureResource.h
@@ -400,6 +400,7 @@ public:
 	 * @param InOwner - 2d texture object to create a resource for
 	 */
 	FTextureRenderTarget2DResource(const class UTextureRenderTarget2D* InOwner);
+	virtual ~FTextureRenderTarget2DResource();
 
 	FORCEINLINE FLinearColor GetClearColor()
 	{
diff --git a/Engine/Source/Runtime/Experimental/Chaos/Public/Chaos/UniformGrid.h b/Engine/Source/Runtime/Experimental/Chaos/Public/Chaos/UniformGrid.h
--- a/Engine/Source/Runtime/Experimental/Chaos/Public/Chaos/UniformGrid.h
+++ b/Engine/Source/Runtime/Experimental/Chaos/Public/Chaos/UniformGrid.h
@@ -263,7 +263,7 @@ class CHAOS_API TUniformGridBase
 };
 
 template<class T, int d>
-class CHAOS_API TUniformGrid : public TUniformGridBase<T, d>
+class TUniformGrid : public TUniformGridBase<T, d>
 {
 	using TUniformGridBase<T, d>::MCells;
 	using TUniformGridBase<T, d>::MMinCorner;
@@ -311,7 +311,7 @@ class CHAOS_API TUniformGrid : public TUniformGridBase<T, d>
 };
 
 template<class T>
-class CHAOS_API TMPMGrid : public TUniformGridBase<T, 3>
+class TMPMGrid : public TUniformGridBase<T, 3>
 {
 	using TUniformGridBase<T, 3>::MCells;
 	using TUniformGridBase<T, 3>::MMinCorner;
diff --git a/Engine/Source/Runtime/GeometryCore/Private/DynamicMesh/DynamicMeshOverlay.cpp b/Engine/Source/Runtime/GeometryCore/Private/DynamicMesh/DynamicMeshOverlay.cpp
--- a/Engine/Source/Runtime/GeometryCore/Private/DynamicMesh/DynamicMeshOverlay.cpp
+++ b/Engine/Source/Runtime/GeometryCore/Private/DynamicMesh/DynamicMeshOverlay.cpp
@@ -1789,17 +1789,17 @@ namespace Geometry
 // These are explicit instantiations of the templates that are exported from the shared lib.
 // Only these instantiations of the template can be used.
 // This is necessary because we have placed most of the templated functions in this .cpp file, instead of the header.
-template class TDynamicMeshOverlay<float, 1>;
-template class TDynamicMeshOverlay<double, 1>;
-template class TDynamicMeshOverlay<int, 1>;
-template class TDynamicMeshOverlay<float, 2>;
-template class TDynamicMeshOverlay<double, 2>;
-template class TDynamicMeshOverlay<int, 2>;
-template class TDynamicMeshOverlay<float, 3>;
-template class TDynamicMeshOverlay<double, 3>;
-template class TDynamicMeshOverlay<int, 3>;
-template class TDynamicMeshOverlay<float, 4>;
-template class TDynamicMeshOverlay<double, 4>;
+template class GEOMETRYCORE_API TDynamicMeshOverlay<float, 1>;
+template class GEOMETRYCORE_API TDynamicMeshOverlay<double, 1>;
+template class GEOMETRYCORE_API TDynamicMeshOverlay<int, 1>;
+template class GEOMETRYCORE_API TDynamicMeshOverlay<float, 2>;
+template class GEOMETRYCORE_API TDynamicMeshOverlay<double, 2>;
+template class GEOMETRYCORE_API TDynamicMeshOverlay<int, 2>;
+template class GEOMETRYCORE_API TDynamicMeshOverlay<float, 3>;
+template class GEOMETRYCORE_API TDynamicMeshOverlay<double, 3>;
+template class GEOMETRYCORE_API TDynamicMeshOverlay<int, 3>;
+template class GEOMETRYCORE_API TDynamicMeshOverlay<float, 4>;
+template class GEOMETRYCORE_API TDynamicMeshOverlay<double, 4>;
 
 template class TDynamicMeshVectorOverlay<float, 2, FVector2f>;
 template class TDynamicMeshVectorOverlay<double, 2, FVector2d>;
diff --git a/Engine/Source/Runtime/GeometryCore/Public/DynamicMesh/DynamicMeshOverlay.h b/Engine/Source/Runtime/GeometryCore/Public/DynamicMesh/DynamicMeshOverlay.h
--- a/Engine/Source/Runtime/GeometryCore/Public/DynamicMesh/DynamicMeshOverlay.h
+++ b/Engine/Source/Runtime/GeometryCore/Public/DynamicMesh/DynamicMeshOverlay.h
@@ -64,7 +64,7 @@ class FDynamicMesh3;
  * you wish to use in the block at the top of DynamicMeshOverlay.cpp
  */
 template<typename RealType, int ElementSize>
-class GEOMETRYCORE_API TDynamicMeshOverlay
+class TDynamicMeshOverlay
 {
 
 protected:
diff --git a/Engine/Source/Runtime/GeometryCore/Public/Generators/CapsuleGenerator.h b/Engine/Source/Runtime/GeometryCore/Public/Generators/CapsuleGenerator.h
--- a/Engine/Source/Runtime/GeometryCore/Public/Generators/CapsuleGenerator.h
+++ b/Engine/Source/Runtime/GeometryCore/Public/Generators/CapsuleGenerator.h
@@ -101,8 +101,8 @@ private:
 	{
 		// generate the UV's
 		int32 NumPhi = (2 * NumHemisphereArcSteps);
-		const float DUVphi = 1.0 / float(NumPhi - 1);
-		const float DUVtheta = -1.0 / float(NumCircleSteps);
+		const float DUVphi = 1.0f / float(NumPhi - 1);
+		const float DUVtheta = -1.0f / float(NumCircleSteps);
 
 		int32 UVIdx = 0;
 		int32 p,t;
diff --git a/Engine/Source/Runtime/GeometryCore/Public/Generators/GridBoxMeshGenerator.h b/Engine/Source/Runtime/GeometryCore/Public/Generators/GridBoxMeshGenerator.h
--- a/Engine/Source/Runtime/GeometryCore/Public/Generators/GridBoxMeshGenerator.h
+++ b/Engine/Source/Runtime/GeometryCore/Public/Generators/GridBoxMeshGenerator.h
@@ -210,8 +210,8 @@ public:
 						UV[UVXDim] = float (D0 * Nscale[D[0][Dim]] - .5);
 						UV[UVYDim] = float (D1 * Nscale[D[1][Dim]] - .5);
 						// invert axes to match the desired UV patterns & so the opp faces are not backwards
-						UV.X *= SideSign * Minor1Flip[Dim];
-						UV.Y *= Minor2Flip[Dim];
+						UV.X *= float(SideSign * Minor1Flip[Dim]);
+						UV.Y *= float(Minor2Flip[Dim]);
 						// recenter and scale up
 						UV[UVXDim] = float ( (UV[UVXDim] + .5f) * WidthUVScale);
 						UV[UVYDim] = float ( (UV[UVYDim] + .5f) * HeightUVScale);
diff --git a/Engine/Source/Runtime/GeometryCore/Public/Generators/MinimalBoxMeshGenerator.h b/Engine/Source/Runtime/GeometryCore/Public/Generators/MinimalBoxMeshGenerator.h
--- a/Engine/Source/Runtime/GeometryCore/Public/Generators/MinimalBoxMeshGenerator.h
+++ b/Engine/Source/Runtime/GeometryCore/Public/Generators/MinimalBoxMeshGenerator.h
@@ -58,7 +58,7 @@ public:
 			{
 				Normals[AttribIndex] = FaceNormal;
 				FVector2i CornerUV = IndexUtil::BoxFacesUV[j];
-				UVs[AttribIndex] = FVector2f(WidthUVScale*(float)CornerUV.X, HeightUVScale*(float)CornerUV.Y);
+				UVs[AttribIndex] = FVector2f(float(WidthUVScale*CornerUV.X), float(HeightUVScale*CornerUV.Y));
 				ElementIndices[j] = AttribIndex;
 				AttribIndex++;
 			}
diff --git a/Engine/Source/Runtime/GeometryCore/Public/Image/ImageInfilling.h b/Engine/Source/Runtime/GeometryCore/Public/Image/ImageInfilling.h
--- a/Engine/Source/Runtime/GeometryCore/Public/Image/ImageInfilling.h
+++ b/Engine/Source/Runtime/GeometryCore/Public/Image/ImageInfilling.h
@@ -2,7 +2,9 @@
 
 #pragma once
 
+#include "IntVectorTypes.h"
 #include "Templates/Tuple.h"
+#include "Util/IndexUtil.h"
 #include "VectorTypes.h"
 #include "Image/ImageDimensions.h"
 #include "Image/ImageBuilder.h"
@@ -67,7 +69,7 @@ public:
 					Count++;
 				}
 			}
-			ActiveSetQueue.Insert(k, 8 - Count);
+			ActiveSetQueue.Insert(k, float(8 - Count));
 			PixelToIndexMap.Add(CenterCoords, k);
 		}
 
@@ -120,7 +122,7 @@ public:
 						if (FoundIndex != nullptr && ActiveSetQueue.Contains(*FoundIndex))
 						{
 							float CurPriority = ActiveSetQueue.GetPriority(*FoundIndex);
-							ActiveSetQueue.Update(*FoundIndex, CurPriority - 1.0);
+							ActiveSetQueue.Update(*FoundIndex, CurPriority - 1.0f);
 						}
 					}
 				}
diff --git a/Engine/Source/Runtime/GeometryCore/Public/Implicit/SparseNarrowBandMeshSDF.h b/Engine/Source/Runtime/GeometryCore/Public/Implicit/SparseNarrowBandMeshSDF.h
--- a/Engine/Source/Runtime/GeometryCore/Public/Implicit/SparseNarrowBandMeshSDF.h
+++ b/Engine/Source/Runtime/GeometryCore/Public/Implicit/SparseNarrowBandMeshSDF.h
@@ -438,7 +438,7 @@ private:
 
 	float upper_bound(const FVector3i& GridDimensions, float DX) const
 	{
-		return (float)((GridDimensions[0] + GridDimensions[1] + GridDimensions[2]) * DX);
+		return float(GridDimensions[0] + GridDimensions[1] + GridDimensions[2]) * DX;
 	}
 
 	// if provided, uses IsInsideFunction to determine sign, otherwise uses a counting method.
diff --git a/Engine/Source/Runtime/GeometryCore/Public/Spatial/SparseDynamicPointOctree3.h b/Engine/Source/Runtime/GeometryCore/Public/Spatial/SparseDynamicPointOctree3.h
--- a/Engine/Source/Runtime/GeometryCore/Public/Spatial/SparseDynamicPointOctree3.h
+++ b/Engine/Source/Runtime/GeometryCore/Public/Spatial/SparseDynamicPointOctree3.h
@@ -955,7 +955,7 @@ void FSparseDynamicPointOctree3::ComputeStatistics(FStatistics& StatsOut) const
 		int ObjCount = CellPointLists.GetCount(CellID);
 		StatsOut.LevelObjCounts[Cell.Level] += ObjCount;
 		StatsOut.LevelMaxObjCounts[Cell.Level] = FMath::Max(StatsOut.LevelMaxObjCounts[Cell.Level], ObjCount);
-		StatsOut.LevelCellSizes[Cell.Level] = GetCellWidth(Cell.Level);
+		StatsOut.LevelCellSizes[Cell.Level] = float(GetCellWidth(Cell.Level));
 	}
 }
 
