Add Tetrahedron3 to TriangleTypes.h and a new TetUtil.h, to support common tet operations

diff --git a/Engine/Source/Runtime/GeometryCore/Private/CompGeom/Delaunay3.cpp b/Engine/Source/Runtime/GeometryCore/Private/CompGeom/Delaunay3.cpp
--- a/Engine/Source/Runtime/GeometryCore/Private/CompGeom/Delaunay3.cpp
+++ b/Engine/Source/Runtime/GeometryCore/Private/CompGeom/Delaunay3.cpp
@@ -5,6 +5,7 @@
 #include "CompGeom/ExactPredicates.h"
 #include "CompGeom/ConvexHull3.h"
 #include "Spatial/ZOrderCurvePoints.h"
+#include "TetUtil.h"
 
 #include "Async/ParallelFor.h"
 
@@ -235,10 +236,7 @@ namespace Delaunay3Internal
 	// Fill an array with the properly-oriented triangle faces of the tet, starting with Tri itself
 	void GetTetFaces(const FIndex4i& Tet, FIndex3i OutFaces[4])
 	{
-		OutFaces[0] = FIndex3i(Tet.A, Tet.B, Tet.C);
-		OutFaces[1] = FIndex3i(Tet.A, Tet.D, Tet.B);
-		OutFaces[2] = FIndex3i(Tet.A, Tet.C, Tet.D);
-		OutFaces[3] = FIndex3i(Tet.B, Tet.D, Tet.C);
+		return TetUtil::GetTetFaces(Tet, OutFaces);
 	}
 
 	// Fill an array with four possible orientations of the given tet, each starting with a different triangle face
diff --git a/Engine/Source/Runtime/GeometryCore/Public/TetUtil.h b/Engine/Source/Runtime/GeometryCore/Public/TetUtil.h
new file mode 100644
--- /dev/null
+++ b/Engine/Source/Runtime/GeometryCore/Public/TetUtil.h
@@ -0,0 +1,46 @@
+// Copyright Epic Games, Inc. All Rights Reserved.
+
+#pragma once
+
+#include "VectorTypes.h"
+#include "TriangleTypes.h"
+#include "CompGeom/ExactPredicates.h"
+
+namespace UE::Geometry::TetUtil
+{
+	// Fill an array with the oriented triangle faces of the tet
+	template<bool bReverseOrientation = false>
+	void GetTetFaces(const FIndex4i& Tet, FIndex3i OutFaces[4])
+	{
+		const FIndex3i* FaceOrdering = GetTetFaceOrdering<bReverseOrientation>();
+		OutFaces[0] = FIndex3i(Tet[FaceOrdering[0].A], Tet[FaceOrdering[0].B], Tet[FaceOrdering[0].C]);
+		OutFaces[1] = FIndex3i(Tet[FaceOrdering[1].A], Tet[FaceOrdering[1].B], Tet[FaceOrdering[1].C]);
+		OutFaces[2] = FIndex3i(Tet[FaceOrdering[2].A], Tet[FaceOrdering[2].B], Tet[FaceOrdering[2].C]);
+		OutFaces[3] = FIndex3i(Tet[FaceOrdering[3].A], Tet[FaceOrdering[3].B], Tet[FaceOrdering[3].C]);
+	}
+
+	// Test whether point Pt is inside or on a (non-degenerate) tetrahedron Tet.
+	// Note: Will always report 'outside' for a fully degenerate (flat) tetrahedron
+	// @return 0 if Pt is outside Tet (or Tet is degenerate), 1 if Pt is inside a positively-oriented Tet, -1 if Pt is inside and Tet is inverted
+	template<typename RealType>
+	int32 IsInsideExact(const TTetrahedron3<RealType>& Tet, TVector<RealType>& Pt)
+	{
+		int32 SeenSide = 0;
+		for (int32 TetFace = 0; TetFace < 4; ++TetFace)
+		{
+			FIndex3i FaceInds = Tet.GetFaceIndices(TetFace);
+			int32 Side = TMathUtil<RealType>::SignAsInt(ExactPredicates::Orient3<RealType>(Tet.V[FaceInds.A], Tet.V[FaceInds.B], Tet.V[FaceInds.C], Pt));
+			if (Side != 0)
+			{
+				if (Side * SeenSide < 0)
+				{
+					return 0;
+				}
+				SeenSide = Side;
+			}
+		}
+		
+		return SeenSide;
+	}
+ 
+} // end namespace UE::Geometry::TetUtil
\ No newline at end of file
diff --git a/Engine/Source/Runtime/GeometryCore/Public/TriangleTypes.h b/Engine/Source/Runtime/GeometryCore/Public/TriangleTypes.h
--- a/Engine/Source/Runtime/GeometryCore/Public/TriangleTypes.h
+++ b/Engine/Source/Runtime/GeometryCore/Public/TriangleTypes.h
@@ -4,6 +4,8 @@
 
 #include "VectorTypes.h"
 #include "VectorUtil.h"
+#include "IndexTypes.h"
+#include "BoxTypes.h"
 
 namespace UE
 {
@@ -247,5 +249,99 @@ typedef TTriangle3<float> FTriangle3f;
 typedef TTriangle3<double> FTriangle3d;
 typedef TTriangle3<int> FTriangle3i;
 
+
+// Note: More TetUtil functions are in TetUtil.h; this subset is here for use by the TTetrahedron3 struct below
+namespace TetUtil
+{
+	// @return a C array of vertex orderings for each of the 4 triangle faces of a tetrahedron
+	template<bool bReverseOrientation = false>
+	inline const FIndex3i* GetTetFaceOrdering()
+	{
+		const static FIndex3i FaceMap[4]{ FIndex3i(0,1,2), FIndex3i(0,3,1), FIndex3i(0,2,3), FIndex3i(1,3,2) };
+		const static FIndex3i FaceMapRev[4]{ FIndex3i(1,0,2), FIndex3i(3,0,1), FIndex3i(2,0,3), FIndex3i(3,1,2) };
+		if constexpr (bReverseOrientation)
+		{
+			return FaceMapRev;
+		}
+		else
+		{
+			return FaceMap;
+		}
+	}
+}
+
+template<typename RealType>
+struct TTetrahedron3
+{
+	TVector<RealType> V[4];
+
+	TTetrahedron3() {}
+
+	TTetrahedron3(const TVector<RealType>& V0, const TVector<RealType>& V1, const TVector<RealType>& V2, const TVector<RealType>& V3)
+	{
+		V[0] = V0;
+		V[1] = V1;
+		V[2] = V2;
+		V[3] = V3;
+	}
+
+	TTetrahedron3(const TVector<RealType> VIn[4])
+	{
+		V[0] = VIn[0];
+		V[1] = VIn[1];
+		V[2] = VIn[2];
+		V[3] = VIn[3];
+	}
+
+	TVector<RealType> BarycentricPoint(RealType Bary0, RealType Bary1, RealType Bary2, RealType Bary3) const
+	{
+		return Bary0 * V[0] + Bary1 * V[1] + Bary2 * V[2] + Bary3 * V[3];
+	}
+
+	TVector<RealType> BarycentricPoint(const TVector4<RealType>& BaryCoords) const
+	{
+		return BaryCoords[0] * V[0] + BaryCoords[1] * V[1] + BaryCoords[2] * V[2] + BaryCoords[3] * V[3];
+	}
+
+	TAxisAlignedBox3<RealType> Bounds() const
+	{
+		FAxisAlignedBox3d RetBounds;
+		RetBounds.Contain(V[0]);
+		RetBounds.Contain(V[1]);
+		RetBounds.Contain(V[2]);
+		RetBounds.Contain(V[3]);
+		return RetBounds;
+	}
+
+	// Get the ith triangular face of the tetrahedron, as indices into the four tetrahedron vertices
+	template<bool bReverseOrientation = false>
+	static FIndex3i GetFaceIndices(int32 Idx)
+	{
+		return TetUtil::GetTetFaceOrdering<bReverseOrientation>()[Idx];
+	}
+
+	template<bool bReverseOrientation = false>
+	TTriangle3<RealType> GetFace(int32 Idx) const
+	{
+		FIndex3i Face = GetFaceIndices<bReverseOrientation>(Idx);
+		return TTriangle3<RealType>(V[Face.A], V[Face.B], V[Face.C]);
+	}
+
+	/** @return centroid of this tetrahedron */
+	TVector<RealType> Centroid() const
+	{
+		constexpr RealType f = 1.0 / 4.0;
+		return TVector<RealType>(
+			(V[0].X + V[1].X + V[2].X + V[3].X) * f,
+			(V[0].Y + V[1].Y + V[2].Y + V[3].Y) * f,
+			(V[0].Z + V[1].Z + V[2].Z + V[3].Z) * f
+			);
+	}
+};
+
+typedef TTetrahedron3<float> FTetrahedron3f;
+typedef TTetrahedron3<double> FTetrahedron3d;
+
+
 } // end namespace UE::Geometry
 } // end namespace UE
\ No newline at end of file
