Add a 'IsWithinDistanceSquared' to the Mesh AABB Tree class, which can early-out faster than the standard FindNearestTriangle

diff --git a/Engine/Source/Runtime/GeometryCore/Public/Spatial/MeshAABBTree3.h b/Engine/Source/Runtime/GeometryCore/Public/Spatial/MeshAABBTree3.h
--- a/Engine/Source/Runtime/GeometryCore/Public/Spatial/MeshAABBTree3.h
+++ b/Engine/Source/Runtime/GeometryCore/Public/Spatial/MeshAABBTree3.h
@@ -203,8 +203,31 @@ public:
 		return Point;
 	}
 
+	/**
+	 * Test whether there is any triangle closer than sqrt(ThresholdDistanceSqr) to Point
+	 * Note that FQueryOptions::MaxDistance is ignored by this query, as ThresholdDistanceSqr controls the query distance instead
+	 * 
+	 * @param OutTriangleID  The ID of a triangle (not necessarily the closest) within the distance threshold, or InvalidID if none found.
+	 */
+	virtual bool IsWithinDistanceSquared(
+		const FVector3d& Point, double ThresholdDistanceSqr, int& OutTriangleID, const FQueryOptions& Options = FQueryOptions()
+	) const
+	{
+		if (ensure(IsValid(Options.bAllowUnsafeModifiedMeshQueries)) == false)
+		{
+			return false;
+		}
+
+		OutTriangleID = IndexConstants::InvalidID;
+		find_nearest_tri<true>(RootIndex, Point, ThresholdDistanceSqr, OutTriangleID, Options);
+		return OutTriangleID != IndexConstants::InvalidID;
+	}
+
 protected:
-	void find_nearest_tri(int IBox, const FVector3d& P, double& NearestDistSqr, int& TID, const FQueryOptions& Options) const
+	// bEarlyStop causes the traversal to stop as soon as any triangle closer than NearestDistSqr is found
+	// @return true if an early stop triangle was found
+	template<bool bEarlyStop = false>
+	bool find_nearest_tri(int IBox, const FVector3d& P, double& NearestDistSqr, int& TID, const FQueryOptions& Options) const
 	{
 		int idx = BoxToIndex[IBox];
 		if (idx < TrianglesEnd)
@@ -222,6 +245,10 @@ protected:
 				{
 					NearestDistSqr = fTriDistSqr;
 					TID = ti;
+					if constexpr (bEarlyStop)
+					{
+						return true;
+					}
 				}
 			}
 		}
@@ -234,7 +261,11 @@ protected:
 				double fChild1DistSqr = BoxDistanceSqr(iChild1, P);
 				if (fChild1DistSqr <= NearestDistSqr)
 				{
-					find_nearest_tri(iChild1, P, NearestDistSqr, TID, Options);
+					bool bFoundEarly = find_nearest_tri<bEarlyStop>(iChild1, P, NearestDistSqr, TID, Options);
+					if (bEarlyStop && bFoundEarly)
+					{
+						return true;
+					}
 				}
 			}
 			else
@@ -248,10 +279,18 @@ protected:
 				{
 					if (fChild1DistSqr < NearestDistSqr)
 					{
-						find_nearest_tri(iChild1, P, NearestDistSqr, TID, Options);
+						bool bFoundEarly1 = find_nearest_tri<bEarlyStop>(iChild1, P, NearestDistSqr, TID, Options);
+						if (bEarlyStop && bFoundEarly1)
+						{
+							return true;
+						}
 						if (fChild2DistSqr < NearestDistSqr)
 						{
-							find_nearest_tri(iChild2, P, NearestDistSqr, TID, Options);
+							bool bFoundEarly2 = find_nearest_tri<bEarlyStop>(iChild2, P, NearestDistSqr, TID, Options);
+							if (bEarlyStop && bFoundEarly2)
+							{
+								return true;
+							}
 						}
 					}
 				}
@@ -259,15 +298,24 @@ protected:
 				{
 					if (fChild2DistSqr < NearestDistSqr)
 					{
-						find_nearest_tri(iChild2, P, NearestDistSqr, TID, Options);
+						bool bFoundEarly1 = find_nearest_tri<bEarlyStop>(iChild2, P, NearestDistSqr, TID, Options);
+						if (bEarlyStop && bFoundEarly1)
+						{
+							return true;
+						}
 						if (fChild1DistSqr < NearestDistSqr)
 						{
-							find_nearest_tri(iChild1, P, NearestDistSqr, TID, Options);
+							bool bFoundEarly2 = find_nearest_tri<bEarlyStop>(iChild1, P, NearestDistSqr, TID, Options);
+							if (bEarlyStop && bFoundEarly2)
+							{
+								return true;
+							}
 						}
 					}
 				}
 			}
 		}
+		return false;
 	}
 
 
