Clean up old code in FastWinding.h and fix accidentally not using the cached triangle data when computing weighted centroids
+ move down some sqrts that only need to be called on the final value

diff --git a/Engine/Source/Runtime/GeometryCore/Public/Spatial/FastWinding.h b/Engine/Source/Runtime/GeometryCore/Public/Spatial/FastWinding.h
--- a/Engine/Source/Runtime/GeometryCore/Public/Spatial/FastWinding.h
+++ b/Engine/Source/Runtime/GeometryCore/Public/Spatial/FastWinding.h
@@ -24,7 +24,7 @@ namespace FastTriWinding
 	 *  R: max distance from P to Triangles
 	 *  Order1: first-order vector coeff
 	 *  Order2: second-order matrix coeff
-	 *  TriCache: precomputed triangle centroid/normal/area (todo @jimmy either support passing in as possibly-nullptr or remove commented-out null branches!)
+	 *  TriCache: precomputed triangle centroid/normal/area
 	 */
 	template <class TriangleMeshType, class IterableTriangleIndices>
 	void ComputeCoeffsSerial(const TriangleMeshType& Mesh,
@@ -44,19 +44,9 @@ namespace FastTriWinding
 		double sum_area = 0;
 		for (int tid : Triangles)
 		{
-			//if (TriCache == null)
-			//{
-			//	double area = TriCache.Areas[tid];
-			//	sum_area += area;
-			//	P += area * TriCache.Centroids[tid];
-			//}
-			//else
-			{
-				Mesh.GetTriVertices(tid, P0, P1, P2);
-				double area = VectorUtil::Area(P0, P1, P2);
-				sum_area += area;
-				P += area * ((P0 + P1 + P2) / 3.0);
-			}
+			double Area = TriCache.Areas[tid];
+			sum_area += Area;
+			P += Area * TriCache.Centroids[tid];
 		}
 		P /= sum_area;
 
@@ -64,29 +54,22 @@ namespace FastTriWinding
 		// 'radius' value R, which is max dist from any tri vertex to P
 		FVector3d n, c;
 		double a = 0;
+		double RSq = 0;
 		for (int tid : Triangles)
 		{
 			Mesh.GetTriVertices(tid, P0, P1, P2);
-
-			//if (TriCache == null)
-			//{
-			//	c = (1.0 / 3.0) * (P0 + P1 + P2);
-			//	n = VectorUtil::NormalArea(P0, P1, P2, out a);
-			//}
-			//else
-			{
-				TriCache.GetTriInfo(tid, n, a, c);
-			}
+			TriCache.GetTriInfo(tid, n, a, c);
 
 			Order1 += a * n;
 
 			FVector3d dcp = c - P;
 			Order2 += a * FMatrix3d(dcp, n);
 
-			// this is just for return value...
-			double maxdist = FMath::Max3(DistanceSquared(P0,P), DistanceSquared(P1,P), DistanceSquared(P2,P));
-			R = FMath::Max(R, sqrt(maxdist));
+			// update max radius R (as squared value in loop)
+			double MaxDistSq = FMath::Max3(DistanceSquared(P0,P), DistanceSquared(P1,P), DistanceSquared(P2,P));
+			RSq = FMath::Max(RSq, MaxDistSq);
 		}
+		R = FMath::Sqrt(RSq);
 	}
 
 	/**
@@ -95,7 +78,7 @@ namespace FastTriWinding
 	 *  R: max distance from P to Triangles
 	 *  Order1: first-order vector coeff
 	 *  Order2: second-order matrix coeff
-	 *  TriCache: precomputed triangle centroid/normal/area (todo @jimmy either support passing in as possibly-nullptr or remove commented-out null branches!)
+	 *  TriCache: precomputed triangle centroid/normal/area
 	 */
 	template <class TriangleMeshType>
 	void ComputeCoeffs(const TriangleMeshType& Mesh,
@@ -128,13 +111,11 @@ namespace FastTriWinding
 		auto CentroidTransform = [&](int64 TriangleSubsetIndex) -> PData
 		{
 			check(TriangleSubsetIndex < TriangleArray.Num());
-			int tid = TriangleArray[(int)TriangleSubsetIndex];
-			FVector3d P0, P1, P2;
-			Mesh.GetTriVertices(tid, P0, P1, P2);
+			int TID = TriangleArray[(int)TriangleSubsetIndex];
 
 			PData DataOut;
-			DataOut.Area = VectorUtil::Area(P0, P1, P2);
-			DataOut.P = DataOut.Area * ((P0 + P1 + P2) / 3.0);
+			DataOut.Area = TriCache.Areas[TID];
+			DataOut.P = DataOut.Area * TriCache.Centroids[TID];
 			return DataOut;
 		};
 
@@ -156,7 +137,7 @@ namespace FastTriWinding
 		{
 			FVector3d Order1;
 			FMatrix3d Order2;
-			double R;
+			double RSq;
 		};
 
 		auto ExpansionTransform = [&, P = CentroidData.P](int32 TriangleSubsetIndex) -> OrderData
@@ -177,8 +158,8 @@ namespace FastTriWinding
 			DataOut.Order2 = a * FMatrix3d(dcp, n);
 
 			// this is just for return value...
-			double maxdist = FMath::Max3(DistanceSquared(P0, P), DistanceSquared(P1, P), DistanceSquared(P2, P));
-			DataOut.R = sqrt(maxdist);
+			double MaxDistSq = FMath::Max3(DistanceSquared(P0, P), DistanceSquared(P1, P), DistanceSquared(P2, P));
+			DataOut.RSq = MaxDistSq;
 
 			return DataOut;
 		};
@@ -188,7 +169,7 @@ namespace FastTriWinding
 			OrderData Out;
 			Out.Order1 = A.Order1 + B.Order1;
 			Out.Order2 = A.Order2 + B.Order2;
-			Out.R = FMath::Max(A.R, B.R);
+			Out.RSq = FMath::Max(A.RSq, B.RSq);
 			return Out;
 		};
 
@@ -201,7 +182,7 @@ namespace FastTriWinding
 		// Set out values
 
 		P = CentroidData.P;
-		R = Orders.R;
+		R = FMath::Sqrt(Orders.RSq);
 		Order1 = Orders.Order1;
 		Order2 = Orders.Order2;
 	}
