make TPointHashGrid queries not create arrays to iterate cell points, and test the center cell first
+ add a find-any variant that can early out
+ fix the FindPointsInBall return value to return the number of points found, not the number of points in the append array (matching the comment, and the behavior if 0 points are present)

diff --git a/Engine/Source/Runtime/GeometryCore/Public/Spatial/PointHashGrid3.h b/Engine/Source/Runtime/GeometryCore/Public/Spatial/PointHashGrid3.h
--- a/Engine/Source/Runtime/GeometryCore/Public/Spatial/PointHashGrid3.h
+++ b/Engine/Source/Runtime/GeometryCore/Public/Spatial/PointHashGrid3.h
@@ -193,9 +193,9 @@ public:
 		return;
 	}
 
-
 	/**
-	 * Find nearest point in grid, within a given sphere, without locking / thread-safety.
+	 * Find nearest point within radius.
+	 * Note: Not thread-safe to update, remove or insert points during this query.
 	 * @param QueryPoint the center of the query sphere
 	 * @param Radius the radius of the query sphere
 	 * @param DistanceSqFunc Function you provide which measures the squared distance between QueryPoint and a Value
@@ -203,23 +203,72 @@ public:
 	 * @return the found pair (Value,DistanceSqFunc(Value)), or (InvalidValue,MaxDouble) if not found
 	 */
 	TPair<PointDataType, RealType> FindNearestInRadius(
-		const TVector<RealType>& QueryPoint, RealType Radius, 
+		const TVector<RealType>& QueryPoint, RealType Radius,
+		TFunctionRef<RealType(const PointDataType&)> DistanceSqFunc,
+		TFunctionRef<bool(const PointDataType&)> IgnoreFunc) const
+	{
+		return FindInRadiusHelper<false>(QueryPoint, Radius, DistanceSqFunc, IgnoreFunc);
+	}
+
+	TPair<PointDataType, RealType> FindNearestInRadius(
+		const TVector<RealType>& QueryPoint, RealType Radius,
+		TFunctionRef<RealType(const PointDataType&)> DistanceSqFunc) const
+	{
+		return FindInRadiusHelper<false>(QueryPoint, Radius, DistanceSqFunc, [](const PointDataType& data) { return false; });
+	}
+
+	/**
+	 * Find any point within radius.
+	 * Note: Not thread-safe to update, remove or insert points during this query.
+	 * @param QueryPoint the center of the query sphere
+	 * @param Radius the radius of the query sphere
+	 * @param DistanceSqFunc Function you provide which measures the squared distance between QueryPoint and a Value
+	 * @param IgnoreFunc optional Function you may provide which will result in a Value being ignored if IgnoreFunc(Value) returns true
+	 * @return the found pair (Value,DistanceSqFunc(Value)), or (InvalidValue,MaxDouble) if not found
+	 */
+	TPair<PointDataType, RealType> FindAnyInRadius(
+		const TVector<RealType>& QueryPoint, RealType Radius,
 		TFunctionRef<RealType(const PointDataType&)> DistanceSqFunc,
+		TFunctionRef<bool(const PointDataType&)> IgnoreFunc) const
+	{
+		return FindInRadiusHelper<true>(QueryPoint, Radius, DistanceSqFunc, IgnoreFunc);
+	}
+
+	TPair<PointDataType, RealType> FindAnyInRadius(
+		const TVector<RealType>& QueryPoint, RealType Radius,
+		TFunctionRef<RealType(const PointDataType&)> DistanceSqFunc) const
+	{
+		return FindInRadiusHelper<true>(QueryPoint, Radius, DistanceSqFunc, [](const PointDataType& data) { return false; });
+	}
+
+
+	/**
+	 * Find all points in grid within a given sphere. 
+	 * Note: Not thread-safe to update, remove or insert points during this query.
+	 * @param QueryPoint the center of the query sphere
+	 * @param Radius the radius of the query sphere
+	 * @param DistanceSqFunc Function you provide which measures the squared distance between QueryPoint and a Value
+	 * @param ResultsOut Array that points in sphere will be added to
+	 * @param IgnoreFunc optional Function you may provide which will result in a Value being ignored if IgnoreFunc(Value) returns true
+	 * @return the number of found points
+	 */
+	int FindPointsInBall(
+		const TVector<RealType>& QueryPoint, RealType Radius,
+		TFunctionRef<RealType(const PointDataType&)> DistanceSqFunc,
+		TArray<PointDataType>& ResultsOut,
 		TFunctionRef<bool(const PointDataType&)> IgnoreFunc = [](const PointDataType& data) { return false; }) const
 	{
 		if (!Hash.Num())
 		{
-			return TPair<PointDataType, RealType>(GetInvalidValue(), TNumericLimits<RealType>::Max());
+			return 0;
 		}
+		int32 InitialNum = ResultsOut.Num();
 
 		FVector3i min_idx = Indexer.ToGrid(QueryPoint - Radius * TVector<RealType>::One());
 		FVector3i max_idx = Indexer.ToGrid(QueryPoint + Radius * TVector<RealType>::One());
 
-		RealType min_distsq = TNumericLimits<RealType>::Max();
-		PointDataType nearest = GetInvalidValue();
 		RealType RadiusSquared = Radius * Radius;
 
-		TArray<PointDataType> Values;
 		for (int zi = min_idx.Z; zi <= max_idx.Z; zi++)
 		{
 			for (int yi = min_idx.Y; yi <= max_idx.Y; yi++)
@@ -227,84 +276,106 @@ public:
 				for (int xi = min_idx.X; xi <= max_idx.X; xi++)
 				{
 					FVector3i idx(xi, yi, zi);
-					Values.Reset();
-					Hash.MultiFind(idx, Values);
-					for (PointDataType Value : Values)
+					for (typename TMultiMap<FVector3i, PointDataType>::TConstKeyIterator It = Hash.CreateConstKeyIterator(idx); It; ++It)
 					{
+						const PointDataType& Value = It.Value();
 						if (IgnoreFunc(Value))
 						{
 							continue;
 						}
 						RealType distsq = DistanceSqFunc(Value);
-						if (distsq < RadiusSquared && distsq < min_distsq)
+						if (distsq < RadiusSquared)
 						{
-							nearest = Value;
-							min_distsq = distsq;
+							ResultsOut.Add(Value);
 						}
 					}
 				}
 			}
 		}
 
-		return TPair<PointDataType, RealType>(nearest, min_distsq);
+		return ResultsOut.Num() - InitialNum;
 	}
 
-
-	/**
-	 * Find all points in grid within a given sphere, without locking / thread-safety.
-	 * @param QueryPoint the center of the query sphere
-	 * @param Radius the radius of the query sphere
-	 * @param DistanceSqFunc Function you provide which measures the squared distance between QueryPoint and a Value
-	 * @param ResultsOut Array of output points in sphere
-	 * @param IgnoreFunc optional Function you may provide which will result in a Value being ignored if IgnoreFunc(Value) returns true
-	 * @return the number of found points
-	 */
-	int FindPointsInBall(
+private:
+	template<bool bEarlyOut = false>
+	TPair<PointDataType, RealType> FindInRadiusHelper(
 		const TVector<RealType>& QueryPoint, RealType Radius,
 		TFunctionRef<RealType(const PointDataType&)> DistanceSqFunc,
-		TArray<PointDataType>& ResultsOut,
-		TFunctionRef<bool(const PointDataType&)> IgnoreFunc = [](const PointDataType& data) { return false; }) const
+		TFunctionRef<bool(const PointDataType&)> IgnoreFunc) const
 	{
 		if (!Hash.Num())
 		{
-			return 0;
+			return TPair<PointDataType, RealType>(GetInvalidValue(), TNumericLimits<RealType>::Max());
 		}
 
-		FVector3i min_idx = Indexer.ToGrid(QueryPoint - Radius * TVector<RealType>::One());
-		FVector3i max_idx = Indexer.ToGrid(QueryPoint + Radius * TVector<RealType>::One());
+		RealType MinDistSq = Radius * Radius;
+		PointDataType Nearest = GetInvalidValue();
 
-		RealType RadiusSquared = Radius * Radius;
+		auto SearchCell = [this, &Nearest, &MinDistSq, &DistanceSqFunc, &IgnoreFunc](FVector3i CellIdx)
+		{
+			bool bFound = false;
+			for (typename TMultiMap<FVector3i, PointDataType>::TConstKeyIterator It = Hash.CreateConstKeyIterator(CellIdx); It; ++It)
+			{
+				const PointDataType& Value = It.Value();
+				if (IgnoreFunc(Value))
+				{
+					continue;
+				}
+				RealType DistSq = DistanceSqFunc(Value);
+				if (DistSq < MinDistSq)
+				{
+					Nearest = Value;
+					MinDistSq = DistSq;
+					if (bEarlyOut) {
+						return true;
+					}
+					bFound = true;
+				}
+			}
+			return bFound;
+		};
+
+		FVector3i CenterIdx = Indexer.ToGrid(QueryPoint);
+		RealType SearchRadius = Radius;
+
+		if (SearchCell(CenterIdx))
+		{
+			if (bEarlyOut)
+			{
+				return TPair<PointDataType, RealType>(Nearest, MinDistSq);
+			}
+			SearchRadius = FMath::Sqrt(MinDistSq);
+		}
+
+		FVector3i min_idx = Indexer.ToGrid(QueryPoint - SearchRadius * TVector<RealType>::One());
+		FVector3i max_idx = Indexer.ToGrid(QueryPoint + SearchRadius * TVector<RealType>::One());
 
-		TArray<PointDataType> Values;
 		for (int zi = min_idx.Z; zi <= max_idx.Z; zi++)
 		{
 			for (int yi = min_idx.Y; yi <= max_idx.Y; yi++)
 			{
 				for (int xi = min_idx.X; xi <= max_idx.X; xi++)
 				{
-					FVector3i idx(xi, yi, zi);
-					Values.Reset();
-					Hash.MultiFind(idx, Values);
-					for (PointDataType Value : Values)
+					FVector3i Idx(xi, yi, zi);
+					if (Idx != CenterIdx)
 					{
-						if (IgnoreFunc(Value))
-						{
-							continue;
-						}
-						RealType distsq = DistanceSqFunc(Value);
-						if (distsq < RadiusSquared)
+						bool bFound = SearchCell(Idx);
+						if (bEarlyOut && bFound)
 						{
-							ResultsOut.Add(Value);
+							return TPair<PointDataType, RealType>(Nearest, MinDistSq);
 						}
 					}
 				}
 			}
 		}
 
-		return ResultsOut.Num();
+		if (Nearest == GetInvalidValue())
+		{
+			MinDistSq = TNumericLimits<RealType>::Max();
+		}
+		return TPair<PointDataType, RealType>(Nearest, MinDistSq);
 	}
 
-
 };
 
 template <typename PointDataType> using TPointHashGrid3d = TPointHashGrid3<PointDataType, double>;
