GeometryCore: Add TSparseGrid3::RangeIteration, that iterates over any allocated cells in sub-region of sparse grid. Remove default bAllocateIfMissing=true on ::Get() to avoid accidental grid expansion. Add FVector3i::operator< to allow sorting of integer vectors.

Various improvements to FSparseDynamicPointOctree:
- replaced FSmallListSet with new TSparseListSet<int>
- removed ValidPointIDs flag array, can be inferred from PointIDToCellMap
- replaced checks with checkSlow
- added some variants of internal functions for box center, width, etc, that can work off level/index directly, instead of requiring a cell
- range query functinos use new TSparseGrid3::RangeIteration to avoid iterating over all root cells
- added ::ParallelInsertDensePointSet() which can do parallel octree build into fully-expanded leaf nodes
- added ::ConfigureFromPointCountEstimate() which can estimate MaxDepth and RootDimension given an estimate of point count, useful primarily with ParallelInsertDensePointSet()
- added ::RemovePointUnsafe() which skips various range checking etc

diff --git a/Engine/Source/Runtime/GeometryCore/Private/Spatial/PointSetHashTable.cpp b/Engine/Source/Runtime/GeometryCore/Private/Spatial/PointSetHashTable.cpp
--- a/Engine/Source/Runtime/GeometryCore/Private/Spatial/PointSetHashTable.cpp
+++ b/Engine/Source/Runtime/GeometryCore/Private/Spatial/PointSetHashTable.cpp
@@ -21,7 +21,7 @@ void FPointSetHashtable::Build(double CellSizeIn, const FVector3d& OriginIn)
 		{
 			FVector3d Pt = Points->GetPoint(i);
 			FVector3i Idx = GridIndexer.ToGrid(Pt);
-			PointList* CellList = Grid.Get(Idx);
+			PointList* CellList = Grid.Get(Idx, true);
 			CellList->Add(i);
 		}
 	}
diff --git a/Engine/Source/Runtime/GeometryCore/Public/IntVectorTypes.h b/Engine/Source/Runtime/GeometryCore/Public/IntVectorTypes.h
--- a/Engine/Source/Runtime/GeometryCore/Public/IntVectorTypes.h
+++ b/Engine/Source/Runtime/GeometryCore/Public/IntVectorTypes.h
@@ -445,6 +445,14 @@ PRAGMA_ENABLE_DEPRECATION_WARNINGS
 	{
 		return X != Other.X || Y != Other.Y || Z != Other.Z;
 	}
+
+	constexpr bool operator<(const FVector3i& Other) const
+	{
+		if ( X != Other.X ) return X < Other.X;
+		else if (Y != Other.Y) return Y < Other.Y;
+		else if (Z != Other.Z) return Z < Other.Z;
+		else return false;
+	}
 };
 
 
diff --git a/Engine/Source/Runtime/GeometryCore/Public/Spatial/SparseDynamicPointOctree3.h b/Engine/Source/Runtime/GeometryCore/Public/Spatial/SparseDynamicPointOctree3.h
--- a/Engine/Source/Runtime/GeometryCore/Public/Spatial/SparseDynamicPointOctree3.h
+++ b/Engine/Source/Runtime/GeometryCore/Public/Spatial/SparseDynamicPointOctree3.h
@@ -8,17 +8,17 @@
 #include "GeometryTypes.h"
 #include "BoxTypes.h"
 #include "Util/DynamicVector.h"
-#include "Util/SmallListSet.h"
+#include "Util/SparseListSet.h"
 #include "Spatial/SparseGrid3.h"
 #include "Intersection/IntrRay3AxisAlignedBox3.h"
 #include "Async/ParallelFor.h"
-#include "Spatial/SparseDynamicOctree3.h"    // for FDynamicFlagArray
 
 namespace UE
 {
 namespace Geometry
 {
 
+
 /**
  * FSparsePointOctreeCell is a Node in a SparseDynamicOctree3. 
  */
@@ -112,7 +112,6 @@ class FSparseDynamicPointOctree3
 	//      insertions end up in very large buckets. When a child expands we should check if any of its parents would fit.
 	//    - Currently insertion is max-depth so we end up with a huge number of single-Point cells. Should only go down a level
 	//      if enough Points exist in current cell. Can do this in a greedy fashion, less optimal but still acceptable...
-	//    - get rid of ValidPointIDs, I don't think we need it?
 
 public:
 
@@ -138,6 +137,13 @@ public:
 
 public:
 
+	/**
+	 * Do a rough ballpark configuration of the RootDimension, MaxTreeDepth, and MaxPointsPerCell tree config values,
+	 * given a bounding-box dimension and estimate of the total number of points that will be inserted.
+	 * This is primarily intended to be used with ParallelInsertDensePointSet()
+	 */
+	inline void ConfigureFromPointCountEstimate(double MaxBoundsDimension, int CountEstimate);
+
 	/**
 	 * Test if an Point is stored in the tree
 	 * @param PointID ID of the Point
@@ -161,14 +167,31 @@ public:
 	 */
 	inline void InsertPoint_DynamicExpand(int32 PointID, TFunctionRef<FVector3d(int)> GetPositionFunc);
 
+	/**
+	 * Insert a set of dense points with IDs in range [0, MaxPointID-1], in parallel.
+	 * The points are only inserted in leaf nodes, at MaxTreeDepth level, so that value should
+	 * be set conservatively. Parallel insertion is across root cells, so if the RootDimension
+	 * is larger than the point set bounds, this will be slower than incremental construction.
+	 * ConfigureFromPointCountEstimate() provides reasonable values for large point sets.
+	 * 
+	 * @param GetPositionFunc function that returns position of point (Required to reinsert points on expand)
+	 */
+	inline void ParallelInsertDensePointSet(int32 MaxPointID, TFunctionRef<FVector3d(int)> GetPositionFunc);
+
 
 	/**
-	 * Remove an Point from the octree
+	 * Remove a Point from the octree
 	 * @param PointID ID of the Point
 	 * @return true if the Point was in the tree and removed
 	 */
 	inline bool RemovePoint(int32 PointID);
 
+	/**
+	 * Remove a Point from the octree. This function must only be called with PointIDs that are certain to be in the tree.
+	 * @param PointID ID of the Point. 
+	 */
+	inline void RemovePointUnsafe(int32 PointID);
+
 	/**
 	 * Update the position of an Point in the octree. This is more efficient than doing a remove+insert
 	 * @param PointID ID of the Point
@@ -200,7 +223,6 @@ public:
 		TArray<int>& PointIDsOut,
 		TArray<const FSparsePointOctreeCell*>* TempBuffer = nullptr) const;
 
-
 	/**
 	 * Collect PointIDs from all the cells with bounding boxes that intersect Bounds, where PredicateFunc passes
 	 * Query is parallelized across Root cells. So if there is only one Root cell, it's actually slower
@@ -216,6 +238,7 @@ public:
 		TArray<const FSparsePointOctreeCell*>* TempBuffer = nullptr) const;
 
 
+
 	/**
 	 * Check that the octree is internally valid
 	 * @param IsValidPointIDFunc function that returns true if given PointID is valid
@@ -260,10 +283,9 @@ protected:
 	// list of cells. Note that some cells may be unused, depending on CellRefCounts
 	TDynamicVector<FSparsePointOctreeCell> Cells;
 
-	FSmallListSet CellPointLists;			// per-cell Point ID lists
+	TSparseListSet<int32> CellPointLists;		// per-cell Point ID lists
 
-	TDynamicVector<uint32> PointIDToCellMap;	// map from external Point IDs to which cell the Point is in (or invalid)
-	FDynamicFlagArray ValidPointIDs;			// set of PointIDs in the tree. This is perhaps not necessary...couldn't we rely on PointIDToCellMap?
+	TDynamicVector<uint32> PointIDToCellMap;	// map from external Point IDs to which cell the Point is in (or InvalidCellID)
 
 	// RootCells are the top-level cells of the octree, of size RootDimension. 
 	// So the elements of this sparse grid are CellIDs
@@ -279,7 +301,7 @@ protected:
 	}
 
 
-	FAxisAlignedBox3d GetCellBox(uint32 Level, const FVector3i& Index) const
+	inline FAxisAlignedBox3d GetCellBox(uint32 Level, const FVector3i& Index) const
 	{
 		double CellWidth = GetCellWidth(Level);
 		double MinX = (CellWidth * (double)Index.X);
@@ -293,18 +315,23 @@ protected:
 	{
 		return GetCellBox(Cell.Level, Cell.Index);
 	}
-	FVector3d GetCellCenter(const FSparsePointOctreeCell& Cell) const
+
+	inline FVector3d GetCellCenter(uint32 Level, const FVector3i& Index) const
 	{
-		double CellWidth = GetCellWidth(Cell.Level);
-		double MinX = CellWidth * (double)Cell.Index.X;
-		double MinY = CellWidth * (double)Cell.Index.Y;
-		double MinZ = CellWidth * (double)Cell.Index.Z;
+		double CellWidth = GetCellWidth(Level);
+		double MinX = CellWidth * (double)Index.X;
+		double MinY = CellWidth * (double)Index.Y;
+		double MinZ = CellWidth * (double)Index.Z;
 		CellWidth *= 0.5;
 		return FVector3d(MinX + CellWidth, MinY + CellWidth, MinZ + CellWidth);
 	}
+	inline FVector3d GetCellCenter(const FSparsePointOctreeCell& Cell) const
+	{
+		return GetCellCenter(Cell.Level, Cell.Index);
+	}
 
 
-	FVector3i PointToIndex(uint32 Level, const FVector3d& Position) const
+	inline FVector3i PointToIndex(uint32 Level, const FVector3d& Position) const
 	{
 		double CellWidth = GetCellWidth(Level);
 		int32 i = (int32)FMathd::Floor(Position.X / CellWidth);
@@ -313,16 +340,19 @@ protected:
 		return FVector3i(i, j, k);
 	}
 
-
-	int ToChildCellIndex(const FSparsePointOctreeCell& Cell, const FVector3d& Position) const
+	int ToChildCellIndex(uint32 Level, const FVector3i& Index, const FVector3d& Position) const
 	{
-		FVector3d Center = GetCellCenter(Cell);
+		FVector3d Center = GetCellCenter(Level, Index);
 		int ChildIndex =
 			((Position.X < Center.X) ? 0 : 1) +
 			((Position.Y < Center.Y) ? 0 : 2) +
 			((Position.Z < Center.Z) ? 0 : 4);
 		return ChildIndex;
 	}
+	int ToChildCellIndex(const FSparsePointOctreeCell& Cell, const FVector3d& Position) const
+	{
+		return ToChildCellIndex(Cell.Level, Cell.Index, Position);
+	}
 
 	bool CellContains(const FSparsePointOctreeCell& Cell, const FVector3d& Position) const
 	{
@@ -341,7 +371,7 @@ protected:
 
 	inline FSparsePointOctreeCell FindCurrentContainingCell(const FVector3d& Position) const;
 
-
+	inline uint32 CreateNewRootCell(FSparsePointOctreeCell NewRootCell, bool bInitializeCellPointList);
 	inline void Insert_NewRoot(int32 PointID, const FVector3d& Position, FSparsePointOctreeCell NewRootCell);
 	inline void Insert_ToCell(int32 PointID, const FVector3d& Position, const FSparsePointOctreeCell& ExistingCell);
 	inline void Insert_NewChildCell(int32 PointID, const FVector3d& Position, int ParentCellID, FSparsePointOctreeCell NewChildCell, int ChildIdx);
@@ -352,19 +382,52 @@ protected:
 
 
 
-bool FSparseDynamicPointOctree3::ContainsPoint(int32 PointID) const
+void FSparseDynamicPointOctree3::ConfigureFromPointCountEstimate(double MaxBoundsDimension, int CountEstimate)
 {
-	return ValidPointIDs.Get(PointID);
+	// These are some rough values collected from some basic profiling. At some point it will
+	// probably be a good idea to do a more comprehensive study and perhaps come up with some
+	// regression formulas. 
+
+	this->MaxTreeDepth = 2;
+	int RootCellDivisions = 6;
+	if (CountEstimate > 500000)
+	{
+		this->MaxTreeDepth = 3;
+		RootCellDivisions = 8;
+	}
+	if (CountEstimate > 5000000)
+	{
+		this->MaxTreeDepth = 4;
+		RootCellDivisions = 10;
+	}
+	if (CountEstimate > 50000000)
+	{
+		this->MaxTreeDepth = 5;
+		RootCellDivisions = 12;
+	}
+	this->RootDimension = MaxBoundsDimension / (double)RootCellDivisions;
+
+	// maybe not necessarily a good estimate? 
+	this->MaxPointsPerCell = 250;
 }
 
 
+bool FSparseDynamicPointOctree3::ContainsPoint(int32 PointID) const
+{
+	return (PointID < PointIDToCellMap.Num() && PointIDToCellMap[PointID] != InvalidCellID);
+}
+
 
 void FSparseDynamicPointOctree3::InsertPoint(int32 PointID, const FVector3d& Position)
 {
-	check(ContainsPoint(PointID) == false);
+	if (ContainsPoint(PointID) == false)
+	{
+		checkSlow(false);
+		return;
+	}
 
 	FSparsePointOctreeCell CurrentCell = FindCurrentContainingCell(Position);
-	check(CurrentCell.Level != FSparsePointOctreeCell::InvalidLevel);
+	checkSlow(CurrentCell.Level != FSparsePointOctreeCell::InvalidLevel);
 
 	// if we found a containing root cell but it doesn't exist, create it and insert
 	if (CurrentCell.Level == 0 && CurrentCell.IsExistingCell() == false)
@@ -373,9 +436,6 @@ void FSparseDynamicPointOctree3::InsertPoint(int32 PointID, const FVector3d& Pos
 		return;
 	}
 
-	// YIKES this currently does max-depth insertion...
-	//   desired behavior is that parent cell accumulates and then splits later!
-
 	int PotentialChildIdx = ToChildCellIndex(CurrentCell, Position);
 	// if current cell does not have this child we might fit there so try it
 	if (CurrentCell.HasChild(PotentialChildIdx) == false)
@@ -400,73 +460,276 @@ void FSparseDynamicPointOctree3::Insert_NewChildCell(int32 PointID, const FVecto
 	int ParentCellID, FSparsePointOctreeCell NewChildCell, int ChildIdx)
 {
 	FSparsePointOctreeCell& OrigParentCell = Cells[ParentCellID];
-	check(OrigParentCell.HasChild(ChildIdx) == false);
+	checkSlow(OrigParentCell.HasChild(ChildIdx) == false);
 
 	NewChildCell.CellID = CellRefCounts.Allocate();
 	Cells.InsertAt(NewChildCell, NewChildCell.CellID);
 
 	PointIDToCellMap.InsertAt(NewChildCell.CellID, PointID);
-	ValidPointIDs.Set(PointID, true);
 
 	CellPointLists.AllocateAt(NewChildCell.CellID);
 	CellPointLists.Insert(NewChildCell.CellID, PointID);
 
 	OrigParentCell.SetChild(ChildIdx, NewChildCell);
 
-	check(CellContains(NewChildCell, Position));
-	check(PointToIndex(NewChildCell.Level, Position) == NewChildCell.Index);
+	checkSlow(CellContains(NewChildCell, Position));
+	checkSlow(PointToIndex(NewChildCell.Level, Position) == NewChildCell.Index);
 }
 
 
 
 void FSparseDynamicPointOctree3::Insert_ToCell(int32 PointID, const FVector3d& Position, const FSparsePointOctreeCell& ExistingCell)
 {
-	check(CellRefCounts.IsValid(ExistingCell.CellID));
+	checkSlow(CellRefCounts.IsValid(ExistingCell.CellID));
 
 	PointIDToCellMap.InsertAt(ExistingCell.CellID, PointID);
-	ValidPointIDs.Set(PointID, true);
 
 	CellPointLists.Insert(ExistingCell.CellID, PointID);
 
-	check(CellContains(ExistingCell, Position));
-	check(PointToIndex(ExistingCell.Level, Position) == ExistingCell.Index);
+	checkSlow(CellContains(ExistingCell, Position));
+	checkSlow(PointToIndex(ExistingCell.Level, Position) == ExistingCell.Index);
 }
 
 
 
 void FSparseDynamicPointOctree3::Insert_NewRoot(int32 PointID, const FVector3d& Position, FSparsePointOctreeCell NewRootCell)
 {
-	check(RootCells.Has(NewRootCell.Index) == false);
+	uint32 NewRootCellID = CreateNewRootCell(NewRootCell, true);
 
-	NewRootCell.CellID = CellRefCounts.Allocate();
-	Cells.InsertAt(NewRootCell, NewRootCell.CellID);
+	PointIDToCellMap.InsertAt(NewRootCellID, PointID);
+	CellPointLists.Insert(NewRootCellID, PointID);
+}
 
-	PointIDToCellMap.InsertAt(NewRootCell.CellID, PointID);
-	ValidPointIDs.Set(PointID, true);
+uint32 FSparseDynamicPointOctree3::CreateNewRootCell(FSparsePointOctreeCell NewRootCell, bool bInitializeCellPointList)
+{
+	checkSlow(RootCells.Has(NewRootCell.Index) == false);
 
+	NewRootCell.CellID = CellRefCounts.Allocate();
+	Cells.InsertAt(NewRootCell, NewRootCell.CellID);
 	uint32* RootCellElem = RootCells.Get(NewRootCell.Index, true);
 	*RootCellElem = NewRootCell.CellID;
 
-	CellPointLists.AllocateAt(NewRootCell.CellID);
-	CellPointLists.Insert(NewRootCell.CellID, PointID);
+	if (bInitializeCellPointList)
+	{
+		CellPointLists.AllocateAt(NewRootCell.CellID);
+	}
+
+	return NewRootCell.CellID;
 }
 
 
 
+void FSparseDynamicPointOctree3::ParallelInsertDensePointSet(int32 MaxPointID, TFunctionRef<FVector3d(int)> GetPositionFunc)
+{
+	// TODO: this current implementation could be adapted to handle non-dense point sets
 
+	// ParallelInsertDensePointSet can currently only be called if the tree is empty
+	if (ensure(Cells.IsEmpty()) == false)
+	{
+		return;
+	}
 
+	// information for an input point
+	struct FPointInfo
+	{
+		int32 PointID;			// ID of this point
+		FVector3i RootCell;		// index of Root cell this point is in
+		uint32 LeafCellIndex;	// index of Leaf Cell this point should go into. This is an index into the RootCell-specific LevelCells array used inside the ParallelFor below.
+	};
+
+	// populate initial list of points
+	TArray<FPointInfo> PointInfos;
+	PointInfos.SetNumUninitialized(MaxPointID);
+	ParallelFor(MaxPointID, [&](int32 k)
+	{
+		PointInfos[k] = FPointInfo{ k, PointToIndex(0, GetPositionFunc(k)), InvalidCellID };
+	});
+	
+	// build a linear list of root cells, and sort all points into a list for each root cell
+	TMap<FVector3i, int> RootCellLinearIndices;
+	TArray<FVector3i> RootCellIndexes;
+	TArray< TUniquePtr<TDynamicVector<int>> > RootCellPointIndices;
+	{
+		//TRACE_CPUPROFILER_EVENT_SCOPE(InitializeRootCellLists);
+		int LinearIndexCounter = 0;
+		for (int32 k = 0; k < MaxPointID; ++k)
+		{
+			int UseLinearIndex = 0;
+			const int* FoundLinearIndex = RootCellLinearIndices.Find(PointInfos[k].RootCell);
+			if (FoundLinearIndex == nullptr)
+			{
+				UseLinearIndex = LinearIndexCounter++;
+				RootCellLinearIndices.Add(PointInfos[k].RootCell, UseLinearIndex);
+				RootCellPointIndices.Add( MakeUnique<TDynamicVector<int>>() );
+				RootCellIndexes.Add(PointInfos[k].RootCell);
+			}
+			else
+			{
+				UseLinearIndex = *FoundLinearIndex;
+			}
+
+			// try doing these Adds in a second per-root-cell ParallelFor? (parallel-scan seems wasteful but might still be faster...)
+			RootCellPointIndices[UseLinearIndex]->Add(k);
+		}
+	}
+
+	// create all the Root Cells
+	{
+		//TRACE_CPUPROFILER_EVENT_SCOPE(CreateRootCells);
+		for (FVector3i RootIndex : RootCellIndexes)
+		{
+			FSparsePointOctreeCell NewRootCell(0, RootIndex);
+			CreateNewRootCell(NewRootCell, false);
+		}
+	}
+
+	// make sure the PointIDToCellMap is large enough for all the PointIDs, then we never have to insert/resize
+	PointIDToCellMap.SetNum(MaxPointID);
+
+	// FParentChildCell represents a child-cell-of-a-parent that needs to be created. We will make lists of these below.
+	// This strategy is maybe over-complicated, and possibly we can just create all the leaf cells and then figure out
+	// the necessary parents that need to exist above it?
+	struct FParentChildCell
+	{
+		FVector3i ParentCellIndex;		// XYZ index of parent cell
+		FVector3i ChildCellIndex;		// XYZ index of child cell that needs to exist (in it's layer)
+		int ChildNum;					// index of child cell in parent cell's 8-children list
+		int NumPoints;					// number of points that will be inserted in this cell (only set for leaf cells)
+		uint32 NewCellID;				// ID of cell after it is created. This is set in the code below
+		bool operator==(const FParentChildCell& OtherCell) const { return ChildCellIndex == OtherCell.ChildCellIndex; }
+	};
+
+	// Parallel iteration over the Root cells. For each Root cell, we iterate over it's points, and for each point, 
+	// descend down to the max depth to figure out which cells need to be created. Then we create those cells, and 
+	// then we insert the points into the leaf cells. 
+	FCriticalSection SharedDataLock, CellPointListsLock;
+	ParallelFor(RootCellIndexes.Num(), [&](int32 Index)
+	{
+		FVector3i RootIndex = RootCellIndexes[Index];
+		uint32* RootCellID = RootCells.Get(RootIndex, false);
+		FSparsePointOctreeCell& RootCell = Cells[*RootCellID];
+		
+		// array of new child cells at each level
+		TArray<TArray<FParentChildCell>> LevelCells;
+		LevelCells.SetNum(MaxTreeDepth);
+
+		// descend the non-existent octree cells mathematically, to determine which cells need to be instantiated
+		TDynamicVector<int>& PointIndices = *RootCellPointIndices[Index];
+		for ( int32 PointIndex : PointIndices)
+		{
+			int32 PointID = PointInfos[PointIndex].PointID;
+			FVector3d Position = GetPositionFunc(PointID);
+
+			int32 CurLevel = 0;
+			int32 LastInsertIndex = 0;
+			FSparsePointOctreeCell CurrentCell = RootCell;
+			while (CurLevel < MaxTreeDepth)
+			{
+				int ChildIndex = ToChildCellIndex(CurrentCell, Position);	
+				FSparsePointOctreeCell NextLevelCell = CurrentCell.MakeChildCell(ChildIndex);
+				LastInsertIndex = LevelCells[CurLevel].AddUnique( FParentChildCell{ CurrentCell.Index, NextLevelCell.Index, ChildIndex, 0, InvalidCellID } );
+				CurLevel++;
+				CurrentCell = NextLevelCell;
+			}
+			// keep track of which leaf cell this point needs to be inserted into, since we computed it in the last iteration of the above loop
+			PointInfos[PointIndex].LeafCellIndex = LastInsertIndex;
+			LevelCells[CurLevel-1][LastInsertIndex].NumPoints++;
+		}
+
+		// sorted CellIDs of all the Leaf Cells we create in the loop below
+		TArray<uint32> LeafCellIDs; 
+
+		// instantiate all the cells at each level. This is a bit messy because at level/depth N+1 we 
+		// do not know our parent Cells/CellIDs, only their indices, and that our parent must be one 
+		// of the cells that were created at level N. So that takes a linear search. Unclear how to
+		// track that mapping, except maybe a TMap?  (note that in profiling this is not remotely a bottleneck...)
+		TArray<FSparsePointOctreeCell*> LastParents, NextParents;
+		LastParents.Add(&RootCell);
+		for (int32 Level = 0; Level < MaxTreeDepth; ++Level)
+		{
+			for (FParentChildCell& ParentChild : LevelCells[Level])
+			{
+				FSparsePointOctreeCell** ParentCell = LastParents.FindByPredicate( [&](FSparsePointOctreeCell* Cell) { return Cell->Index == ParentChild.ParentCellIndex; } );
+				checkSlow(ParentCell != nullptr);
+				checkSlow( (*ParentCell)->HasChild(ParentChild.ChildNum) == false );
+				FSparsePointOctreeCell NewChildCell = (*ParentCell)->MakeChildCell( ParentChild.ChildNum );
+				checkSlow( NewChildCell.Index == ParentChild.ChildCellIndex );
+
+				// only this block accesses data structures shared between root cells
+				SharedDataLock.Lock();
+				NewChildCell.CellID = CellRefCounts.Allocate();
+				Cells.InsertAt(NewChildCell, NewChildCell.CellID);
+				if ( Level == MaxTreeDepth-1 )
+				{
+					LeafCellIDs.Add(NewChildCell.CellID);
+				}
+				SharedDataLock.Unlock();
+
+				(*ParentCell)->SetChild(ParentChild.ChildNum, NewChildCell);
+				ParentChild.NewCellID = NewChildCell.CellID;
+
+				NextParents.Add( &Cells[NewChildCell.CellID] );
+			}
+			Swap(LastParents, NextParents);
+			NextParents.Reset();
+		}
+
+		TArray<FParentChildCell>& LeafCells = LevelCells[MaxTreeDepth-1];
+		int NumLeafCells = LeafCells.Num();
+		TArray<TSparseListSet<int32>::FListHandle> LeafCellLists;
+		LeafCellLists.SetNum(NumLeafCells);
+
+		// initialize the cell-point-lists for each new leaf cell. This has to be thread-safe, but
+		// it will return handles to the lists that can be used in parallel below 
+		// (remember we are in a big ParallelFor over root cells here, that are accessing CellPointLists simultaneously)
+		CellPointListsLock.Lock();
+		for (int32 k = 0; k < NumLeafCells; ++k)
+		{
+			LeafCellLists[k] = CellPointLists.AllocateAt(LeafCellIDs[k]);
+		}
+		CellPointListsLock.Unlock();
+
+		// It is faster to batch-insert points into the leaf cell lists, even though we have 
+		// to build up new arrays for it. So first we collect per-leaf-cell lists, 
+		// then we send those lists to the leaf cells
+		TArray<TArray<int32>> LeafCellPointLists;		// todo could convert to a single array with offsets?
+		LeafCellPointLists.SetNum(NumLeafCells);
+		for (int32 k = 0; k < NumLeafCells; ++k)
+		{
+			LeafCellPointLists.Reserve( LeafCells[k].NumPoints );
+		}
+		for ( int32 PointIndex : PointIndices)
+		{
+			int PointID = PointInfos[PointIndex].PointID;
+			int LeafIndex = PointInfos[PointIndex].LeafCellIndex;
+			LeafCellPointLists[LeafIndex].Add(PointID);
+			PointIDToCellMap[PointID] = LeafCells[LeafIndex].NewCellID;
+		}
+		for ( int32 k = 0; k < NumLeafCells; ++k)
+		{
+			// TODO: doing this step in parallel did not help, but perhaps the entire above block could be done in parallel (with repeat iterations over PointIndices)
+			CellPointLists.SetValues(LeafCellLists[k], LeafCellPointLists[k]);
+		}
+
+	}, EParallelForFlags::Unbalanced);
+
+}
 
 
 void FSparseDynamicPointOctree3::InsertPoint_DynamicExpand(
 	int32 PointID, 
 	TFunctionRef<FVector3d(int)> GetPositionFunc)
 {
-	check(ContainsPoint(PointID) == false);
+	if (ContainsPoint(PointID) == false)
+	{
+		checkSlow(false);
+		return;
+	}
 
 	FVector3d Position = GetPositionFunc(PointID);
 
 	FSparsePointOctreeCell CurrentCell = FindCurrentContainingCell(Position);
-	check(CurrentCell.Level != FSparsePointOctreeCell::InvalidLevel);
+	checkSlow(CurrentCell.Level != FSparsePointOctreeCell::InvalidLevel);
 
 	// if we found a containing root cell but it doesn't exist, create it and insert
 	if (CurrentCell.Level == 0 && CurrentCell.IsExistingCell() == false)
@@ -504,10 +767,10 @@ void FSparseDynamicPointOctree3::InsertPoint_DynamicExpand(
 	};
 
 	// otherwise current cell got too big, so move points to child cells
-	for (int32 MovePointID : CellPointLists.Values(ParentCell->CellID))
+	CellPointLists.Enumerate(ParentCell->CellID, [&](int MovePointID)
 	{
 		InsertToChildLocalFunc(MovePointID);
-	}
+	});
 
 	// we removed all the points from ParentCell
 	CellPointLists.Clear(ParentCell->CellID);
@@ -534,14 +797,22 @@ bool FSparseDynamicPointOctree3::RemovePoint(int32 PointID)
 	}
 
 	PointIDToCellMap[PointID] = InvalidCellID;
-	ValidPointIDs.Set(PointID, false);
 
 	bool bInList = CellPointLists.Remove(CellID, PointID);
-	check(bInList);
+	checkSlow(bInList);
 	return true;
 }
 
 
+void FSparseDynamicPointOctree3::RemovePointUnsafe(int32 PointID)
+{
+	uint32 CellID = PointIDToCellMap[PointID];
+	if (CellID != InvalidCellID)
+	{
+		PointIDToCellMap[PointID] = InvalidCellID;
+		CellPointLists.Remove(CellID, PointID);
+	}
+}
 
 
 
@@ -616,7 +887,7 @@ int32 FSparseDynamicPointOctree3::FindNearestHitPoint(const FRay3d& Ray,
 		const FSparsePointOctreeCell* CurCell = Queue.Pop(false);
 		
 		// process elements
-		for (int PointID : CellPointLists.Values(CurCell->CellID))
+		CellPointLists.Enumerate(CurCell->CellID, [&](int32 PointID)
 		{
 			double HitDist = HitPointDistFunc(PointID, Ray);
 			if (HitDist < MaxDistance)
@@ -624,7 +895,7 @@ int32 FSparseDynamicPointOctree3::FindNearestHitPoint(const FRay3d& Ray,
 				MaxDistance = HitDist;
 				HitPointID = PointID;
 			}
-		}
+		});
 
 		// descend to child cells
 		// sort by distance? use DDA?
@@ -657,37 +928,32 @@ void FSparseDynamicPointOctree3::RangeQuery(
 	TArray<int>& PointIDs,
 	TArray<const FSparsePointOctreeCell*>* TempBuffer) const
 {
-	// todo: this should take advantage of raster!
-
 	TArray<const FSparsePointOctreeCell*> InternalBuffer;
 	TArray<const FSparsePointOctreeCell*>& Queue =
 		(TempBuffer == nullptr) ? InternalBuffer : *TempBuffer;
 	Queue.Reset();
 	Queue.Reserve(128);
 
-	// start at root cells
-	RootCells.AllocatedIteration([&](const uint32* RootCellID)
+	FVector3i RootMinIndex = PointToIndex(0, Bounds.Min);
+	FVector3i RootMaxIndex = PointToIndex(0, Bounds.Max);
+	RootCells.RangeIteration(RootMinIndex, RootMaxIndex, [&](uint32 RootCellID)
 	{
-		const FSparsePointOctreeCell* RootCell = &Cells[*RootCellID];
-		if (GetCellBox(*RootCell).Intersects(Bounds))
-		{
-			Queue.Add(&Cells[*RootCellID]);
-		}
+		const FSparsePointOctreeCell* RootCell = &Cells[RootCellID];
+		Queue.Add(RootCell);
 	});
 
-
 	while (Queue.Num() > 0)
 	{
 		const FSparsePointOctreeCell* CurCell = Queue.Pop(false);
 
 		// process elements
-		for (int PointID : CellPointLists.Values(CurCell->CellID))
+		CellPointLists.Enumerate(CurCell->CellID, [&](int32 PointID)
 		{
 			if (PredicateFunc(PointID))
 			{
 				PointIDs.Add(PointID);
 			}
-		}
+		});
 
 		for (int k = 0; k < 8; ++k)
 		{
@@ -708,32 +974,28 @@ void FSparseDynamicPointOctree3::RangeQuery(
 
 
 
-
-
-
 void FSparseDynamicPointOctree3::ParallelRangeQuery(
 	const FAxisAlignedBox3d& Bounds,
 	TFunctionRef<bool(int)> PredicateFunc,
 	TArray<int>& PointIDs,
 	TArray<const FSparsePointOctreeCell*>* TempBuffer) const
 {
-	// todo: this should take advantage of raster!
+	TArray<const FSparsePointOctreeCell*> InternalBuffer;
+	TArray<const FSparsePointOctreeCell*>& ProcessRootCells =
+		(TempBuffer == nullptr) ? InternalBuffer : *TempBuffer;
+	ProcessRootCells.Reset();
+	ProcessRootCells.Reserve(128);
 
-	TArray<const FSparsePointOctreeCell*> ProcessRootCells;
-	// start at root cells
-	RootCells.AllocatedIteration([&](const uint32* RootCellID)
+	FCriticalSection OutputLock;
+
+	FVector3i RootMinIndex = PointToIndex(0, Bounds.Min);
+	FVector3i RootMaxIndex = PointToIndex(0, Bounds.Max);
+	RootCells.RangeIteration(RootMinIndex, RootMaxIndex, [&](uint32 RootCellID)
 	{
-		const FSparsePointOctreeCell* RootCell = &Cells[*RootCellID];
-		if (GetCellBox(*RootCell).Intersects(Bounds))
-		{
-			ProcessRootCells.Add(&Cells[*RootCellID]);
-		}
+		const FSparsePointOctreeCell* RootCell = &Cells[RootCellID];
+		ProcessRootCells.Add(RootCell);
 	});
 
-	FCriticalSection OutputLock;
-	TArray<int> PerThreadCounts;
-	TArray<int> PerThreadTraversed;
-
 	ParallelFor(ProcessRootCells.Num(), [&](int idx)
 	{
 		const FSparsePointOctreeCell* RootCell = ProcessRootCells[idx];
@@ -741,26 +1003,24 @@ void FSparseDynamicPointOctree3::ParallelRangeQuery(
 		Queue.Add(RootCell);
 
 		TArray<int, TInlineAllocator<256>> LocalPointIDs;
-		int Processed = 0;
 
 		while (Queue.Num() > 0)
 		{
 			const FSparsePointOctreeCell* CurCell = Queue.Pop(false);
 
 			// process elements
-			for (int PointID : CellPointLists.Values(CurCell->CellID))
+			CellPointLists.Enumerate(CurCell->CellID, [&](int PointID)
 			{
 				if (PredicateFunc(PointID))
 				{
 					LocalPointIDs.Add(PointID);
 				}
-			}
+			});
 
 			for (int k = 0; k < 8; ++k)
 			{
 				if (CurCell->HasChild(k))
 				{
-					Processed++;
 					const FSparsePointOctreeCell* ChildCell = &Cells[CurCell->GetChildCellID(k)];
 					if (GetCellBox(*ChildCell).Intersects(Bounds))
 					{
@@ -777,11 +1037,8 @@ void FSparseDynamicPointOctree3::ParallelRangeQuery(
 			{
 				PointIDs.Add(ObjID);
 			}
-			PerThreadCounts.Add(LocalPointIDs.Num());
-			PerThreadTraversed.Add(Processed);
 			OutputLock.Unlock();
 		}
-
 	});
 
 
@@ -802,7 +1059,7 @@ FSparsePointOctreeCell FSparseDynamicPointOctree3::FindCurrentContainingCell(con
 	{
 		return FSparsePointOctreeCell(0, RootIndex);
 	}
-	check(CellRefCounts.IsValid(*RootCellID));
+	checkSlow(CellRefCounts.IsValid(*RootCellID));
 
 	// check if point is contained in root cell
 	// (should we do this before checking for existence? we can...)
@@ -819,7 +1076,7 @@ FSparsePointOctreeCell FSparseDynamicPointOctree3::FindCurrentContainingCell(con
 		if (CurrentCell->HasChild(ChildIdx))
 		{
 			int32 ChildCellID = CurrentCell->GetChildCellID(ChildIdx);
-			check(CellRefCounts.IsValid(ChildCellID));
+			checkSlow(CellRefCounts.IsValid(ChildCellID));
 			const FSparsePointOctreeCell* ChildCell = &Cells[ChildCellID];
 			if (CellContains(*ChildCell, Position))
 			{
@@ -875,10 +1132,10 @@ void FSparseDynamicPointOctree3::CheckValidity(
 	// check that all Point IDs in per-cell Point lists is valid
 	for (int32 CellID : CellRefCounts.Indices())
 	{
-		for (int32 PointID : CellPointLists.Values(CellID))
+		CellPointLists.Enumerate(CellID, [&](int32 PointID)
 		{
 			CheckOrFailF(IsValidPointIDFunc(PointID));
-		}
+		});
 	}
 
 	uint32 NumPointIDs = PointIDToCellMap.Num();
@@ -924,10 +1181,10 @@ void FSparseDynamicPointOctree3::CheckValidity(
 
 	if (bVerbose)
 	{
-		UE_LOG(LogTemp, Warning, TEXT("FSparseDynamicPointOctree3::CheckValidity: MaxLevel %d  MissingCount %d"), MaxLevel, MissingPointCount);
+		UE_LOG(LogGeometry, Warning, TEXT("FSparseDynamicPointOctree3::CheckValidity: MaxLevel %d  MissingCount %d"), MaxLevel, MissingPointCount);
 		for (uint32 k = 0; k <= MaxLevel; ++k)
 		{
-			UE_LOG(LogTemp, Warning, TEXT("    Level %4d  Cells %4d  Points %4d"), k, CellsAtLevels[k], PointsAtLevel[k]);
+			UE_LOG(LogGeometry, Warning, TEXT("    Level %4d  Cells %4d  Points %4d"), k, CellsAtLevels[k], PointsAtLevel[k]);
 		}
 	}
 }
@@ -952,7 +1209,7 @@ void FSparseDynamicPointOctree3::ComputeStatistics(FStatistics& StatsOut) const
 	{
 		const FSparsePointOctreeCell& Cell = Cells[CellID];
 		StatsOut.LevelBoxCounts[Cell.Level]++;
-		int ObjCount = CellPointLists.GetCount(CellID);
+		int ObjCount = CellPointLists.IsAllocated(CellID) ? CellPointLists.GetCount(CellID) : 0;
 		StatsOut.LevelObjCounts[Cell.Level] += ObjCount;
 		StatsOut.LevelMaxObjCounts[Cell.Level] = FMath::Max(StatsOut.LevelMaxObjCounts[Cell.Level], ObjCount);
 		StatsOut.LevelCellSizes[Cell.Level] = float(GetCellWidth(Cell.Level));
diff --git a/Engine/Source/Runtime/GeometryCore/Public/Spatial/SparseGrid3.h b/Engine/Source/Runtime/GeometryCore/Public/Spatial/SparseGrid3.h
--- a/Engine/Source/Runtime/GeometryCore/Public/Spatial/SparseGrid3.h
+++ b/Engine/Source/Runtime/GeometryCore/Public/Spatial/SparseGrid3.h
@@ -94,7 +94,7 @@ public:
 	 * @param bAllocateIfMissing if the element at this index is null, allocate a new one
 	 * @return pointer to ElemType instance, or nullptr if element doesn't exist
 	 */
-	ElemType* Get(const FVector3i& Index, bool bAllocateIfMissing = true)
+	ElemType* Get(const FVector3i& Index, bool bAllocateIfMissing)
 	{
 		ElemType* result = Elements.FindRef(Index);
 		if (result != nullptr)
@@ -170,7 +170,7 @@ public:
 
 
 	/**
-	 * Iterate over existing elements and apply ElementFunc to each of them
+	 * Iterate over existing elements and apply ElementFunc(ElementType) to each of them
 	 */
 	template<typename Func>
 	void AllocatedIteration(Func ElementFunc) const
@@ -182,6 +182,31 @@ public:
 	}
 
 
+	/**
+	 * Iterate over existing allocated elements within the integer bounds defined 
+	 * by Min and Max (inclusive), and apply ElementFunc(ElementType) to each of them
+	 */
+	template<typename Func>
+	void RangeIteration(FVector3i MinIndex, FVector3i MaxIndex, Func ElementFunc) const
+	{
+		for (int32 zi = MinIndex.Z; zi <= MaxIndex.Z; ++zi)
+		{
+			for (int32 yi = MinIndex.Y; yi <= MaxIndex.Y; ++yi)
+			{
+				for (int32 xi = MinIndex.X; xi <= MaxIndex.X; ++xi)
+				{
+					const ElemType* Found = Elements.FindRef(FVector3i(xi,yi,zi));
+					if (Found != nullptr)
+					{
+						ElementFunc(*Found);
+					}
+				}
+			}
+		}
+	}
+
+
+
 protected:
 
 	ElemType* Allocate(const FVector3i& index)
diff --git a/Engine/Source/Runtime/GeometryCore/Public/Util/SparseListSet.h b/Engine/Source/Runtime/GeometryCore/Public/Util/SparseListSet.h
--- a/Engine/Source/Runtime/GeometryCore/Public/Util/SparseListSet.h
+++ b/Engine/Source/Runtime/GeometryCore/Public/Util/SparseListSet.h
@@ -42,7 +42,7 @@ public:
 	TSparseListSet()
 	{
 		ChunkList.Add(MakeUnique<FChunk>());
-		ChunkList[0]->ElementBuffer.SetNumUninitialized(BlocksPerChunk * BlockSize);
+		ChunkList[0]->Initialize(BlocksPerChunk * BlockSize);
 	}
 
 	/**
