GeometryCore: add FMeshSurfacePointSampling, which generates sample points on the surface of a mesh with various strategies (currently focused on an approximate-Poisson sampling strategy).
GeometryScript: add ComputePointSampling, ComputeNonUniformPointSampling, and ComputeVertexWeightedPointSampling functions

diff --git a/Engine/Plugins/Experimental/GeometryScripting/Source/GeometryScriptingCore/Private/MeshSamplingFunctions.cpp b/Engine/Plugins/Experimental/GeometryScripting/Source/GeometryScriptingCore/Private/MeshSamplingFunctions.cpp
new file mode 100644
--- /dev/null
+++ b/Engine/Plugins/Experimental/GeometryScripting/Source/GeometryScriptingCore/Private/MeshSamplingFunctions.cpp
@@ -0,0 +1,149 @@
+// Copyright Epic Games, Inc. All Rights Reserved.
+
+#include "GeometryScript/MeshSamplingFunctions.h"
+#include "UDynamicMesh.h"
+#include "DynamicMesh/DynamicMesh3.h"
+#include "Sampling/MeshSurfacePointSampling.h"
+
+#include UE_INLINE_GENERATED_CPP_BY_NAME(MeshSamplingFunctions)
+
+using namespace UE::Geometry;
+
+#define LOCTEXT_NAMESPACE "UGeometryScriptLibrary_MeshSamplingFunctions"
+
+
+UDynamicMesh* UGeometryScriptLibrary_MeshSamplingFunctions::ComputePointSampling(
+	UDynamicMesh* TargetMesh,
+	FGeometryScriptMeshPointSamplingOptions Options,
+	TArray<FTransform>& Samples,
+	UGeometryScriptDebug* Debug)
+{
+	if (TargetMesh == nullptr)
+	{
+		UE::Geometry::AppendError(Debug, EGeometryScriptErrorType::InvalidInputs, LOCTEXT("ComputePointSampling", "ComputePointSampling: TargetMesh is Null"));
+		return TargetMesh;
+	}
+
+	TRACE_CPUPROFILER_EVENT_SCOPE(GS_ComputePointSampling);
+
+	FMeshSurfacePointSampling Sampler;
+	Sampler.SampleRadius = Options.SamplingRadius;
+	Sampler.MaxSamples = Options.MaxNumSamples;
+	Sampler.SubSampleDensity = Options.SubSampleDensity;
+	Sampler.RandomSeed = Options.RandomSeed;
+
+	TargetMesh->ProcessMesh([&](const FDynamicMesh3& ReadMesh)
+	{
+		Sampler.ComputePoissonSampling(ReadMesh);
+	});
+
+	for (FFrame3d Frame : Sampler.Samples)
+	{
+		Samples.Add( Frame.ToFTransform() );
+	}
+
+	return TargetMesh;
+}
+
+
+UDynamicMesh* UGeometryScriptLibrary_MeshSamplingFunctions::ComputeNonUniformPointSampling(
+	UDynamicMesh* TargetMesh,
+	FGeometryScriptMeshPointSamplingOptions Options,
+	FGeometryScriptNonUniformPointSamplingOptions NonUniformOptions,
+	TArray<FTransform>& Samples,
+	TArray<double>& SampleRadii,
+	UGeometryScriptDebug* Debug)
+{
+	if (TargetMesh == nullptr)
+	{
+		UE::Geometry::AppendError(Debug, EGeometryScriptErrorType::InvalidInputs, LOCTEXT("ComputeNonUniformPointSampling", "ComputeNonUniformPointSampling: TargetMesh is Null"));
+		return TargetMesh;
+	}
+
+	TRACE_CPUPROFILER_EVENT_SCOPE(GS_ComputeNonUniformPointSampling);
+
+	FMeshSurfacePointSampling Sampler;
+	Sampler.SampleRadius = Options.SamplingRadius;
+	Sampler.MaxSamples = Options.MaxNumSamples;
+	Sampler.SubSampleDensity = Options.SubSampleDensity;
+	Sampler.RandomSeed = Options.RandomSeed;
+	if (NonUniformOptions.MaxSamplingRadius > Options.SamplingRadius)
+	{
+		Sampler.MaxSampleRadius = NonUniformOptions.MaxSamplingRadius;
+		Sampler.SizeDistribution = (FMeshSurfacePointSampling::ESizeDistribution)(int)(NonUniformOptions.SizeDistribution);
+		Sampler.SizeDistributionPower = FMath::Clamp(NonUniformOptions.SizeDistributionPower, 1.0, 10.0);
+	}
+
+	TargetMesh->ProcessMesh([&](const FDynamicMesh3& ReadMesh)
+	{
+		Sampler.ComputePoissonSampling(ReadMesh);
+	});
+
+	for ( int32 k = 0; k < Sampler.Samples.Num(); ++k)
+	{
+		Samples.Add( Sampler.Samples[k].ToFTransform());
+		SampleRadii.Add( Sampler.Radii[k] );
+	}
+
+	return TargetMesh;
+}
+
+
+
+
+
+UDynamicMesh* UGeometryScriptLibrary_MeshSamplingFunctions::ComputeVertexWeightedPointSampling(
+	UDynamicMesh* TargetMesh,
+	FGeometryScriptMeshPointSamplingOptions Options,
+	FGeometryScriptNonUniformPointSamplingOptions NonUniformOptions,
+	FGeometryScriptScalarList VertexWeights,
+	TArray<FTransform>& Samples,
+	TArray<double>& SampleRadii,
+	UGeometryScriptDebug* Debug)
+{
+	if (TargetMesh == nullptr)
+	{
+		UE::Geometry::AppendError(Debug, EGeometryScriptErrorType::InvalidInputs, LOCTEXT("ComputeVertexWeightedPointSampling_InvalidInput", "ComputeVertexWeightedPointSampling: TargetMesh is Null"));
+		return TargetMesh;
+	}
+
+	TRACE_CPUPROFILER_EVENT_SCOPE(GS_ComputeVertexWeightedPointSampling);
+
+	FMeshSurfacePointSampling Sampler;
+	Sampler.SampleRadius = Options.SamplingRadius;
+	Sampler.MaxSamples = Options.MaxNumSamples;
+	Sampler.SubSampleDensity = Options.SubSampleDensity;
+	Sampler.RandomSeed = Options.RandomSeed;
+	bool bIsNonUniform = false;
+	if (NonUniformOptions.MaxSamplingRadius > Options.SamplingRadius)
+	{
+		bIsNonUniform = true;
+		Sampler.MaxSampleRadius = NonUniformOptions.MaxSamplingRadius;
+		Sampler.SizeDistribution = (FMeshSurfacePointSampling::ESizeDistribution)(int)(NonUniformOptions.SizeDistribution);
+		Sampler.SizeDistributionPower = FMath::Clamp(NonUniformOptions.SizeDistributionPower, 1.0, 10.0);
+	}
+
+	TargetMesh->ProcessMesh([&](const FDynamicMesh3& ReadMesh)
+	{
+		if ( bIsNonUniform && VertexWeights.List.IsValid() && VertexWeights.List->Num() == ReadMesh.MaxVertexID() )
+		{
+			Sampler.VertexWeights = *VertexWeights.List.Get();
+			Sampler.bUseVertexWeights = true;
+			Sampler.InterpretWeightMode = (FMeshSurfacePointSampling::EInterpretWeightMode)(int)(NonUniformOptions.WeightMode);
+			Sampler.bInvertWeights = NonUniformOptions.bInvertWeights;
+		}
+		Sampler.ComputePoissonSampling(ReadMesh);
+	});
+
+	for ( int32 k = 0; k < Sampler.Samples.Num(); ++k)
+	{
+		Samples.Add( Sampler.Samples[k].ToFTransform());
+		SampleRadii.Add( Sampler.Radii[k] );
+	}
+
+	return TargetMesh;
+}
+
+
+
+#undef LOCTEXT_NAMESPACE
\ No newline at end of file
diff --git a/Engine/Plugins/Experimental/GeometryScripting/Source/GeometryScriptingCore/Public/GeometryScript/MeshSamplingFunctions.h b/Engine/Plugins/Experimental/GeometryScripting/Source/GeometryScriptingCore/Public/GeometryScript/MeshSamplingFunctions.h
new file mode 100644
--- /dev/null
+++ b/Engine/Plugins/Experimental/GeometryScripting/Source/GeometryScriptingCore/Public/GeometryScript/MeshSamplingFunctions.h
@@ -0,0 +1,147 @@
+// Copyright Epic Games, Inc. All Rights Reserved.
+
+#pragma once
+
+#include "CoreMinimal.h"
+#include "Kismet/BlueprintFunctionLibrary.h"
+#include "GeometryScript/GeometryScriptTypes.h"
+#include "MeshSamplingFunctions.generated.h"
+
+class UDynamicMesh;
+
+
+USTRUCT(BlueprintType)
+struct GEOMETRYSCRIPTINGCORE_API FGeometryScriptMeshPointSamplingOptions
+{
+	GENERATED_BODY()
+public:
+	/** Desired "radius" of sample points. Spacing between samples is at least 2x this value. */
+	UPROPERTY(EditAnywhere, BlueprintReadWrite, Category = Options)
+	float SamplingRadius = 10.0;
+
+	/** Maximum number of samples requested. If 0 or default value, mesh will be maximally sampled */
+	UPROPERTY(EditAnywhere, BlueprintReadWrite, Category = Options)
+	int MaxNumSamples = 0;
+
+	/** Random Seed used to initialize sampling strategies */
+	UPROPERTY(EditAnywhere, BlueprintReadWrite, Category = Options)
+	int RandomSeed = 0;
+
+	/** Density of subsampling used in Poisson strategy. Larger numbers mean "more accurate" (but slower) results. */
+	UPROPERTY(EditAnywhere, BlueprintReadWrite, Category = Options)
+	double SubSampleDensity = 10;
+};
+
+/** Control how active Weights are used to affect point radius */
+UENUM(BlueprintType)
+enum class EGeometryScriptSamplingWeightMode : uint8
+{
+	/** 
+	 * Weights are clamped to [0,1] and used to interpolate Min/Max Radius. This is a "hard constraint", ie if the weight
+	 * at a point is 1, only a "max radius" sample may be placed there, otherwise no samples at all (so no "filling in" smaller samples between large ones)
+	 */
+	WeightToRadius = 0,
+	/** 
+	 * Weights are clamped to [0,1] and used to interpolate Min/Max Radius, with decay, so that smaller-radius samples will infill between large ones.
+	 * So areas with large weight may still end up with some variable-radius samples, but areas with 0 weight will only ever have min-radius samples.
+	 */
+	FilledWeightToRadius = 1,
+	/** 
+	 * Weight is used to create nonuniform random sampling, ie it nudges the random sample-radius distribution but does not directly control it.
+	 * So samples with any radius can still appear at any location, but if weight=1 then max-radius samples are more likely, etc.
+	 */
+	WeightedRandom = 2
+};
+
+/** Controls the distribution of sample radii */
+UENUM(BlueprintType)
+enum class EGeometryScriptSamplingDistributionMode : uint8
+{
+	/** Uniform distribution of sizes, ie all equally likely */
+	Uniform = 0,
+	/** Distribution is weighted towards smaller points */
+	Smaller = 1,
+	/** Distribution is weighted towards larger points */
+	Larger = 2
+};
+
+
+USTRUCT(BlueprintType)
+struct GEOMETRYSCRIPTINGCORE_API FGeometryScriptNonUniformPointSamplingOptions
+{
+	GENERATED_BODY()
+public:
+	/** If MaxSampleRadius > SampleRadius, then output sample radius will be in range [SampleRadius, MaxSampleRadius] */
+	UPROPERTY(EditAnywhere, BlueprintReadWrite, Category = Options)
+	float MaxSamplingRadius = 0.0;
+
+	/** SizeDistribution setting controls the distribution of sample radii */
+	UPROPERTY(EditAnywhere, BlueprintReadWrite, Category = Options)
+	EGeometryScriptSamplingDistributionMode SizeDistribution = EGeometryScriptSamplingDistributionMode::Uniform;
+
+	/** SizeDistributionPower is used to control how extreme the Size Distribution shift is. Valid range is [1,10] */
+	UPROPERTY(EditAnywhere, BlueprintReadWrite, Category = Options, Meta = (UIMin=1, UIMax=10, EditCondition ="(SizeDistribution != EGeometryScriptSamplingDistributionMode::Uniform)" ))
+	double SizeDistributionPower = 2.0;
+
+	/** WeightMode controls how any active Weight scheme is used to affect sample radius */
+	UPROPERTY(EditAnywhere, BlueprintReadWrite, Category = Options)
+	EGeometryScriptSamplingWeightMode WeightMode = EGeometryScriptSamplingWeightMode::WeightedRandom;
+
+	/** If true, weight values are inverted */
+	UPROPERTY(EditAnywhere, BlueprintReadWrite, Category = Options)
+	bool bInvertWeights = false;
+};
+
+
+UCLASS(meta = (ScriptName = "GeometryScript_MeshSampling"))
+class GEOMETRYSCRIPTINGCORE_API UGeometryScriptLibrary_MeshSamplingFunctions : public UBlueprintFunctionLibrary
+{
+	GENERATED_BODY()
+public:
+
+	/**
+	 * Compute a set of sample points lying on the surface of TargetMesh based on the provided sampling Options.
+	 * Samples are approximately uniformly distributed, and non-overlapping relative to the provided Options.SamplingRadius,
+	 * ie the distance between any pair of samples if >= 2*SamplingRadius.
+	 */
+	UFUNCTION(BlueprintCallable, Category = "GeometryScript|MeshSampling", meta=(ScriptMethod))
+	static UPARAM(DisplayName = "Target Mesh") UDynamicMesh* 
+	ComputePointSampling( 
+		UDynamicMesh* TargetMesh, 
+		FGeometryScriptMeshPointSamplingOptions Options,
+		TArray<FTransform>& Samples,
+		UGeometryScriptDebug* Debug = nullptr);
+
+	/**
+	 * Compute a set of sample points lying on the surface of TargetMesh based on the provided sampling Options and NonUniformOptions.
+	 * The sample points have radii in the range [Options.SamplingRadius, NonUniformOptions.MaxSamplingRadius], and
+	 * are non-overlapping, ie the distance between two points is always larger than the sum of their respective radii.
+	 */
+	UFUNCTION(BlueprintCallable, Category = "GeometryScript|MeshSampling", meta=(ScriptMethod))
+	static UPARAM(DisplayName = "Target Mesh") UDynamicMesh* 
+	ComputeNonUniformPointSampling( 
+		UDynamicMesh* TargetMesh, 
+		FGeometryScriptMeshPointSamplingOptions Options,
+		FGeometryScriptNonUniformPointSamplingOptions NonUniformOptions,
+		TArray<FTransform>& Samples,
+		TArray<double>& SampleRadii,
+		UGeometryScriptDebug* Debug = nullptr);
+
+	/**
+	 * Compute a set of sample points lying on the surface of TargetMesh based on the provided sampling Options and NonUniformOptions.
+	 * The sample points have radii in the range [Options.SamplingRadius, NonUniformOptions.MaxSamplingRadius], and
+	 * are non-overlapping, ie the distance between two points is always larger than the sum of their respective radii.
+	 * @param VertexWeights defines a per-vertex weight in range [0,1], these are interpolated to create a scalar field over the mesh triangles which is used to weight the sampling radii
+	 */
+	UFUNCTION(BlueprintCallable, Category = "GeometryScript|MeshSampling", meta=(ScriptMethod))
+	static UPARAM(DisplayName = "Target Mesh") UDynamicMesh* 
+	ComputeVertexWeightedPointSampling( 
+		UDynamicMesh* TargetMesh, 
+		FGeometryScriptMeshPointSamplingOptions Options,
+		FGeometryScriptNonUniformPointSamplingOptions NonUniformOptions,
+		FGeometryScriptScalarList VertexWeights,
+		TArray<FTransform>& Samples,
+		TArray<double>& SampleRadii,
+		UGeometryScriptDebug* Debug = nullptr);
+
+};
\ No newline at end of file
diff --git a/Engine/Source/Runtime/GeometryCore/Private/Sampling/MeshSurfacePointSampling.cpp b/Engine/Source/Runtime/GeometryCore/Private/Sampling/MeshSurfacePointSampling.cpp
new file mode 100644
--- /dev/null
+++ b/Engine/Source/Runtime/GeometryCore/Private/Sampling/MeshSurfacePointSampling.cpp
@@ -0,0 +1,680 @@
+// Copyright Epic Games, Inc. All Rights Reserved.
+
+#include "Sampling/MeshSurfacePointSampling.h"
+
+#include "VectorUtil.h"
+#include "DynamicMesh/DynamicMesh3.h"
+#include "DynamicMesh/DynamicVerticesOctree3.h"
+#include "Async/ParallelFor.h"
+#include "Tasks/Task.h"
+#include "Math/RandomStream.h"
+
+using namespace UE::Geometry;
+
+namespace UELocal
+{
+
+
+// cache of per-triangle information for a triangle mesh
+struct FTriangleInfoCache
+{
+	TArray<FVector3d> TriNormals;
+	TArray<double> TriAreas;
+	TArray<FFrame3d> TriFrames;
+	TArray<FTriangle2d> UVTriangles;
+
+	double TotalArea;
+
+	template<typename MeshType>
+	void InitializeForTriangleSet(const MeshType& SampleMesh)
+	{
+		TriNormals.SetNumUninitialized(SampleMesh.MaxTriangleID());
+		TriAreas.SetNumUninitialized(SampleMesh.MaxTriangleID());
+		TriFrames.SetNumUninitialized(SampleMesh.MaxTriangleID());
+		UVTriangles.SetNumUninitialized(SampleMesh.MaxTriangleID());
+				
+		ParallelFor(SampleMesh.MaxTriangleID(), [&](int32 tid)
+		{
+			if (SampleMesh.IsTriangle(tid))
+			{
+				FVector3d A, B, C, Centroid;
+				SampleMesh.GetTriVertices(tid, A, B, C);
+				Centroid = (A + B + C) / 3.0;
+				TriNormals[tid] = VectorUtil::NormalArea(A, B, C, TriAreas[tid]);
+				TriFrames[tid] = FFrame3d(Centroid, TriNormals[tid]);
+				UVTriangles[tid] = FTriangle2d(
+					TriFrames[tid].ToPlaneUV(A), TriFrames[tid].ToPlaneUV(B), TriFrames[tid].ToPlaneUV(C));
+			}
+		});
+		TotalArea = 0;
+		for (double TriArea : TriAreas)
+		{
+			TotalArea += TriArea;
+		}
+	}
+};
+
+
+struct FNonUniformSamplingConfig
+{
+	FMeshSurfacePointSampling::ESizeDistribution SizeDistribution = FMeshSurfacePointSampling::ESizeDistribution::Uniform;
+	double SizeDistributionPower = 2.0;
+
+	TOptional<TFunctionRef<double(int TriangleID, FVector3d Position, FVector3d BaryCoords)>> WeightFunction;
+
+	FMeshSurfacePointSampling::EInterpretWeightMode InterpretWeightMode = FMeshSurfacePointSampling::EInterpretWeightMode::WeightedRandom;
+};
+
+
+struct FDenseSamplePointSet
+{
+	TArray<FVector3d> DensePoints;
+	TArray<int> Triangles;
+	TArray<double> Weights;
+
+	FAxisAlignedBox3d Bounds;
+	// orientation array
+	int MaxVertexID() const { return DensePoints.Num(); }
+	bool IsVertex(int VertexID) const { return true; }
+	FVector3d GetVertex(int Index) const { return DensePoints[Index]; }
+	FAxisAlignedBox3d GetBounds() const { return Bounds; }
+};
+
+
+struct FPerTriangleDensePointSampling
+{
+	struct FTriangleSubArray
+	{
+		int32 StartIndex;
+		int32 NumSamples;
+	};
+	TArray<FTriangleSubArray> TriSubArrays;
+
+	template<typename MeshType>
+	void InitializeForTriangleSet(
+		const MeshType& SampleMesh, 
+		const FTriangleInfoCache& TriInfo, 
+		double DenseSampleArea, 
+		int RandomSeed,
+		FDenseSamplePointSet& PointSetOut)
+	{
+	}
+
+	template<typename MeshType>
+	void InitializeForTriangleSet(
+		const MeshType& SampleMesh, 
+		const FTriangleInfoCache& TriInfo, 
+		double DenseSampleArea, 
+		int RandomSeed,
+		const FNonUniformSamplingConfig& NonUniformConfig,
+		FDenseSamplePointSet& PointSetOut )
+	{
+		// figure out how many samples in each triangle, and assign each triangle a starting index into PointSetOut arrays
+		TriSubArrays.SetNumUninitialized(SampleMesh.MaxTriangleID());
+		int32 CurIndex = 0;
+		for (int32 tid = 0; tid < SampleMesh.MaxTriangleID(); ++tid)
+		{
+			if ( SampleMesh.IsTriangle(tid) )
+			{
+				int NumSamples = FMath::Max( (int)(TriInfo.TriAreas[tid] / DenseSampleArea), 2 );  // a bit arbitrary...
+				TriSubArrays[tid] = FTriangleSubArray{CurIndex, NumSamples};
+				CurIndex += NumSamples;
+			}
+			else
+			{
+				TriSubArrays[tid] = FTriangleSubArray{-1, 0};
+			}
+		}
+		PointSetOut.DensePoints.SetNumUninitialized(CurIndex);
+		PointSetOut.Triangles.SetNumUninitialized(CurIndex);
+
+		bool bComputeWeights = (NonUniformConfig.WeightFunction.IsSet());
+		if (bComputeWeights)
+		{
+			PointSetOut.Weights.SetNumUninitialized(CurIndex);
+		}
+
+		// This parallel for seems to be quite expensive...maybe contention because all threads are writing to same DensePoints array?
+		// Amount of work per TriangleID can vary quite a bit because it depends on triangle size...possibly should pass
+		// unbalanced flag if triangle count is small relative to point count
+		EParallelForFlags UseFlags = (SampleMesh.MaxTriangleID() < 100) ? EParallelForFlags::Unbalanced : EParallelForFlags::None;
+		ParallelFor(TEXT("ComputeDenseSamples"), SampleMesh.MaxTriangleID(), 500, [&](int tid)
+		{
+			int NumSamples = TriSubArrays[tid].NumSamples;
+			if ( NumSamples == 0 ) return;
+			int StartIndex = TriSubArrays[tid].StartIndex;
+
+			const FFrame3d& ProjectFrame = TriInfo.TriFrames[tid];
+			const FTriangle2d& TriUV = TriInfo.UVTriangles[tid];
+
+			// generate uniform random point in 2D quadrilateral  (http://mathworld.wolfram.com/TrianglePointPicking.html)
+			FVector2d V1 = TriUV.V[1] - TriUV.V[0];
+			FVector2d V2 = TriUV.V[2] - TriUV.V[0];
+
+			int NumGenerated = 0;
+			FRandomStream RandomStream(tid + RandomSeed);
+			while (NumGenerated < NumSamples)
+			{
+				double a1 = RandomStream.GetFraction();
+				double a2 = RandomStream.GetFraction();
+				FVector2d PointUV = TriUV.V[0] + a1 * V1 + a2 * V2;
+				if (TriUV.IsInside(PointUV))
+				{
+					FVector3d Position = ProjectFrame.FromPlaneUV(PointUV, 2);
+					PointSetOut.DensePoints[StartIndex+NumGenerated] = Position;
+					PointSetOut.Triangles[StartIndex+NumGenerated] = tid;
+
+					if (bComputeWeights)
+					{
+						FVector3d BaryCoords = TriUV.GetBarycentricCoords(PointUV);
+						PointSetOut.Weights[StartIndex+NumGenerated] = 
+							NonUniformConfig.WeightFunction.GetValue()(tid, Position, BaryCoords);
+					}
+
+					NumGenerated++;
+				}
+			}
+		}, UseFlags);
+	}
+
+};
+
+template<typename MeshType>
+void ConstructDenseUniformMeshPointSampling(
+	const MeshType& SampleMesh,
+	double SampleRadius,
+	double SubSampleDensity,
+	int RandomSeed,
+	const FNonUniformSamplingConfig& NonUniformConfig,
+	int MaxNumDenseSamples,
+	FDenseSamplePointSet& DensePointSetOut )
+{
+	FTriangleInfoCache TriInfoCache;
+	TriInfoCache.InitializeForTriangleSet<MeshType>(SampleMesh);
+
+	// compute mesh bounds in background thread. ParallelFor could make this faster (eg FDynamicMesh3::GetBounds) but it's going to take longer than the other steps below anyway
+	UE::Tasks::FTask BoundsTask = UE::Tasks::Launch(UE_SOURCE_LOCATION, [&DensePointSetOut, &SampleMesh]()
+	{
+		DensePointSetOut.Bounds = FAxisAlignedBox3d::Empty();
+		int MaxVertexID = SampleMesh.MaxVertexID();
+		for (int32 k = 0; k < MaxVertexID; ++k)
+		{
+			if (SampleMesh.IsVertex(k))
+			{
+				DensePointSetOut.Bounds.Contain(SampleMesh.GetVertex(k));
+			}
+		}
+	});
+	
+	double DiscArea = (FMathd::Pi * SampleRadius * SampleRadius);
+	double ApproxNumUniformSamples = TriInfoCache.TotalArea / DiscArea;		// uniform disc area
+	double EstNumDenseSamples = ApproxNumUniformSamples * SubSampleDensity * 2;		// 2 is fudge-factor
+	if (MaxNumDenseSamples != 0 && EstNumDenseSamples > MaxNumDenseSamples)
+	{
+		EstNumDenseSamples = MaxNumDenseSamples;
+	}
+	double DenseSampleArea = TriInfoCache.TotalArea / EstNumDenseSamples;
+
+
+	FPerTriangleDensePointSampling DensePerTriangleSampling;
+	DensePerTriangleSampling.InitializeForTriangleSet<MeshType>(SampleMesh, TriInfoCache, DenseSampleArea, RandomSeed, 
+		NonUniformConfig, DensePointSetOut);
+
+	//UE_LOG(LogTemp, Warning, TEXT("TotalArea %f  SampleArea %f EstSamples %d  DenseSamples %d   DenseArea %f  NumDensePoints %d"),
+	//	TriInfo.TotalArea, DiscArea, (int)ApproxNumUniformSamples, (int)EstNumDenseSamples, DenseSampleArea, NumDensePoints);
+
+	// make sure we have bounds initialized
+	BoundsTask.Wait();
+}
+
+
+
+template<typename MeshType>
+void UniformMeshPointSampling(
+	const MeshType& SampleMesh,
+	TFunctionRef<void(FVector3d, int32, double)> EmitSampleFunc,
+	double SampleRadius,
+	int32 MaxSamples,
+	double SubSampleDensity,
+	int RandomSeed,
+	int MaxNumDenseSamples,
+	FProgressCancel* Progress)
+{
+	TRACE_CPUPROFILER_EVENT_SCOPE(UniformMeshPointSampling);
+
+	MaxSamples = (MaxSamples == 0) ? TNumericLimits<int>::Max() : MaxSamples;
+	const bool bShuffle = (MaxSamples < TNumericLimits<int>::Max());
+
+	//
+	// Step 1: generate dense random point sampling of mesh surface. 
+	// 
+	FDenseSamplePointSet DensePointSet;
+	{
+		TRACE_CPUPROFILER_EVENT_SCOPE(ComputeDenseSamples);
+		
+		ConstructDenseUniformMeshPointSampling<MeshType>(SampleMesh, SampleRadius, SubSampleDensity, RandomSeed, 
+			FNonUniformSamplingConfig(), MaxNumDenseSamples, DensePointSet);
+	}
+	int32 NumDensePoints = DensePointSet.MaxVertexID();
+
+	if (Progress && Progress->Cancelled() ) return;
+
+	//
+	// Step 2: store dense point sampling in octree
+	// 
+	FSparseDynamicPointOctree3 Octree;
+	Octree.ConfigureFromPointCountEstimate(DensePointSet.Bounds.MaxDim(), NumDensePoints);
+	{
+		TRACE_CPUPROFILER_EVENT_SCOPE(BuildDenseOctree);
+		Octree.ParallelInsertDensePointSet(NumDensePoints, [&](int32 VertexID) { return DensePointSet.GetVertex(VertexID); } );
+	}
+
+	if (Progress && Progress->Cancelled() ) return;
+
+	//
+	// Step 3: dart throwing. Draw "random" dense point, add it to output set, and then
+	// remove all other dense points within radius of this point. 
+	// *NOTE* that in this implementation, we are not necessarily drawing random points.
+	// We are drawing from a random sampling on the triangles, but the per-triangle points
+	// are added to the PointSet in triangle index order. This introduces some bias
+	// but makes the algorithm quite a bit faster...
+
+	// currently only generating a semi-random point ordering (via shuffling) 
+	// if a subset of points is requested.  This likely does create some bias, 
+	TArray<int32> PointOrdering;
+	PointOrdering.Reserve(NumDensePoints);
+	if (bShuffle)
+	{
+		FModuloIteration Iter(NumDensePoints);
+		uint32 NextIndex;
+		while (Iter.GetNextIndex(NextIndex))
+		{
+			PointOrdering.Add(NextIndex);
+		}
+	}
+	else
+	{
+		for (int32 k = 0; k < NumDensePoints; ++k)
+		{
+			PointOrdering.Add(k);
+		}
+	}
+	int32 CurOrderingIndex = 0;
+
+	double QueryRadiusSqr = 4 * SampleRadius * SampleRadius;
+	TArray<bool> IsValidPoint;
+	IsValidPoint.Init(true, DensePointSet.MaxVertexID());
+	TArray<const FSparsePointOctreeCell*> QueryTempBuffer;
+	TArray<int> PointsInBall;
+
+	int NumEmittedSamples = 0;
+	{ TRACE_CPUPROFILER_EVENT_SCOPE(ExtractPoints);
+
+		while (CurOrderingIndex < NumDensePoints && NumEmittedSamples < MaxSamples)
+		{
+			if (NumEmittedSamples % 25 == 0 && Progress && Progress->Cancelled() ) return;
+
+			// pick a vertex in the dense point set, ie "throw a dart that is guaranteed to be valid"
+			int32 UseVertexID = IndexConstants::InvalidID;
+			{ //TRACE_CPUPROFILER_EVENT_SCOPE(FindInitialPoint);
+
+				while (UseVertexID == IndexConstants::InvalidID && CurOrderingIndex < NumDensePoints)
+				{
+					int VertexID = PointOrdering[CurOrderingIndex++];
+					if (IsValidPoint[VertexID])
+					{
+						UseVertexID = VertexID;
+						break;
+					}
+				}
+			}
+			if (UseVertexID == IndexConstants::InvalidID)
+			{
+				continue;
+			}
+
+			// found a valid point
+			FVector3d SamplePoint = DensePointSet.DensePoints[UseVertexID];
+			EmitSampleFunc(SamplePoint, DensePointSet.Triangles[UseVertexID], SampleRadius);
+			NumEmittedSamples++;
+			Octree.RemovePoint(UseVertexID);
+			IsValidPoint[UseVertexID] = false;
+
+			// remove dense points within sample radius from this point
+			PointsInBall.Reset();
+			FAxisAlignedBox3d QueryBox(SamplePoint, 2*SampleRadius);
+			Octree.RangeQuery(QueryBox,		// add SphereQuery to Octree? would save a chunk...
+				[&](int32 PointID) {
+					return IsValidPoint[PointID] && DistanceSquared(DensePointSet.GetVertex(PointID), SamplePoint) < QueryRadiusSqr;
+				},
+				PointsInBall, &QueryTempBuffer);
+
+			for (int32 QueryPointID : PointsInBall)
+			{
+				Octree.RemovePointUnsafe(QueryPointID);
+				IsValidPoint[QueryPointID] = false;
+			}
+
+		}
+	}
+}
+
+
+
+
+
+
+
+template<typename MeshType>
+void NonUniformMeshPointSampling(
+	const MeshType& SampleMesh,
+	TFunctionRef<void(FVector3d, int32, double)> EmitSampleFunc,
+	double MinSampleRadius,
+	double MaxSampleRadius,
+	int32 MaxSamples,
+	double SubSampleDensity,
+	int RandomSeed,
+	const FNonUniformSamplingConfig& NonUniformConfig,
+	int MaxNumDenseSamples,
+	FProgressCancel* Progress)
+{
+	TRACE_CPUPROFILER_EVENT_SCOPE(UniformMeshPointSampling);
+
+	MaxSamples = (MaxSamples == 0) ? TNumericLimits<int>::Max() : MaxSamples;
+
+	//
+	// Step 1: generate dense random point sampling of mesh surface. 
+	// 
+	FDenseSamplePointSet DensePointSet;
+	{
+		TRACE_CPUPROFILER_EVENT_SCOPE(ComputeDenseSamples);
+		ConstructDenseUniformMeshPointSampling<MeshType>(SampleMesh, MinSampleRadius, SubSampleDensity, RandomSeed, 
+			NonUniformConfig, MaxNumDenseSamples, DensePointSet);
+	}
+	int32 NumDensePoints = DensePointSet.MaxVertexID();
+	bool bHaveWeights = DensePointSet.Weights.Num() > 0;
+
+	if (Progress && Progress->Cancelled() ) return;
+
+	//
+	// Step 2: store dense point sampling in octree
+	// 
+	FSparseDynamicPointOctree3 Octree;
+	Octree.ConfigureFromPointCountEstimate(DensePointSet.Bounds.MaxDim(), NumDensePoints);
+	{
+		TRACE_CPUPROFILER_EVENT_SCOPE(BuildDenseOctree);
+		Octree.ParallelInsertDensePointSet(NumDensePoints, [&](int32 VertexID) { return DensePointSet.GetVertex(VertexID); } );
+	}
+
+	if (Progress && Progress->Cancelled() ) return;
+
+	//
+	// Step 3: dart throwing. Draw "random" dense point, add it to output set, and then
+	// remove all other dense points within radius of this point. 
+	// *NOTE* that in this implementation, we are not necessarily drawing random points.
+	// We are drawing from a random sampling on the triangles, but the per-triangle points
+	// are added to the PointSet in triangle index order. This introduces some bias
+	// but makes the algorithm quite a bit faster...
+
+	// currently only generating a semi-random point ordering (via shuffling) 
+	// if a subset of points is requested.  This likely does create some bias, 
+	TArray<int32> PointOrdering;
+	PointOrdering.Reserve(NumDensePoints);
+	FModuloIteration Iter(NumDensePoints);
+	{
+		uint32 NextIndex;
+		while (Iter.GetNextIndex(NextIndex))
+		{
+			PointOrdering.Add(NextIndex);
+		}
+	}
+	int32 CurOrderingIndex = 0;
+
+	TArray<bool> IsValidPoint;
+	IsValidPoint.Init(true, DensePointSet.MaxVertexID());
+	TArray<const FSparsePointOctreeCell*> QueryTempBuffer;
+	TArray<int> PointsInBall;
+
+	FRandomStream RadiusStream(RandomSeed);
+
+	// likely could benefit from hash grid here? 
+	TArray<FVector3d> EmittedSamples;
+	TArray<double> EmittedRadius;
+	auto FindOverlappingSample = [&EmittedSamples, &EmittedRadius, &PointsInBall, &QueryTempBuffer, MaxSampleRadius](FVector3d Position, double Radius) -> double
+	{
+		for (int32 k = 0; k < EmittedSamples.Num(); ++k)
+		{
+			double NeighbourDist = Distance(EmittedSamples[k], Position);
+			//if ( NeighbourDist < (Radius + EmittedRadius[k]) )
+			if ( Radius > (NeighbourDist - EmittedRadius[k]) )
+			{
+				return (NeighbourDist - EmittedRadius[k]);
+			}
+		}
+		return TNumericLimits<double>::Max();
+	};
+
+	// In weighted sampling, we cannot guarantee that a dense sample point w/ a radius > MinSampleRadius will
+	// actually fit without collision. The "correct" way to handle this, by randomly choosing new points until
+	// a valid one is found, can take a very long time. So instead we can "decay" the radius down to 
+	// MinSampleRadius in multiple steps, trying to find a radius that fits. We are guaranteed that any active
+	// dense point will fit with MinSampleRadius, so this significantly accellerates the sampling, at the cost
+	// of introducing some bias.
+	TArray<double, TInlineAllocator<10>> DecaySteps;
+	bool bIsFixedRadiusMethod = false;
+	if (NonUniformConfig.InterpretWeightMode == FMeshSurfacePointSampling::EInterpretWeightMode::RadiusInterp)
+	{
+		DecaySteps = TArray<double, TInlineAllocator<10>>({1.0});
+		bIsFixedRadiusMethod = true;
+	}
+	else
+	{
+		DecaySteps = TArray<double, TInlineAllocator<10>>({1.0, 0.8, 0.6, 0.4, 0.2, 0.0});
+	}
+
+	TArray<double> CurDistances;
+	CurDistances.Init(TNumericLimits<double>::Max(), NumDensePoints);
+
+	int NumEmittedSamples = 0;
+	{ TRACE_CPUPROFILER_EVENT_SCOPE(ExtractPoints);
+
+		int NumFailures = 0;
+		while (PointOrdering.Num() > 0 && NumEmittedSamples < MaxSamples && NumFailures < 1000)
+		{
+			if (NumEmittedSamples % 25 == 0 && Progress && Progress->Cancelled() ) return;
+
+			int32 UseVertexID = IndexConstants::InvalidID;
+			int32 PointOrderingIndex = -1;
+			double SampleRadius = MinSampleRadius;
+
+			// try to find a valid (point, radius) pair. This may fail if we cannot find a 
+			// valid radius...
+			int32 NumRemaining = PointOrdering.Num();
+			for ( int32 k = 0; k < NumRemaining; ++k)
+			{
+				int VertexID = PointOrdering[k];
+
+				if (IsValidPoint[VertexID] == false)		// if point is expired, discard it
+				{
+					PointOrdering.RemoveAtSwap(k, 1, false);
+					NumRemaining--;
+					k--;		// reconsider point we just swapped to index k
+					continue;
+				}
+
+				FVector3d Position = DensePointSet.GetVertex(VertexID);
+
+				// based on the weight/random strategy, generate a parameter in range [0,1] that will
+				// be used to interpolate the Min/Max Radius below
+				double InterpRadiusT = 0;
+				if (bHaveWeights)
+				{
+					if (NonUniformConfig.InterpretWeightMode == FMeshSurfacePointSampling::EInterpretWeightMode::WeightedRandom)
+					{
+						double Weight = FMathd::Clamp(DensePointSet.Weights[VertexID], 0, 1);
+						double Random = RadiusStream.GetFraction();
+						InterpRadiusT = (Weight + Random) / 2.0;		// can parameterize this as ( (N-1)*Weight + Random ) / N
+					}
+					else  // RadiusInterp
+					{
+						InterpRadiusT = DensePointSet.Weights[VertexID];
+					}
+				}
+				else 
+				{
+					InterpRadiusT = RadiusStream.GetFraction();
+				}
+				if (NonUniformConfig.SizeDistribution == FMeshSurfacePointSampling::ESizeDistribution::Smaller)
+				{
+					InterpRadiusT = FMathd::Pow(InterpRadiusT, NonUniformConfig.SizeDistributionPower);
+				}
+				else if (NonUniformConfig.SizeDistribution == FMeshSurfacePointSampling::ESizeDistribution::Larger)
+				{
+					InterpRadiusT = FMathd::Pow(InterpRadiusT, 1.0 / NonUniformConfig.SizeDistributionPower);
+				}
+
+				// try to fit a sample at the selected point, possibly incrementally shrinking the 
+				// sample radius down to MinRadius to guarantee a fit
+				bool bFound = false;
+				double MinNeighbourGap = TNumericLimits<double>::Max();
+				for (double DecayStep : DecaySteps)
+				{
+					double UseRadius = FMathd::Lerp(MinSampleRadius, MaxSampleRadius, InterpRadiusT * DecayStep);
+					if (UseRadius > CurDistances[VertexID])
+					{
+						continue;
+					}
+
+					double NeighbourGap = FindOverlappingSample(Position, UseRadius);
+					//if ( NeighbourGap == TNumericLimts<double>::Max() )
+					if ( UseRadius < NeighbourGap )
+					{
+						UseVertexID = VertexID;
+						SampleRadius = UseRadius;
+						PointOrderingIndex = k;
+						bFound = true;
+						break;
+					}
+					else
+					{
+						MinNeighbourGap = FMathd::Min(MinNeighbourGap, NeighbourGap);
+					}
+				}
+				if (bFound)
+				{
+					break;
+				}
+				CurDistances[VertexID] = FMathd::Min(CurDistances[VertexID], MinNeighbourGap);
+
+				// if this is a method w/ no random variation or decay, this (point,radius) pair will never fit and can be removed
+				if (bIsFixedRadiusMethod)
+				{
+					PointOrdering.RemoveAtSwap(k, 1, false);
+					NumRemaining--;
+					k--;		// reconsider point we just swapped to index k
+				}
+			}
+			if (UseVertexID == IndexConstants::InvalidID)
+			{
+				NumFailures++;
+				continue;
+			}
+
+			// remove selected point from ordering
+			PointOrdering.RemoveAtSwap(PointOrderingIndex, 1, false);
+
+			// emit our valid (point, triangle, radius) sample
+			FVector3d SamplePoint = DensePointSet.DensePoints[UseVertexID];
+			EmitSampleFunc(SamplePoint, DensePointSet.Triangles[UseVertexID], SampleRadius);
+			NumEmittedSamples++;
+			Octree.RemovePoint(UseVertexID);
+			IsValidPoint[UseVertexID] = false;
+
+			// add point to known samples list
+			int32 SampleID = EmittedSamples.Num();
+			EmittedSamples.Add(SamplePoint);
+			EmittedRadius.Add(SampleRadius);
+
+			// once we add this point, no point can be within it's radius, and any other point closer than
+			// MinSampleRadius would collide, so we can decimate all points within the radius sum
+			double CombinedRadiusSqr = ((SampleRadius + MinSampleRadius) * (SampleRadius + MinSampleRadius));
+
+			// find all dense points within our query radius
+			PointsInBall.Reset();
+			FAxisAlignedBox3d QueryBox(SamplePoint, 2*SampleRadius);
+			Octree.RangeQuery(QueryBox,		// add SphereQuery to Octree? would save a chunk...
+				[&](int32 PointID) {
+					return IsValidPoint[PointID] && DistanceSquared(DensePointSet.GetVertex(PointID), SamplePoint) < CombinedRadiusSqr;
+				},
+				PointsInBall, &QueryTempBuffer);
+
+			// remove all those dense points from the octree and mark them invalid
+			for (int32 QueryPointID : PointsInBall)
+			{
+				Octree.RemovePointUnsafe(QueryPointID);
+				IsValidPoint[QueryPointID] = false;
+			}
+		}
+	}
+}
+
+
+
+
+
+
+
+
+} // end namespace UELocal
+
+
+void FMeshSurfacePointSampling::ComputePoissonSampling(const FDynamicMesh3& Mesh, FProgressCancel* Progress)
+{
+	Result = FGeometryResult(EGeometryResultType::InProgress);
+
+	Samples.Reset();
+	Radii.Reset();
+	auto AddSampleFunc = [&](FVector3d Position, int TriangleID, double Radius)
+		{
+			Samples.Add(FFrame3d(Position, Mesh.GetTriNormal(TriangleID)));
+			Radii.Add(Radius);
+		};
+
+	if (MaxSampleRadius > SampleRadius)
+	{
+		UELocal::FNonUniformSamplingConfig NonUniformConfig;
+		NonUniformConfig.InterpretWeightMode = this->InterpretWeightMode;
+		NonUniformConfig.SizeDistribution = this->SizeDistribution;
+		NonUniformConfig.SizeDistributionPower = FMath::Clamp(this->SizeDistributionPower, 1.0, 10.0);
+
+		if (bUseVertexWeights && VertexWeights.Num() == Mesh.MaxVertexID())
+		{
+			auto ComputeVertexWeightFunc = [&](int TriangleID, FVector3d Position, FVector3d BaryCoords)
+				{
+					FIndex3i Tri = Mesh.GetTriangle(TriangleID);
+					double Weight = BaryCoords.X*VertexWeights[Tri.A] + BaryCoords.Y*VertexWeights[Tri.B] + BaryCoords.Z*VertexWeights[Tri.C];
+					if (bInvertWeights)
+					{
+						Weight = 1.0 - FMathd::Clamp(Weight, 0.0, 1.0);
+					}
+					return Weight;
+				};
+			NonUniformConfig.WeightFunction = ComputeVertexWeightFunc;
+
+			UELocal::NonUniformMeshPointSampling<FDynamicMesh3>(Mesh, AddSampleFunc,
+				SampleRadius, MaxSampleRadius, MaxSamples, SubSampleDensity, RandomSeed, NonUniformConfig, MaxSubSamplePoints, Progress);
+		}
+		else
+		{
+			UELocal::NonUniformMeshPointSampling<FDynamicMesh3>(Mesh, AddSampleFunc,
+				SampleRadius, MaxSampleRadius, MaxSamples, SubSampleDensity, RandomSeed, NonUniformConfig, MaxSubSamplePoints, Progress );
+		}
+	}
+	else
+	{
+		UELocal::UniformMeshPointSampling<FDynamicMesh3>(Mesh, AddSampleFunc,
+			SampleRadius, MaxSamples, SubSampleDensity, RandomSeed, MaxSubSamplePoints, Progress );
+	}
+
+	Result.SetSuccess(true, Progress);
+}
\ No newline at end of file
diff --git a/Engine/Source/Runtime/GeometryCore/Public/Sampling/MeshSurfacePointSampling.h b/Engine/Source/Runtime/GeometryCore/Public/Sampling/MeshSurfacePointSampling.h
new file mode 100644
--- /dev/null
+++ b/Engine/Source/Runtime/GeometryCore/Public/Sampling/MeshSurfacePointSampling.h
@@ -0,0 +1,152 @@
+// Copyright Epic Games, Inc. All Rights Reserved.
+
+#pragma once
+
+#include "MathUtil.h"
+#include "VectorTypes.h"
+#include "FrameTypes.h"
+#include "Util/ProgressCancel.h"
+
+class FProgressCancel;
+
+namespace UE
+{
+namespace Geometry
+{
+
+class FDynamicMesh3;
+
+
+/**
+ * FMeshSurfacePointSampling computes oriented point samples on the surface of
+ * a Mesh using various sampling strategies.
+ */
+class GEOMETRYCORE_API FMeshSurfacePointSampling
+{
+public:
+	// 
+	// Basic sampling parameters
+	//
+
+	/** Desired "radius" of sample points. Spacing between samples is at least 2x this value. */
+	double SampleRadius = 10.0;
+
+	/** Maximum number of samples requested. If 0 or default value, mesh will be maximally sampled */
+	uint32 MaxSamples = TNumericLimits<uint32>::Max();
+	
+	/** Density of subsampling used in Poisson strategy. Larger numbers mean "more accurate" (but slower) results. */
+	double SubSampleDensity = 10.0;
+
+	/** Random Seed used to initialize sampling strategies */
+	int32 RandomSeed = 0;
+
+	/** Safety threshold for maximum number of points in subsampling. */
+	int32 MaxSubSamplePoints = 50000000;
+
+	//
+	// Parameters for non-uniform / variable-radius sampling. 
+	// Disabled if MaxSampleRadius <= SampleRadius
+	//
+
+	/** If MaxSampleRadius > SampleRadius, then output point radii will be in range [SampleRadius, MaxSampleRadius] */
+	double MaxSampleRadius = 0.0;
+
+	/** Controls the distribution of sample radii */
+	enum class ESizeDistribution
+	{
+		/** Uniform distribution of sizes, ie all equally likely */
+		Uniform = 0,
+		/** Distribution is weighted towards smaller points (T^Power) */
+		Smaller = 1,
+		/** Distribution is weighted towards larger points (T^1/Power) */
+		Larger = 2
+	};
+	/** Active Size Distribution mode */
+	ESizeDistribution SizeDistribution = ESizeDistribution::Uniform;
+
+	/** Used to define how extreme the Size Distribution shift is. Valid range is [1,10] */
+	double SizeDistributionPower = 2.0;
+
+
+	/** Control whether VertexWeights (if valid) will be interpolated to modulate sampling */
+	bool bUseVertexWeights = false;
+
+	/** Per-vertex weights, size must be == Max Vertex Index of input Mesh */
+	TArray<double> VertexWeights;
+
+	/** Control how active Weights are used to affect point radius */
+	enum class EInterpretWeightMode
+	{
+		/** Weights are clamped to [0,1] and used to interpolate Min/Max Radius */
+		RadiusInterp = 0,
+		/** Weights are clamped to [0,1] and used to interpolate Min/Max Radius, with decay, so that smaller-radius samples will infill between large ones */
+		RadiusInterpWithFill = 1,
+		/** Weight is used to create nonuniform random sampling, ie it nudges the random point radius distribution but does not directly control it */
+		WeightedRandom = 2
+	};
+	/** Active weight interpretation mode */
+	EInterpretWeightMode InterpretWeightMode = EInterpretWeightMode::RadiusInterpWithFill;
+
+	/** If true, weights are inverted */
+	bool bInvertWeights = false;
+
+
+	//
+	// TODO: when MaxSamples is set, it would be useful to be able to use Weight to modulate
+	// positional distribution instead of radius, (or both!)
+	//
+
+public:
+	//
+	// Outputs
+	//
+
+	/** Result of last computation */
+	FGeometryResult Result;
+
+	/** Oriented Sample Points on the mesh surface. Z axis of frame points along mesh normal, X and Y are arbitrary */
+	TArray<FFrame3d> Samples;
+
+	/** Radius of each Sample Point, length is the same as Samples array */
+	TArray<double> Radii;
+
+
+public:
+
+	/**
+	 * Compute an approximate Poisson sampling of the mesh, either uniform or non-uniform depending on the settings above.
+	 * The sampling will attempt to fully cover the mesh unless .MaxSamples is provided, in which case exactly
+	 * that many samples will be generated and they will be randomly distributed (so not at all Poisson...!)
+	 * 
+	 * By default the sampling will be uniform - all points will be spaced at least .SampleRadius*2, ie no "collisions"
+	 * between their bounding spheres.
+	 * 
+	 * If .MaxSampleRadius is larger than .SampleRadius, the sampling will be non-uniform, ie samples will be emitted
+	 * with radii within this range. By default a random uniform distribution of radii will be attempted, .SizeDistribution 
+	 * and related parameters can be used to make this distribution non-uniform. 
+	 * The spacing between points will always be greater than the sum of their two sample radii, so again no collisions. 
+	 * However the "density" of the sampling will vary depending on how well the algorithm can find gaps to fill.
+	 * 
+	 * If .bUseVertexWeights is defined and valid VertexWeights are provided, they will be used to modulate the sampling radii.
+	 * The .InterpretWeightMode setting controls how the weights are used to influence the sample radii.
+	 * 
+	 * The strategy used is to compute a much higher density sampling than needed (based on .SubSampleDensity),
+	 * then iteratively select from that point set and decimate it within the radius of selected samples. The implementation
+	 * also introduces various biases to increase performance. Generally increasing SubSampleDensity will result in 
+	 * more tightly-packed results, but at increasingly expensive computation time.
+	 */
+	void ComputePoissonSampling(const FDynamicMesh3& Mesh, FProgressCancel* Progress = nullptr);
+
+
+	// ability to incrementally increase existing sample sizes to more tightly pack them?
+	// ability to iteratively redistribute samples to improve uniformity (ie potential field / mass-spring type method, point smoothing, ?)
+};
+
+
+
+
+
+
+} // end namespace UE::Geometry
+} // end namespace UE
+
