Fix an unused TDenseGrid2 constructor

diff --git a/Engine/Source/Runtime/GeometryCore/Public/Spatial/DenseGrid2.h b/Engine/Source/Runtime/GeometryCore/Public/Spatial/DenseGrid2.h
--- a/Engine/Source/Runtime/GeometryCore/Public/Spatial/DenseGrid2.h
+++ b/Engine/Source/Runtime/GeometryCore/Public/Spatial/DenseGrid2.h
@@ -38,8 +38,9 @@ public:
 
 	TDenseGrid2(int32 DimX, int32 DimY, ElemType InitialValue)
 	{
-		Resize(DimX, DimY);
-		Assign(InitialValue);
+		DimensionX = (int64)DimX;
+		DimensionY = (int64)DimY;
+		Buffer.Init(InitialValue, Size());
 	}
 
 	int64 Size() const
