Fixed non-unity (with orphaned headers compiled independently) compile errors

diff --git a/Engine/Plugins/FX/Niagara/Source/Niagara/Public/INiagaraEditorOnlyDataUtlities.h b/Engine/Plugins/FX/Niagara/Source/Niagara/Public/INiagaraEditorOnlyDataUtlities.h
--- a/Engine/Plugins/FX/Niagara/Source/Niagara/Public/INiagaraEditorOnlyDataUtlities.h
+++ b/Engine/Plugins/FX/Niagara/Source/Niagara/Public/INiagaraEditorOnlyDataUtlities.h
@@ -2,11 +2,11 @@
 
 #pragma once
 
+#include "UObject/Object.h"
+
 class UNiagaraScriptSourceBase;
 class UNiagaraEditorParametersAdapterBase;
 class UNiagaraEditorDataBase;
-class UObject;
-struct FAssetRegistryTag;
 
 /** Defines utility methods for creating editor only data which is stored on runtime objects. */
 class INiagaraEditorOnlyDataUtilities
diff --git a/Engine/Source/Runtime/Core/Internal/IO/IoPriorityQueue.h b/Engine/Source/Runtime/Core/Internal/IO/IoPriorityQueue.h
--- a/Engine/Source/Runtime/Core/Internal/IO/IoPriorityQueue.h
+++ b/Engine/Source/Runtime/Core/Internal/IO/IoPriorityQueue.h
@@ -2,6 +2,7 @@
 
 #pragma once
 
+#include "Algo/BinarySearch.h"
 #include "Algo/IsSorted.h"
 
 template<typename T>
diff --git a/Engine/Source/Runtime/Engine/Private/Materials/MaterialShaderMapLayout.h b/Engine/Source/Runtime/Engine/Private/Materials/MaterialShaderMapLayout.h
--- a/Engine/Source/Runtime/Engine/Private/Materials/MaterialShaderMapLayout.h
+++ b/Engine/Source/Runtime/Engine/Private/Materials/MaterialShaderMapLayout.h
@@ -11,6 +11,7 @@ class FShaderPipelineType;
 class FVertexFactoryType;
 struct FMaterialShaderParameters;
 enum class EShaderPermutationFlags : uint32;
+enum EShaderPlatform : uint16;
 
 class FShaderLayoutEntry
 {
diff --git a/Engine/Source/Runtime/GeometryCore/Public/Spatial/OccupancyGrid3.h b/Engine/Source/Runtime/GeometryCore/Public/Spatial/OccupancyGrid3.h
--- a/Engine/Source/Runtime/GeometryCore/Public/Spatial/OccupancyGrid3.h
+++ b/Engine/Source/Runtime/GeometryCore/Public/Spatial/OccupancyGrid3.h
@@ -36,7 +36,7 @@ struct FOccupancyGrid3
 		TMeshAABBTree3<MeshType> Spatial(&InMesh);
 		TFastWindingTree FastWinding(&Spatial);
 		const FAxisAlignedBox3d Bounds = Spatial.GetBoundingBox();
-		CellSize = Bounds.MaxDim() / InVoxelResolution;
+		CellSize = float(Bounds.MaxDim() / InVoxelResolution);
 		CellMidPoint = {CellSize / 2.0f, CellSize / 2.0f, CellSize / 2.0f};
 
 		TMeshWindingNumberGrid WindingGrid(&InMesh, &FastWinding, CellSize);
@@ -110,7 +110,7 @@ struct FOccupancyGrid3
 	/// in the winding number grid, from which the occupancy got computed.
 	FVector3f GetCellCenterFromIndex(const FVector3i &Index) const
 	{
-		return {Index.X * CellSize + GridOrigin.X, Index.Y * CellSize + GridOrigin.Y, Index.Z * CellSize + GridOrigin.Z};
+		return {float(Index.X) * CellSize + GridOrigin.X, float(Index.Y) * CellSize + GridOrigin.Y, float(Index.Z) * CellSize + GridOrigin.Z};
 	}
 
 	/// Given an index into the occupancy grid, returns the bbox that represents the cell in the winding number
diff --git a/Engine/Source/Runtime/GeometryCore/Public/Spatial/SparseDynamicPointOctree3.h b/Engine/Source/Runtime/GeometryCore/Public/Spatial/SparseDynamicPointOctree3.h
--- a/Engine/Source/Runtime/GeometryCore/Public/Spatial/SparseDynamicPointOctree3.h
+++ b/Engine/Source/Runtime/GeometryCore/Public/Spatial/SparseDynamicPointOctree3.h
@@ -8,6 +8,7 @@
 #include "GeometryTypes.h"
 #include "BoxTypes.h"
 #include "Util/DynamicVector.h"
+#include "Util/RefCountVector.h"
 #include "Util/SparseListSet.h"
 #include "Spatial/SparseGrid3.h"
 #include "Intersection/IntrRay3AxisAlignedBox3.h"
diff --git a/Engine/Source/Runtime/RHI/Public/RHIAccess.h b/Engine/Source/Runtime/RHI/Public/RHIAccess.h
--- a/Engine/Source/Runtime/RHI/Public/RHIAccess.h
+++ b/Engine/Source/Runtime/RHI/Public/RHIAccess.h
@@ -2,6 +2,11 @@
 
 #pragma once
 
+#include "CoreTypes.h"
+#include "Misc/EnumClassFlags.h"
+
+class FString;
+
 enum class ERHIAccess : uint32
 {
 	// Used when the previous state of a resource is not known,
