GeometryCore: improve FDynamicMeshAttributeSet::EnableMatchingAttributes() function. SkinWeights and Generic attributes now are only reset/discarded if bClearExisting=true. Added bDiscardExtraAttributes flag which results in removal of any attributes not in target attribute set, but does not clear existing attributes. Also added bDiscardExtraAttributes flag to FDynamicMesh3::EnableMatchingAttributes(). GeoScript Mesh Append functions use this flag now when matching "source" attributes, and behave properly in default (combine) mode for Skin Weights.

diff --git a/Engine/Plugins/Experimental/GeometryScripting/Source/GeometryScriptingCore/Private/MeshBasicEditFunctions.cpp b/Engine/Plugins/Experimental/GeometryScripting/Source/GeometryScriptingCore/Private/MeshBasicEditFunctions.cpp
--- a/Engine/Plugins/Experimental/GeometryScripting/Source/GeometryScriptingCore/Private/MeshBasicEditFunctions.cpp
+++ b/Engine/Plugins/Experimental/GeometryScripting/Source/GeometryScriptingCore/Private/MeshBasicEditFunctions.cpp
@@ -370,16 +370,15 @@ void FGeometryScriptAppendMeshOptions::UpdateAttributesForCombineMode(FDynamicMe
 {
 	if (CombineMode == EGeometryScriptCombineAttributesMode::EnableAllMatching)
 	{
-		Target.EnableMatchingAttributes(Source, false);
+		Target.EnableMatchingAttributes(Source, false, false);
 	}
 	else if (CombineMode == EGeometryScriptCombineAttributesMode::UseSource)
 	{
-		Target.EnableMatchingAttributes(Source, true);
+		Target.EnableMatchingAttributes(Source, false, true);
 	}
 	// else the mode is UseTarget, which already corresponds to the default behavior for AppendMesh
 }
 
-
 UDynamicMesh* UGeometryScriptLibrary_MeshBasicEditFunctions::AppendMesh(
 	UDynamicMesh* TargetMesh,
 	UDynamicMesh* AppendMesh,
@@ -424,7 +423,6 @@ UDynamicMesh* UGeometryScriptLibrary_MeshBasicEditFunctions::AppendMesh(
 }
 
 
-
 UDynamicMesh* UGeometryScriptLibrary_MeshBasicEditFunctions::AppendMeshTransformed(
 	UDynamicMesh* TargetMesh,
 	UDynamicMesh* AppendMesh,
diff --git a/Engine/Source/Runtime/GeometryCore/Private/DynamicMesh/DynamicMesh3.cpp b/Engine/Source/Runtime/GeometryCore/Private/DynamicMesh/DynamicMesh3.cpp
--- a/Engine/Source/Runtime/GeometryCore/Private/DynamicMesh/DynamicMesh3.cpp
+++ b/Engine/Source/Runtime/GeometryCore/Private/DynamicMesh/DynamicMesh3.cpp
@@ -335,10 +335,10 @@ void FDynamicMesh3::Clear()
 }
 
 
-void FDynamicMesh3::EnableMatchingAttributes(const FDynamicMesh3& ToMatch, bool bClearExisting)
+void FDynamicMesh3::EnableMatchingAttributes(const FDynamicMesh3& ToMatch, bool bClearExisting, bool bDiscardExtraAttributes)
 {
-	bool bWantVertexNormals = ToMatch.HasVertexNormals() || (this->HasVertexNormals() && bClearExisting == false);
-	if (bClearExisting || bWantVertexNormals == false)
+	bool bWantVertexNormals = (bClearExisting || bDiscardExtraAttributes) ? ToMatch.HasVertexNormals() : ( ToMatch.HasVertexNormals() || this->HasVertexNormals() );
+	if (bClearExisting || bWantVertexNormals == false )
 	{
 		DiscardVertexNormals();
 	}
@@ -347,8 +347,8 @@ void FDynamicMesh3::EnableMatchingAttributes(const FDynamicMesh3& ToMatch, bool
 		EnableVertexNormals(FVector3f::UnitZ());
 	}
 
-	bool bWantVertexColors = ToMatch.HasVertexColors() || (this->HasVertexColors() && bClearExisting == false);
-	if (bClearExisting || bWantVertexColors == false)
+	bool bWantVertexColors = (bClearExisting || bDiscardExtraAttributes) ? ToMatch.HasVertexColors() : ( ToMatch.HasVertexColors() || this->HasVertexColors() );
+	if (bClearExisting || bWantVertexColors == false )
 	{
 		DiscardVertexColors();
 	}
@@ -357,9 +357,8 @@ void FDynamicMesh3::EnableMatchingAttributes(const FDynamicMesh3& ToMatch, bool
 		EnableVertexColors(FVector3f::Zero());
 	}
 
-
-	bool bWantVertexUVs = ToMatch.HasVertexUVs() || (this->HasVertexUVs() && bClearExisting == false);
-	if (bClearExisting || bWantVertexUVs == false)
+	bool bWantVertexUVs = (bClearExisting || bDiscardExtraAttributes) ? ToMatch.HasVertexUVs() : ( ToMatch.HasVertexUVs() || this->HasVertexUVs() );
+	if (bClearExisting || bWantVertexUVs == false )
 	{
 		DiscardVertexUVs();
 	}
@@ -368,9 +367,8 @@ void FDynamicMesh3::EnableMatchingAttributes(const FDynamicMesh3& ToMatch, bool
 		EnableVertexUVs(FVector2f::Zero());
 	}
 
-
-	bool bWantTriangleGroups = ToMatch.HasTriangleGroups() || (this->HasTriangleGroups() && bClearExisting == false);
-	if (bClearExisting || bWantTriangleGroups == false)
+	bool bWantTriangleGroups = (bClearExisting || bDiscardExtraAttributes) ? ToMatch.HasTriangleGroups() : ( ToMatch.HasTriangleGroups() || this->HasTriangleGroups() );
+	if (bClearExisting || bWantTriangleGroups == false )
 	{
 		DiscardTriangleGroups();
 	}
@@ -379,9 +377,8 @@ void FDynamicMesh3::EnableMatchingAttributes(const FDynamicMesh3& ToMatch, bool
 		EnableTriangleGroups();
 	}
 
-
-	bool bWantAttributes = ToMatch.HasAttributes() || (this->HasAttributes() && bClearExisting == false);
-	if (bClearExisting || bWantAttributes == false)
+	bool bWantAttributes = (bClearExisting || bDiscardExtraAttributes) ? ToMatch.HasAttributes() : ( ToMatch.HasAttributes() || this->HasAttributes() );
+	if (bClearExisting || bWantAttributes == false )
 	{
 		DiscardAttributes();
 	}
@@ -391,9 +388,8 @@ void FDynamicMesh3::EnableMatchingAttributes(const FDynamicMesh3& ToMatch, bool
 	}
 	if (HasAttributes() && ToMatch.HasAttributes())
 	{
-		Attributes()->EnableMatchingAttributes(*ToMatch.Attributes(), bClearExisting);
+		Attributes()->EnableMatchingAttributes(*ToMatch.Attributes(), bClearExisting, bDiscardExtraAttributes);
 	}
-
 }
 
 
diff --git a/Engine/Source/Runtime/GeometryCore/Private/DynamicMesh/DynamicMeshAttributeSet.cpp b/Engine/Source/Runtime/GeometryCore/Private/DynamicMesh/DynamicMeshAttributeSet.cpp
--- a/Engine/Source/Runtime/GeometryCore/Private/DynamicMesh/DynamicMeshAttributeSet.cpp
+++ b/Engine/Source/Runtime/GeometryCore/Private/DynamicMesh/DynamicMeshAttributeSet.cpp
@@ -275,10 +275,10 @@ void FDynamicMeshAttributeSet::SplitAllBowties(bool bParallel)
 
 
 
-void FDynamicMeshAttributeSet::EnableMatchingAttributes(const FDynamicMeshAttributeSet& ToMatch, bool bClearExisting)
+void FDynamicMeshAttributeSet::EnableMatchingAttributes(const FDynamicMeshAttributeSet& ToMatch, bool bClearExisting, bool bDiscardExtraAttributes)
 {
 	int32 ExistingUVLayers = NumUVLayers();
-	int32 RequiredUVLayers = (bClearExisting) ? ToMatch.NumUVLayers() : FMath::Max(ExistingUVLayers, ToMatch.NumUVLayers());
+	int32 RequiredUVLayers = (bClearExisting || bDiscardExtraAttributes) ? ToMatch.NumUVLayers() : FMath::Max(ExistingUVLayers, ToMatch.NumUVLayers());
 	SetNumUVLayers(RequiredUVLayers);
 	for (int32 k = bClearExisting ? 0 : ExistingUVLayers; k < NumUVLayers(); k++)
 	{
@@ -286,14 +286,14 @@ void FDynamicMeshAttributeSet::EnableMatchingAttributes(const FDynamicMeshAttrib
 	}
 
 	int32 ExistingNormalLayers = NumNormalLayers();
-	int32 RequiredNormalLayers = (bClearExisting) ? ToMatch.NumUVLayers() : FMath::Max(ExistingNormalLayers, ToMatch.NumNormalLayers());
+	int32 RequiredNormalLayers = (bClearExisting || bDiscardExtraAttributes) ? ToMatch.NumUVLayers() : FMath::Max(ExistingNormalLayers, ToMatch.NumNormalLayers());
 	SetNumNormalLayers(RequiredNormalLayers);
 	for (int32 k = bClearExisting ? 0 : ExistingNormalLayers; k < NumNormalLayers(); k++)
 	{
 		NormalLayers[k].ClearElements();
 	}
 
-	bool bWantColorLayer = ToMatch.HasPrimaryColors() || (this->HasPrimaryColors() && bClearExisting == false);
+	bool bWantColorLayer = (bClearExisting || bDiscardExtraAttributes) ? ToMatch.HasPrimaryColors() : ( ToMatch.HasPrimaryColors() || this->HasPrimaryColors() );
 	if (bClearExisting || bWantColorLayer == false)
 	{
 		DisablePrimaryColors();
@@ -303,7 +303,7 @@ void FDynamicMeshAttributeSet::EnableMatchingAttributes(const FDynamicMeshAttrib
 		EnablePrimaryColors();
 	}
 
-	bool bWantMaterialID = ToMatch.HasMaterialID() || (this->HasMaterialID() && bClearExisting == false);
+	bool bWantMaterialID = (bClearExisting || bDiscardExtraAttributes) ? ToMatch.HasMaterialID() : ( ToMatch.HasMaterialID() || this->HasMaterialID() );
 	if (bClearExisting || bWantMaterialID == false)
 	{
 		DisableMaterialID();
@@ -313,9 +313,9 @@ void FDynamicMeshAttributeSet::EnableMatchingAttributes(const FDynamicMeshAttrib
 		EnableMaterialID();
 	}
 
-
+	// polygroup layers are handled by count, not by name...maybe wrong
 	int32 ExistingPolygroupLayers = NumPolygroupLayers();
-	int32 RequiredPolygroupLayers = (bClearExisting) ? ToMatch.NumPolygroupLayers() : FMath::Max(ExistingPolygroupLayers, ToMatch.NumPolygroupLayers());
+	int32 RequiredPolygroupLayers = (bClearExisting || bDiscardExtraAttributes) ? ToMatch.NumPolygroupLayers() : FMath::Max(ExistingPolygroupLayers, ToMatch.NumPolygroupLayers());
 	SetNumPolygroupLayers(RequiredPolygroupLayers);
 	for (int32 k = bClearExisting ? 0 : ExistingPolygroupLayers; k < NumPolygroupLayers(); k++)
 	{
@@ -326,8 +326,9 @@ void FDynamicMeshAttributeSet::EnableMatchingAttributes(const FDynamicMeshAttrib
 		}
 	}
 
+	// weightmap layers are handled by count, not by name...maybe wrong
 	int32 ExistingWeightLayers = NumWeightLayers();
-	int32 RequiredWeightLayers = (bClearExisting) ? ToMatch.NumWeightLayers() : FMath::Max(ExistingWeightLayers, ToMatch.NumWeightLayers());
+	int32 RequiredWeightLayers = (bClearExisting || bDiscardExtraAttributes) ? ToMatch.NumWeightLayers() : FMath::Max(ExistingWeightLayers, ToMatch.NumWeightLayers());
 	SetNumWeightLayers(RequiredWeightLayers);
 	for (int32 k = bClearExisting ? 0 : ExistingWeightLayers; k < NumWeightLayers(); k++)
 	{
@@ -338,22 +339,72 @@ void FDynamicMeshAttributeSet::EnableMatchingAttributes(const FDynamicMeshAttrib
 		}
 	}
 
-	// currently SkinWeightAttributes and generic attributes do not respect bClearExisting, ie they are always cleared
-
-	ResetRegisteredAttributes();
-	
-	SkinWeightAttributes.Reset();
-	for (const TPair<FName, TUniquePtr<FDynamicMeshVertexSkinWeightsAttribute>>& AttribPair : ToMatch.SkinWeightAttributes)
+	// SkinWeights and GenericAttributes require more complex handling...
+	if (bClearExisting)
 	{
-		AttachSkinWeightsAttribute(AttribPair.Key,
-			static_cast<FDynamicMeshVertexSkinWeightsAttribute *>(AttribPair.Value->MakeNew(ParentMesh)));
-	}
+		// discard existing
+		SkinWeightAttributes.Reset();
+		GenericAttributes.Reset();
+		ResetRegisteredAttributes();
 
-	GenericAttributes.Reset();
-	for (const TPair<FName, TUniquePtr<FDynamicMeshAttributeBase>>& AttribPair : ToMatch.GenericAttributes)
+		// register new attributes by name
+		for (const TPair<FName, TUniquePtr<FDynamicMeshVertexSkinWeightsAttribute>>& AttribPair : ToMatch.SkinWeightAttributes)
+		{
+			AttachSkinWeightsAttribute(AttribPair.Key,
+				static_cast<FDynamicMeshVertexSkinWeightsAttribute *>(AttribPair.Value->MakeNew(ParentMesh)));	
+		}
+	}
+	else
 	{
-		AttachAttribute(AttribPair.Key, AttribPair.Value->MakeNew(ParentMesh));
+		// get rid of any attributes in current SkinWeights and Generic sets that are not in ToMatch
+		if (bDiscardExtraAttributes)
+		{
+			SkinWeightAttributesMap ExistingSkinWeights = MoveTemp(SkinWeightAttributes);
+			GenericAttributesMap ExistingGenericAttributes = MoveTemp(GenericAttributes);
+			SkinWeightAttributes.Reset();
+			GenericAttributes.Reset();
+			ResetRegisteredAttributes();
+			for (TPair<FName, TUniquePtr<FDynamicMeshVertexSkinWeightsAttribute>>& AttribPair : ExistingSkinWeights)
+			{
+				if ( ToMatch.SkinWeightAttributes.Contains(AttribPair.Key) )
+				{
+					AttachSkinWeightsAttribute(AttribPair.Key, AttribPair.Value.Release());
+				}
+			}
+			for (TPair<FName, TUniquePtr<FDynamicMeshAttributeBase>>& AttribPair : ExistingGenericAttributes)
+			{
+				if ( ToMatch.GenericAttributes.Contains(AttribPair.Key) )
+				{
+					AttachAttribute(AttribPair.Key, AttribPair.Value.Release());
+				}
+			}
+		}
+
+		// add any new SkinWeight attributes that did not previously exist
+		for (const TPair<FName, TUniquePtr<FDynamicMeshVertexSkinWeightsAttribute>>& AttribPair : ToMatch.SkinWeightAttributes)
+		{
+			TUniquePtr<FDynamicMeshVertexSkinWeightsAttribute>* FoundMatch = SkinWeightAttributes.Find(AttribPair.Key);
+			if (FoundMatch == nullptr)
+			{
+				AttachSkinWeightsAttribute(AttribPair.Key,
+					static_cast<FDynamicMeshVertexSkinWeightsAttribute *>(AttribPair.Value->MakeNew(ParentMesh)));	
+			}
+		}
+
+		// add any new Generic attributes that did not previously exist, matching by name
+		for (const TPair<FName, TUniquePtr<FDynamicMeshAttributeBase>>& AttribPair : ToMatch.GenericAttributes)
+		{
+			TUniquePtr<FDynamicMeshAttributeBase>* FoundMatch = GenericAttributes.Find(AttribPair.Key);
+			if (FoundMatch == nullptr)
+			{
+				AttachAttribute(AttribPair.Key, AttribPair.Value->MakeNew(ParentMesh));	
+			}
+		}
 	}
+
+
+
+
 }
 
 
diff --git a/Engine/Source/Runtime/GeometryCore/Public/DynamicMesh/DynamicMesh3.h b/Engine/Source/Runtime/GeometryCore/Public/DynamicMesh/DynamicMesh3.h
--- a/Engine/Source/Runtime/GeometryCore/Public/DynamicMesh/DynamicMesh3.h
+++ b/Engine/Source/Runtime/GeometryCore/Public/DynamicMesh/DynamicMesh3.h
@@ -262,8 +262,9 @@ public:
 	 * Ensure that all the same extended attributes available in ToMatch are also enabled.
 	 * By default, clears existing attributes, so that there will be an exact match
 	 * If bClearExisting is passed as false, existing attributes are not removed/cleared.
+	 * If bDiscardExtraAttributes=true and bClearExisting=false, extra attributes not in ToMatch are discarded, but existing attributes are not cleared/reset
 	 */
-	void EnableMatchingAttributes(const FDynamicMesh3& ToMatch, bool bClearExisting = true);
+	void EnableMatchingAttributes(const FDynamicMesh3& ToMatch, bool bClearExisting = true, bool bDiscardExtraAttributes = false);
 
 	/**
 	 * Serialization operator for FDynamicMesh3.
diff --git a/Engine/Source/Runtime/GeometryCore/Public/DynamicMesh/DynamicMeshAttributeSet.h b/Engine/Source/Runtime/GeometryCore/Public/DynamicMesh/DynamicMeshAttributeSet.h
--- a/Engine/Source/Runtime/GeometryCore/Public/DynamicMesh/DynamicMeshAttributeSet.h
+++ b/Engine/Source/Runtime/GeometryCore/Public/DynamicMesh/DynamicMeshAttributeSet.h
@@ -106,13 +106,13 @@ public:
 
 
 	/**
-	 * Enable the matching attributes and overlay layers as the reference Copy set, but do not copy any data across.
-	 * By default, clears existing attributes, so that there will be an exact match
-	 * If bClearExisting is passed as false, attribute sets will grow if necessary, and new sets are cleared,
-	 * but existing attributes are not removed/cleared.
-	 * @note Currently SkinWeightAttributes and GenericAttributes are always fully cleared regardless of bClearExisting value
+	 * Enable the matching attributes and overlay layers as the ToMatch set, but do not copy any data across.
+	 * If bClearExisting=true, all existing attributes are cleared, so after the function there is an exact match
+	 * but any existing attribute data is lost (!)
+	 * If bClearExisting=false, new attributes are added, but existing attributes and data are preserved
+	 * If bDiscardExtraAttributes=true and bClearExisting=false, then attributes in the current set but not in ToMatch are discarded, but existing attributes are not cleared/reset
 	 */
-	void EnableMatchingAttributes(const FDynamicMeshAttributeSet& ToMatch, bool bClearExisting = true);
+	void EnableMatchingAttributes(const FDynamicMeshAttributeSet& ToMatch, bool bClearExisting = true, bool bDiscardExtraAttributes = false);
 
 	/** @return the parent mesh for this overlay */
 	const FDynamicMesh3* GetParentMesh() const { return ParentMesh; }
