DynamicMesh: Fix skin weight attributes serialization and comparison, and add to serialization unit tests

diff --git a/Engine/Source/Runtime/GeometryCore/Private/DynamicMesh/DynamicMeshAttributeSet.cpp b/Engine/Source/Runtime/GeometryCore/Private/DynamicMesh/DynamicMeshAttributeSet.cpp
--- a/Engine/Source/Runtime/GeometryCore/Private/DynamicMesh/DynamicMeshAttributeSet.cpp
+++ b/Engine/Source/Runtime/GeometryCore/Private/DynamicMesh/DynamicMeshAttributeSet.cpp
@@ -1124,6 +1124,26 @@ bool FDynamicMeshAttributeSet::IsSameAs(const FDynamicMeshAttributeSet& Other, b
 	}
 	if (!SkinWeightAttributes.IsEmpty())
 	{
+		auto VertexBoneWeightsAreIdentical = [](const AnimationCore::FBoneWeights& BoneWeights, const AnimationCore::FBoneWeights& BoneWeightsOther) -> bool
+		{
+			if (BoneWeights.Num() != BoneWeightsOther.Num())
+			{
+				return false;
+			}
+
+			for (int32 Index = 0; Index < BoneWeights.Num(); ++Index)
+			{
+				// If the weight is the same, the order is nondeterministic. Hence, we need to "manually" look for the same values.
+				const int32 IndexOther = BoneWeightsOther.FindWeightIndexByBone(BoneWeights[Index].GetBoneIndex());
+				if (IndexOther == INDEX_NONE || BoneWeights[Index].GetRawWeight() != BoneWeightsOther[IndexOther].GetRawWeight())
+				{
+					return false;
+				}
+			}
+
+			return true;
+		};
+		
 		SkinWeightAttributesMap::TConstIterator It(SkinWeightAttributes);
 		SkinWeightAttributesMap::TConstIterator ItOther(Other.SkinWeightAttributes);
 		while (It && ItOther)
@@ -1135,29 +1155,45 @@ bool FDynamicMeshAttributeSet::IsSameAs(const FDynamicMeshAttributeSet& Other, b
 			const FDynamicMeshVertexSkinWeightsAttribute* SkinWeights = It->Value.Get();
 			const FDynamicMeshVertexSkinWeightsAttribute* SkinWeightsOther = ItOther->Value.Get();
 
-			if (SkinWeights->VertexBoneWeights.Num() != SkinWeightsOther->VertexBoneWeights.Num())
+			if (!bIgnoreDataLayout)
 			{
-				return false;
-			}
+				if (SkinWeights->VertexBoneWeights.Num() != SkinWeightsOther->VertexBoneWeights.Num())
+				{
+					return false;
+				}
 
-			for (int32 i = 0, Num = SkinWeights->VertexBoneWeights.Num(); i < Num; ++i)
+				for (int32 i = 0, Num = SkinWeights->VertexBoneWeights.Num(); i < Num; ++i)
+				{
+					if (!VertexBoneWeightsAreIdentical(SkinWeights->VertexBoneWeights[i], SkinWeightsOther->VertexBoneWeights[i]))
+					{
+						return false;
+					}
+				}
+			}
+			else
 			{
-				AnimationCore::FBoneWeights BoneWeights = SkinWeights->VertexBoneWeights[i];
-				AnimationCore::FBoneWeights BoneWeightsOther = SkinWeightsOther->VertexBoneWeights[i];
-
-				if (BoneWeights.Num() != BoneWeightsOther.Num())
+				const FRefCountVector& VertexRefCounts = SkinWeights->Parent->GetVerticesRefCounts();
+				const FRefCountVector& VertexRefCountsOther = SkinWeightsOther->Parent->GetVerticesRefCounts();
+				
+				if (VertexRefCounts.GetCount() != VertexRefCountsOther.GetCount())
 				{
 					return false;
 				}
 
-				for (int32 Index = 0; Index < BoneWeights.Num(); ++Index)
+				FRefCountVector::IndexIterator ItVid = VertexRefCounts.BeginIndices();
+				const FRefCountVector::IndexIterator ItVidEnd = VertexRefCounts.EndIndices();
+				FRefCountVector::IndexIterator ItVidOther = VertexRefCountsOther.BeginIndices();
+				const FRefCountVector::IndexIterator ItVidEndOther = VertexRefCountsOther.EndIndices();
+
+				while (ItVid != ItVidEnd && ItVidOther != ItVidEndOther)
 				{
-					// If the weight is the same, the order is nondeterministic. Hence, we need to "manually" look for the same values.
-					const int32 IndexOther = BoneWeightsOther.FindWeightIndexByBone(BoneWeights[Index].GetBoneIndex());
-					if (IndexOther == INDEX_NONE || BoneWeights[Index].GetRawWeight() != BoneWeightsOther[IndexOther].GetRawWeight())
+					if (!VertexBoneWeightsAreIdentical(SkinWeights->VertexBoneWeights[*ItVid], SkinWeightsOther->VertexBoneWeights[*ItVidOther]))
 					{
 						return false;
 					}
+
+					++ItVid;
+					++ItVidOther;
 				}
 			}
 
diff --git a/Engine/Source/Runtime/GeometryCore/Public/DynamicMesh/DynamicVertexSkinWeightsAttribute.h b/Engine/Source/Runtime/GeometryCore/Public/DynamicMesh/DynamicVertexSkinWeightsAttribute.h
--- a/Engine/Source/Runtime/GeometryCore/Public/DynamicMesh/DynamicVertexSkinWeightsAttribute.h
+++ b/Engine/Source/Runtime/GeometryCore/Public/DynamicMesh/DynamicVertexSkinWeightsAttribute.h
@@ -429,7 +429,8 @@ public:
 
 				// Set buffer size to hold number of vertex bone weight arrays, size for vertex each bone weight array, and all individual bone weights.
 				// We also serialize out size of uncompressed buffer to allow for allocation to correct size during loading.
-				uint32 BufferSize = 1 + VertexBoneWeights.Num() + NumBoneWeights;
+				const uint32 NumVertexBoneWeightArrays = bUseVertexCompactMap ? Parent->VertexCount() : VertexBoneWeights.Num();
+				uint32 BufferSize = 1 + NumVertexBoneWeightArrays + IntCastChecked<uint32>(NumBoneWeights);
 				Buffer.SetNumUninitialized(BufferSize);
 				Ar << BufferSize;
 
@@ -447,9 +448,9 @@ public:
 				};
 
 				int32* BufferPtr = Buffer.GetData();
+				*BufferPtr++ = NumVertexBoneWeightArrays;
 				if (!bUseVertexCompactMap)
 				{
-					*BufferPtr++ = VertexBoneWeights.Num();
 					for (const FBoneWeights& BoneWeights : VertexBoneWeights)
 					{
 						WriteBoneWeights(BoneWeights, BufferPtr);
@@ -457,7 +458,6 @@ public:
 				}
 				else
 				{
-					*BufferPtr++ = Parent->VertexCount();
 					for (int32 Vid = 0, Num = VertexBoneWeights.Num(); Vid < Num; ++Vid)
 					{
 						const int32 VidCompact = CompactMaps->GetVertexMapping(Vid);
