Fix "FDynamicMeshAttributeSet::EnableMatchingAttributes" losing generic attributes w/ bClearAttributes=true
(fixes crash in executing fracture w/ fracture mode)

diff --git a/Engine/Source/Runtime/GeometryCore/Private/DynamicMesh/DynamicMeshAttributeSet.cpp b/Engine/Source/Runtime/GeometryCore/Private/DynamicMesh/DynamicMeshAttributeSet.cpp
--- a/Engine/Source/Runtime/GeometryCore/Private/DynamicMesh/DynamicMeshAttributeSet.cpp
+++ b/Engine/Source/Runtime/GeometryCore/Private/DynamicMesh/DynamicMeshAttributeSet.cpp
@@ -353,6 +353,10 @@ void FDynamicMeshAttributeSet::EnableMatchingAttributes(const FDynamicMeshAttrib
 			AttachSkinWeightsAttribute(AttribPair.Key,
 				static_cast<FDynamicMeshVertexSkinWeightsAttribute *>(AttribPair.Value->MakeNew(ParentMesh)));	
 		}
+		for (const TPair<FName, TUniquePtr<FDynamicMeshAttributeBase>>& AttribPair : ToMatch.GenericAttributes)
+		{
+			AttachAttribute(AttribPair.Key, AttribPair.Value->MakeNew(ParentMesh));
+		}
 	}
 	else
 	{
