fixed wrong test from CL 23258977

diff --git a/Engine/Source/Runtime/GeometryCore/Public/Spatial/SparseDynamicPointOctree3.h b/Engine/Source/Runtime/GeometryCore/Public/Spatial/SparseDynamicPointOctree3.h
--- a/Engine/Source/Runtime/GeometryCore/Public/Spatial/SparseDynamicPointOctree3.h
+++ b/Engine/Source/Runtime/GeometryCore/Public/Spatial/SparseDynamicPointOctree3.h
@@ -421,7 +421,7 @@ bool FSparseDynamicPointOctree3::ContainsPoint(int32 PointID) const
 
 void FSparseDynamicPointOctree3::InsertPoint(int32 PointID, const FVector3d& Position)
 {
-	if (ContainsPoint(PointID) == false)
+	if (ContainsPoint(PointID))
 	{
 		checkSlow(false);
 		return;
@@ -721,7 +721,7 @@ void FSparseDynamicPointOctree3::InsertPoint_DynamicExpand(
 	int32 PointID, 
 	TFunctionRef<FVector3d(int)> GetPositionFunc)
 {
-	if (ContainsPoint(PointID) == false)
+	if (ContainsPoint(PointID))
 	{
 		checkSlow(false);
 		return;
