Add BRIO point ordering and update Delaunay2/3 to use it, to improve performance e.g. for grids of points

diff --git a/Engine/Source/Runtime/GeometryCore/Private/CompGeom/Delaunay2.cpp b/Engine/Source/Runtime/GeometryCore/Private/CompGeom/Delaunay2.cpp
--- a/Engine/Source/Runtime/GeometryCore/Private/CompGeom/Delaunay2.cpp
+++ b/Engine/Source/Runtime/GeometryCore/Private/CompGeom/Delaunay2.cpp
@@ -1333,8 +1333,7 @@ namespace DelaunayInternal
 			return false;
 		}
 
-		// TODO: Combine the Z Order Curve ordering with a BRIO, to add enough randomization to break up pathological bad orderings
-		FZOrderCurvePoints InsertOrder;
+		FBRIOPoints InsertOrder;
 		InsertOrder.Compute(Vertices);
 		TArray<int32>& Order = InsertOrder.Order;
 
diff --git a/Engine/Source/Runtime/GeometryCore/Private/CompGeom/Delaunay3.cpp b/Engine/Source/Runtime/GeometryCore/Private/CompGeom/Delaunay3.cpp
--- a/Engine/Source/Runtime/GeometryCore/Private/CompGeom/Delaunay3.cpp
+++ b/Engine/Source/Runtime/GeometryCore/Private/CompGeom/Delaunay3.cpp
@@ -518,8 +518,7 @@ namespace Delaunay3Internal
 			return false;
 		}
 
-		// TODO: Combine the Z Order Curve ordering with a BRIO, to add enough randomization to break up pathological bad orderings
-		FZOrderCurvePoints InsertOrder;
+		FBRIOPoints InsertOrder;
 		InsertOrder.Compute(Vertices);
 		TArray<int32>& Order = InsertOrder.Order;
 
diff --git a/Engine/Source/Runtime/GeometryCore/Private/Spatial/ZOrderCurvePoints.cpp b/Engine/Source/Runtime/GeometryCore/Private/Spatial/ZOrderCurvePoints.cpp
--- a/Engine/Source/Runtime/GeometryCore/Private/Spatial/ZOrderCurvePoints.cpp
+++ b/Engine/Source/Runtime/GeometryCore/Private/Spatial/ZOrderCurvePoints.cpp
@@ -213,11 +213,14 @@ struct TIndexSpatialTree
 	TArray<TArray<int32, TInlineAllocator<4>>> MaxDepthBuckets;
 
 	int32 MaxDepth;
+	// Number of points inserted in the tree
+	int32 Num = 0;
 
 	void Reset()
 	{
 		Nodes.Reset();
 		MaxDepthBuckets.Reset();
+		Num = 0;
 	}
 
 	void Init(TArrayView<const PointT> Points, BoxT Bounds, int32 MaxTreeDepth)
@@ -234,114 +237,210 @@ struct TIndexSpatialTree
 			}
 		}
 
-		auto PlaceInMaxDepthBucket = [this](int32 PtIdx, const PointT& Pt, NodeT& Node, const PointT& Center)
+		for (int32 PtIdx = 0; PtIdx < Points.Num(); PtIdx++)
+		{
+			AddPointHelper(Points, PtIdx, Bounds);
+		}
+	}
+
+	void InitSubset(TArrayView<const PointT> Points, TArrayView<const int32> PointIndices, BoxT Bounds, int32 MaxTreeDepth)
+	{
+		Reset();
+		Nodes.Emplace();
+		MaxDepth = MaxTreeDepth;
+
+		if (Bounds.IsEmpty())
 		{
-			int32 SubIdx = NodeT::GetSubIdx(Center, Pt);
-			int32& EncodedBucketIndex = Node.Children[SubIdx];
-			if (NodeT::ToBucketIndex(EncodedBucketIndex) == NodeT::InvalidIndex)
+			for (int32 PtIdx : PointIndices)
 			{
-				EncodedBucketIndex = NodeT::FromBucketIndex(MaxDepthBuckets.Emplace());
+				const PointT& Pt = Points[PtIdx];
+				Bounds.Contain(Pt);
 			}
-			MaxDepthBuckets[NodeT::ToBucketIndex(EncodedBucketIndex)].Add(PtIdx);
-		};
+		}
 
-		for (int32 PtIdx = 0; PtIdx < Points.Num(); PtIdx++)
+		for (int32 PtIdx : PointIndices)
 		{
-			const PointT& Pt = Points[PtIdx];
-			int32 CurNodeIdx = 0, CurDepth = 0;
-			BoxT CurBox = Bounds;
-			while (true)
+			AddPointHelper(Points, PtIdx, Bounds);
+		}
+	}
+	
+protected:
+
+	void PlaceInMaxDepthBucketHelper(int32 PtIdx, const PointT& Pt, NodeT& Node, const PointT& Center)
+	{
+		int32 SubIdx = NodeT::GetSubIdx(Center, Pt);
+		int32& EncodedBucketIndex = Node.Children[SubIdx];
+		if (NodeT::ToBucketIndex(EncodedBucketIndex) == NodeT::InvalidIndex)
+		{
+			EncodedBucketIndex = NodeT::FromBucketIndex(MaxDepthBuckets.Emplace());
+		}
+		MaxDepthBuckets[NodeT::ToBucketIndex(EncodedBucketIndex)].Add(PtIdx);
+	}
+
+	void AddPointHelper(TArrayView<const PointT> Points, int32 PtIdx, const BoxT& Bounds)
+	{
+		Num++;
+		const PointT& Pt = Points[PtIdx];
+		int32 CurNodeIdx = 0, CurDepth = 0;
+		BoxT CurBox = Bounds;
+		while (true)
+		{
+			PointT CurCenter = CurBox.Center();
+
+			if (CurDepth >= MaxDepth)
 			{
-				PointT CurCenter = CurBox.Center();
+				PlaceInMaxDepthBucketHelper(PtIdx, Pt, Nodes[CurNodeIdx], CurCenter);
+				break;
+			}
 
-				if (CurDepth >= MaxDepth)
+			// Node is a leaf that directly stores point indices
+			if (Nodes[CurNodeIdx].Children[0] >= NodeT::InvalidIndex)
+			{
+				// Try to add the point directly to the node
+				if (Nodes[CurNodeIdx].AddPtAtFirstZeroIdx(PtIdx))
 				{
-					PlaceInMaxDepthBucket(PtIdx, Pt, Nodes[CurNodeIdx], CurCenter);
 					break;
 				}
-
-				// Node is a leaf that directly stores point indices
-				if (Nodes[CurNodeIdx].Children[0] >= NodeT::InvalidIndex)
+				else // didn't have room for the point: switch node to instead store indices to new children, and propagate old children down
 				{
-					// Try to add the point directly to the node
-					if (Nodes[CurNodeIdx].AddPtAtFirstZeroIdx(PtIdx))
+					NodeT Old = Nodes[CurNodeIdx];
+					// Create an empty non-leaf node
+					for (int32 SubIdx = 0; SubIdx < NodeT::NumChildren; SubIdx++)
 					{
-						break;
+						Nodes[CurNodeIdx].Children[SubIdx] = NodeT::FromChildNodeIndex(NodeT::InvalidIndex);
 					}
-					else // didn't have room for the point: switch node to instead store indices to new children, and propagate old children down
+					if (CurDepth + 1 >= MaxDepth)
 					{
-						NodeT Old = Nodes[CurNodeIdx];
-						// Create an empty non-leaf node
-						for (int32 SubIdx = 0; SubIdx < NodeT::NumChildren; SubIdx++)
+						// the new nodes are final bucket nodes; immediately push the children through to their respective max depth buckets
+						for (int32 SubIdx = 0; SubIdx < NodeT::NumPerLeaf; SubIdx++)
 						{
-							Nodes[CurNodeIdx].Children[SubIdx] = NodeT::FromChildNodeIndex(NodeT::InvalidIndex);
-						}
-						if (CurDepth + 1 >= MaxDepth)
-						{
-							// the new nodes are final bucket nodes; immediately push the children through to their respective max depth buckets
-							for (int32 SubIdx = 0; SubIdx < NodeT::NumPerLeaf; SubIdx++)
+							int32 OldPtIdx = NodeT::ToPointIndex(Old.Children[SubIdx]);
+							checkSlow(OldPtIdx > NodeT::InvalidIndex);
+							PointT OldPt = Points[OldPtIdx];
+							int32 OldPtSubIdx = NodeT::GetSubIdx(CurCenter, OldPt);
+							int32 NewChildNodeIdx = NodeT::ToChildNodeIndex(Nodes[CurNodeIdx].Children[OldPtSubIdx]);
+							if (NewChildNodeIdx == NodeT::InvalidIndex)
 							{
-								int32 OldPtIdx = NodeT::ToPointIndex(Old.Children[SubIdx]);
-								checkSlow(OldPtIdx > NodeT::InvalidIndex);
-								PointT OldPt = Points[OldPtIdx];
-								int32 OldPtSubIdx = NodeT::GetSubIdx(CurCenter, OldPt);
-								int32 NewChildNodeIdx = NodeT::ToChildNodeIndex(Nodes[CurNodeIdx].Children[OldPtSubIdx]);
-								if (NewChildNodeIdx == NodeT::InvalidIndex)
-								{
-									NewChildNodeIdx = Nodes.Emplace();
-									Nodes[CurNodeIdx].Children[OldPtSubIdx] = NodeT::FromChildNodeIndex(NewChildNodeIdx);
-								}
-								NodeT& NewChildNode = Nodes[NewChildNodeIdx];
-								PointT ChildCenter = NodeT::GetSubCenter(Bounds, CurCenter, OldPtSubIdx);
-								PlaceInMaxDepthBucket(OldPtIdx, OldPt, NewChildNode, ChildCenter);
+								NewChildNodeIdx = Nodes.Emplace();
+								Nodes[CurNodeIdx].Children[OldPtSubIdx] = NodeT::FromChildNodeIndex(NewChildNodeIdx);
 							}
+							NodeT& NewChildNode = Nodes[NewChildNodeIdx];
+							PointT ChildCenter = NodeT::GetSubCenter(Bounds, CurCenter, OldPtSubIdx);
+							PlaceInMaxDepthBucketHelper(OldPtIdx, OldPt, NewChildNode, ChildCenter);
 						}
-						else
+					}
+					else
+					{
+						// the new nodes are regular leaf nodes; push the children down to the leaves
+						for (int32 SubIdx = 0; SubIdx < NodeT::NumPerLeaf; SubIdx++)
 						{
-							// the new nodes are regular leaf nodes; push the children down to the leaves
-							for (int32 SubIdx = 0; SubIdx < NodeT::NumPerLeaf; SubIdx++)
+							int32 OldPtIdx = NodeT::ToPointIndex(Old.Children[SubIdx]);
+							checkSlow(OldPtIdx > NodeT::InvalidIndex);
+							PointT OldPt = Points[OldPtIdx];
+							int32 OldPtSubIdx = NodeT::GetSubIdx(CurCenter, OldPt);
+							int32 NewChildNodeIdx = NodeT::ToChildNodeIndex(Nodes[CurNodeIdx].Children[OldPtSubIdx]);
+							if (NewChildNodeIdx == NodeT::InvalidIndex)
 							{
-								int32 OldPtIdx = NodeT::ToPointIndex(Old.Children[SubIdx]);
-								checkSlow(OldPtIdx > NodeT::InvalidIndex);
-								PointT OldPt = Points[OldPtIdx];
-								int32 OldPtSubIdx = NodeT::GetSubIdx(CurCenter, OldPt);
-								int32 NewChildNodeIdx = NodeT::ToChildNodeIndex(Nodes[CurNodeIdx].Children[OldPtSubIdx]);
-								if (NewChildNodeIdx == NodeT::InvalidIndex)
-								{
-									NewChildNodeIdx = Nodes.Emplace();
-									Nodes[CurNodeIdx].Children[OldPtSubIdx] = NodeT::FromChildNodeIndex(NewChildNodeIdx);
-								}
-								NodeT& NewChildNode = Nodes[NewChildNodeIdx];
-								bool bPlaced = NewChildNode.AddPtAtFirstZeroIdx(OldPtIdx);
-								checkSlow(bPlaced); // should always succeed
+								NewChildNodeIdx = Nodes.Emplace();
+								Nodes[CurNodeIdx].Children[OldPtSubIdx] = NodeT::FromChildNodeIndex(NewChildNodeIdx);
 							}
+							NodeT& NewChildNode = Nodes[NewChildNodeIdx];
+							bool bPlaced = NewChildNode.AddPtAtFirstZeroIdx(OldPtIdx);
+							checkSlow(bPlaced); // should always succeed
 						}
 					}
 				}
+			}
 
-				check(NodeT::IsNodeIndex(Nodes[CurNodeIdx].Children[0])); // should only reach here if we need to iterate further down the tree
-				int32 PtSubIdx = NodeT::GetSubIdx(CurCenter, Pt);
-				int32 OldNodeIdx = CurNodeIdx;
-				CurNodeIdx = NodeT::ToChildNodeIndex(Nodes[CurNodeIdx].Children[PtSubIdx]);
-				if (CurNodeIdx == NodeT::InvalidIndex) // if child didn't exist yet, create it here
-				{
-					CurNodeIdx = Nodes.Emplace();
-					Nodes[OldNodeIdx].Children[PtSubIdx] = NodeT::FromChildNodeIndex(CurNodeIdx);
-				}
-				CurBox = NodeT::GetSubBounds(CurBox, CurCenter, PtSubIdx);
-				CurDepth++;
+			check(NodeT::IsNodeIndex(Nodes[CurNodeIdx].Children[0])); // should only reach here if we need to iterate further down the tree
+			int32 PtSubIdx = NodeT::GetSubIdx(CurCenter, Pt);
+			int32 OldNodeIdx = CurNodeIdx;
+			CurNodeIdx = NodeT::ToChildNodeIndex(Nodes[CurNodeIdx].Children[PtSubIdx]);
+			if (CurNodeIdx == NodeT::InvalidIndex) // if child didn't exist yet, create it here
+			{
+				CurNodeIdx = Nodes.Emplace();
+				Nodes[OldNodeIdx].Children[PtSubIdx] = NodeT::FromChildNodeIndex(CurNodeIdx);
 			}
+			CurBox = NodeT::GetSubBounds(CurBox, CurCenter, PtSubIdx);
+			CurDepth++;
 		}
 	}
+	
+public:
 
-	// Add indices of all points in the quadtree to the PointOrder array, following a z order curve when extracting them.
+	// Add indices of all points in the quad/octree to the PointOrder array, following a z order curve when extracting them.
 	void ExtractZOrderCurve(TArray<int32>& PointOrder, int32 ExpectedNumberPoints = 0) const
 	{
 		PointOrder.Reserve(PointOrder.Num() + ExpectedNumberPoints);
 		ExtractZOrderCurveHelper(PointOrder, Nodes[0], 0);
 	}
 
+	// Overwrite the values in the OrderView array view with the indices in the quad/octree, extracted in a z-order curve
+	// The OrderView should have the same size as the number of points in the tree
+	void InPlaceExtractZOrderCurve(TArrayView<int32> OrderView, bool bReverse) const
+	{
+		if (!ensure(Num == OrderView.Num()))
+		{
+			return;
+		}
+		int32 Idx = 0;
+		InPlaceExtractZOrderCurveHelper(OrderView, Nodes[0], 0, bReverse, Idx);
+		check(Idx == OrderView.Num());
+	}
+
 protected:
+
+	void InPlaceExtractZOrderCurveHelper(TArrayView<int32> OrderView, const NodeT& Node, int32 CurDepth, bool bReverse, int& CurIdx) const
+	{
+		// at max depth: append all the non-empty buckets
+		if (CurDepth >= MaxDepth)
+		{
+			for (int32 SubIdx = 0; SubIdx < NodeT::NumChildren; SubIdx++)
+			{
+				int32 ChildValue = Node.Children[bReverse ? NodeT::NumChildren - 1 - SubIdx : SubIdx];
+				if (ChildValue == NodeT::InvalidIndex)
+				{
+					continue;
+				}
+
+				for (int32 Val : MaxDepthBuckets[NodeT::ToBucketIndex(ChildValue)])
+				{
+					OrderView[CurIdx++] = Val;
+				}
+			}
+			return;
+		}
+
+		// at leaf: add all the attached points
+		if (Node.IsLeaf())
+		{
+			// Note: bReverse doesn't affect this ordering, because leaf points are in arbitrary order
+			for (int32 SubIdx = 0; SubIdx < NodeT::NumPerLeaf; SubIdx++)
+			{
+				int32 ChildValue = Node.Children[SubIdx];
+				if (ChildValue == NodeT::InvalidIndex)
+				{
+					break;
+				}
+
+				OrderView[CurIdx++] = NodeT::ToPointIndex(ChildValue);
+			}
+			return;
+		}
+
+		// recurse to children
+		for (int32 SubIdx = 0; SubIdx < NodeT::NumChildren; SubIdx++)
+		{
+			int32 ChildIndex = NodeT::ToChildNodeIndex(Node.Children[bReverse ? NodeT::NumChildren - 1 - SubIdx : SubIdx]);
+			if (ChildIndex == NodeT::InvalidIndex)
+			{
+				continue;
+			}
+
+			InPlaceExtractZOrderCurveHelper(OrderView, Nodes[ChildIndex], CurDepth + 1, bReverse, CurIdx);
+		}
+	}
+
 	// recursive helper for following a z order curve to list the points
 	void ExtractZOrderCurveHelper(TArray<int32>& PointOrder, const NodeT& Node, int32 CurDepth) const
 	{
@@ -407,7 +506,64 @@ void ZOrderCurve3(TArray<int32>& PointOrder, TArrayView<const TVector<RealType>>
 	Tree.ExtractZOrderCurve(PointOrder, Points.Num());
 }
 
-} // end of local helper function namespace
+template<typename RealType, typename VectorType, typename BoxType, typename NodeType>
+void ApplyBRIO(TArray<int32>& PointOrder, TArrayView<const VectorType> Points, BoxType Bounds, int32 MaxTreeDepth, const FRandomStream& Random)
+{
+	if (Bounds.IsEmpty())
+	{
+		Bounds.Contain(Points);
+	}
+
+	int32 NumPoints = Points.Num();
+
+	int32 OrderInitialNum = PointOrder.Num();
+	PointOrder.SetNumUninitialized(PointOrder.Num() + NumPoints);
+	TArrayView<int32> OrderSlice(PointOrder.GetData() + OrderInitialNum, NumPoints);
+	for (int32 Idx = 0; Idx < Points.Num(); ++Idx)
+	{
+		OrderSlice[Idx] = Idx;
+	}
+
+	if (NumPoints < 2) // not enough points to need any BRIO phases
+	{
+		return;
+	}
+
+	int32 NumPhases = (int32)FMath::Log2((double)NumPoints);
+	TArray<int32> PhaseEnds;
+	PhaseEnds.SetNumUninitialized(NumPhases);
+	int32 CurEnd = NumPoints;
+	PhaseEnds[NumPhases - 1] = CurEnd;
+	for (int32 Phase = NumPhases - 1; Phase > 0; --Phase)
+	{
+		for (int32 Idx = 0; Idx < CurEnd; ++Idx)
+		{
+			if (Random.GetFraction() < .5)
+			{
+				CurEnd--;
+				Swap(OrderSlice[Idx], OrderSlice[CurEnd]);
+				Idx--;
+			}
+		}
+		PhaseEnds[Phase - 1] = CurEnd;
+	}
+
+	TIndexSpatialTree<NodeType, VectorType, BoxType> Tree;
+	for (int32 Phase = 0, PhaseStart = 0; Phase < NumPhases; PhaseStart = PhaseEnds[Phase++])
+	{
+		int32 PhaseEnd = PhaseEnds[Phase];
+		if (PhaseEnd == PhaseStart)
+		{
+			continue;
+		}
+		Tree.Reset();
+		TArrayView<int32> PhaseSlice(OrderSlice.GetData() + PhaseStart, PhaseEnd - PhaseStart);
+		Tree.InitSubset(Points, PhaseSlice, Bounds, MaxTreeDepth);
+		Tree.InPlaceExtractZOrderCurve(PhaseSlice, bool(Phase & 1) /*alternate phases have reversed order*/);
+	}
+}
+
+} // end of local helper function namespace for ZOrderCurvePoints functions
 
 
 void FZOrderCurvePoints::Compute(TArrayView<const FVector2d> Points, const FAxisAlignedBox2d& Bounds)
@@ -425,4 +581,21 @@ void FZOrderCurvePoints::Compute(TArrayView<const FVector3d> Points, const FAxis
 void FZOrderCurvePoints::Compute(TArrayView<const FVector3f> Points, const FAxisAlignedBox3f& Bounds)
 {
 	ZOrderCurvePointsLocal::ZOrderCurve3<float>(Order, Points, Bounds, MaxTreeDepth);
+}
+
+void FBRIOPoints::Compute(TArrayView<const FVector2d> Points, const FAxisAlignedBox2d& Bounds, const FRandomStream& Random)
+{
+	ZOrderCurvePointsLocal::ApplyBRIO<double, FVector2d, FAxisAlignedBox2d, ZOrderCurvePointsLocal::FQuadTreeNode>(Order, Points, Bounds, MaxTreeDepth, Random);
+}
+void FBRIOPoints::Compute(TArrayView<const FVector2f> Points, const FAxisAlignedBox2f& Bounds, const FRandomStream& Random)
+{
+	ZOrderCurvePointsLocal::ApplyBRIO<float, FVector2f, FAxisAlignedBox2f, ZOrderCurvePointsLocal::FQuadTreeNode>(Order, Points, Bounds, MaxTreeDepth, Random);
+}
+void FBRIOPoints::Compute(TArrayView<const FVector3d> Points, const FAxisAlignedBox3d& Bounds, const FRandomStream& Random)
+{
+	ZOrderCurvePointsLocal::ApplyBRIO<double, FVector3d, FAxisAlignedBox3d, ZOrderCurvePointsLocal::FOctreeNode>(Order, Points, Bounds, MaxTreeDepth, Random);
+}
+void FBRIOPoints::Compute(TArrayView<const FVector3f> Points, const FAxisAlignedBox3f& Bounds, const FRandomStream& Random)
+{
+	ZOrderCurvePointsLocal::ApplyBRIO<float, FVector3f, FAxisAlignedBox3f, ZOrderCurvePointsLocal::FOctreeNode>(Order, Points, Bounds, MaxTreeDepth, Random);
 }
\ No newline at end of file
diff --git a/Engine/Source/Runtime/GeometryCore/Public/Spatial/ZOrderCurvePoints.h b/Engine/Source/Runtime/GeometryCore/Public/Spatial/ZOrderCurvePoints.h
--- a/Engine/Source/Runtime/GeometryCore/Public/Spatial/ZOrderCurvePoints.h
+++ b/Engine/Source/Runtime/GeometryCore/Public/Spatial/ZOrderCurvePoints.h
@@ -8,6 +8,7 @@
 #include "CoreMinimal.h"
 #include "Math/UnrealMathSSE.h"
 #include "Math/Vector2D.h"
+#include "Math/RandomStream.h"
 
 namespace UE
 {
@@ -52,5 +53,42 @@ public:
 
 };
 
+/**
+ * Generate a "Biased Randomized Insertion Order" (BRIO) ordering for a point set by randomly bucketing the points, then applying Z-Order-Curve sorting (above) to each bucket
+ * This gives an ordering with point locality but enough randomization to typically avoid worst case behavior if using the ordering for Delaunay meshing
+ */
+struct GEOMETRYCORE_API FBRIOPoints
+{
+	/// Inputs
+
+	// Maximum depth of quad/octree used to distribute points in space
+	int32 MaxTreeDepth = 10;
+
+
+	/// Outputs
+
+	TArray<int32> Order;
+
+public:
+
+	FBRIOPoints(int32 MaxTreeDepth = 10) : MaxTreeDepth(MaxTreeDepth)
+	{
+	}
+
+	/**
+	 * Compute an ordering for the input points that attempts to order points so adjacent points are close together spatially, while still keeping the order sufficiently randomized to expect to avoid worst-case orderings for Delaunay triangulation
+	 * Supports both 2D and 3D points.
+	 *
+	 * @param Points				Points to re-order
+	 * @param Bounds				Optional bounding box of the points; if an empty box is passed in, it will be computed for you
+	 */
+	void Compute(TArrayView<const FVector2d> Points, const FAxisAlignedBox2d& Bounds = FAxisAlignedBox2d::Empty(), const FRandomStream& Random = FRandomStream());
+	void Compute(TArrayView<const FVector2f> Points, const FAxisAlignedBox2f& Bounds = FAxisAlignedBox2f::Empty(), const FRandomStream& Random = FRandomStream());
+	void Compute(TArrayView<const FVector3d> Points, const FAxisAlignedBox3d& Bounds = FAxisAlignedBox3d::Empty(), const FRandomStream& Random = FRandomStream());
+	void Compute(TArrayView<const FVector3f> Points, const FAxisAlignedBox3f& Bounds = FAxisAlignedBox3f::Empty(), const FRandomStream& Random = FRandomStream());
+
+};
+
+
 } // end namespace UE::Geometry
 } // end namespace UE
