make sparse dynamic octree use range iteration to find relevant root cells for point containment and range queries, when there are enough root cells for this to make sense

diff --git a/Engine/Source/Runtime/GeometryCore/Private/Spatial/SparseDynamicOctree3.cpp b/Engine/Source/Runtime/GeometryCore/Private/Spatial/SparseDynamicOctree3.cpp
--- a/Engine/Source/Runtime/GeometryCore/Private/Spatial/SparseDynamicOctree3.cpp
+++ b/Engine/Source/Runtime/GeometryCore/Private/Spatial/SparseDynamicOctree3.cpp
@@ -317,26 +317,13 @@ void FSparseDynamicOctree3::ContainmentQuery(
 	const FVector3d& Point,
 	TFunctionRef<void(int)> ObjectIDFunc) const
 {
-	// todo: this should take advantage of raster!
-
 	// always process spill objects
 	for (int ObjectID : SpillObjectSet)
 	{
 		ObjectIDFunc(ObjectID);
 	}
 
-	TArray<const FSparseOctreeCell*, TInlineAllocator<32>> Queue;
-
-	// start at root cells
-	RootCells.AllocatedIteration([&](const uint32* RootCellID)
-	{
-		const FSparseOctreeCell* RootCell = &Cells[*RootCellID];
-		if ( GetCellBox(*RootCell, MaxExpandFactor).Contains(Point) )
-		{
-			Queue.Add(&Cells[*RootCellID]);
-		}
-	});
-
+	TArray<const FSparseOctreeCell*, TInlineAllocator<32>> Queue = InitializeQueryQueue(Point);
 
 	while (Queue.Num() > 0)
 	{
@@ -369,8 +356,6 @@ bool FSparseDynamicOctree3::ContainmentQueryCancellable(
 	const FVector3d& Point,
 	TFunctionRef<bool(int)> ObjectIDFunc) const
 {
-	// todo: this should take advantage of raster!
-
 	// always process spill objects
 	for (int ObjectID : SpillObjectSet)
 	{
@@ -380,17 +365,7 @@ bool FSparseDynamicOctree3::ContainmentQueryCancellable(
 		}
 	}
 
-	TArray<const FSparseOctreeCell*, TInlineAllocator<32>> Queue;
-
-	// start at root cells
-	RootCells.AllocatedIteration([&](const uint32* RootCellID)
-	{
-		const FSparseOctreeCell* RootCell = &Cells[*RootCellID];
-		if ( GetCellBox(*RootCell, MaxExpandFactor).Contains(Point) )
-		{
-			Queue.Add(&Cells[*RootCellID]);
-		}
-	});
+	TArray<const FSparseOctreeCell*, TInlineAllocator<32>> Queue = InitializeQueryQueue(Point);
 
 
 	while (Queue.Num() > 0)
@@ -434,26 +409,13 @@ void FSparseDynamicOctree3::RangeQuery(
 	const FAxisAlignedBox3d& Bounds,
 	TFunctionRef<void(int)> ObjectIDFunc) const
 {
-	// todo: this should take advantage of raster!
-
 	// always process spill objects
 	for (int ObjectID : SpillObjectSet)
 	{
 		ObjectIDFunc(ObjectID);
 	}
 
-	TArray<const FSparseOctreeCell*> Queue;
-
-	// start at root cells
-	RootCells.AllocatedIteration([&](const uint32* RootCellID)
-	{
-		const FSparseOctreeCell* RootCell = &Cells[*RootCellID];
-		if ( GetCellBox(*RootCell, MaxExpandFactor).Intersects(Bounds) )
-		{
-			Queue.Add(&Cells[*RootCellID]);
-		}
-	});
-
+	TArray<const FSparseOctreeCell*, TInlineAllocator<32>> Queue = InitializeQueryQueue(Bounds);
 
 	while (Queue.Num() > 0)
 	{
@@ -489,25 +451,13 @@ void FSparseDynamicOctree3::RangeQuery(
 	const FAxisAlignedBox3d& Bounds,
 	TArray<int>& ObjectIDs) const
 {
-	// todo: this should take advantage of raster!
-
 	// always collect spill objects
 	for (int ObjectID : SpillObjectSet)
 	{
 		ObjectIDs.Add(ObjectID);
 	}
 
-	TArray<const FSparseOctreeCell*> Queue;
-
-	// start at root cells
-	RootCells.AllocatedIteration([&](const uint32* RootCellID)
-	{
-		const FSparseOctreeCell* RootCell = &Cells[*RootCellID];
-		if (GetCellBox(*RootCell, MaxExpandFactor).Intersects(Bounds))
-		{
-			Queue.Add(&Cells[*RootCellID]);
-		}
-	});
+	TArray<const FSparseOctreeCell*, TInlineAllocator<32>> Queue = InitializeQueryQueue(Bounds);
 
 	while (Queue.Num() > 0)
 	{
@@ -546,17 +496,7 @@ void FSparseDynamicOctree3::ParallelRangeQuery(
 		ObjectIDs.Add(ObjectID);
 	}
 
-	TArray<const FSparseOctreeCell*> Queue;
-
-	// start at root cells
-	RootCells.AllocatedIteration([&](const uint32* RootCellID)
-	{
-		const FSparseOctreeCell* RootCell = &Cells[*RootCellID];
-		if (GetCellBox(*RootCell, MaxExpandFactor).Intersects(Bounds))
-		{
-			Queue.Add(&Cells[*RootCellID]);
-		}
-	});
+	TArray<const FSparseOctreeCell*, TInlineAllocator<32>> Queue = InitializeQueryQueue(Bounds);
 
 	// parallel strategy here is to collect each rootcell subtree into a separate
 	// array, and then merge them. Although this means we do some dynamic memory
diff --git a/Engine/Source/Runtime/GeometryCore/Public/Spatial/SparseDynamicOctree3.h b/Engine/Source/Runtime/GeometryCore/Public/Spatial/SparseDynamicOctree3.h
--- a/Engine/Source/Runtime/GeometryCore/Public/Spatial/SparseDynamicOctree3.h
+++ b/Engine/Source/Runtime/GeometryCore/Public/Spatial/SparseDynamicOctree3.h
@@ -460,6 +460,78 @@ protected:
 		const FAxisAlignedBox3d& Bounds,
 		TArray<int>& ObjectIDs) const;
 
+private:
+	// helper to find the root-level cells that could intersect with a given query object, specialized for point containment queries
+	TArray<const FSparseOctreeCell*, TInlineAllocator<32>> InitializeQueryQueue(const FVector3d& Point) const
+	{
+		TArray<const FSparseOctreeCell*, TInlineAllocator<32>> Queue;
+
+		// Skip range iteration if there are not many root cells -- should be faster to just directly iterate all the root cells in that case
+		constexpr int32 MinCountForRangeQuery = 10;
+		if (RootCells.GetCount() > MinCountForRangeQuery)
+		{
+			// because of the cell expand factor, search a range for cells that contain the point
+			FVector3d RootBoundExpand(RootDimension * MaxExpandFactor);
+			FVector3i RootMinIndex = PointToIndex(0, Point - RootBoundExpand);
+			FVector3i RootMaxIndex = PointToIndex(0, Point + RootBoundExpand);
+			// double-check that there are enough root cells that we are likely to save time by local iteration
+			FVector3i QuerySize = RootMaxIndex - RootMinIndex + FVector3i(1, 1, 1);
+			if (RootCells.GetCount() > QuerySize.X * QuerySize.Y * QuerySize.Z)
+			{
+				RootCells.RangeIteration(RootMinIndex, RootMaxIndex, [&](uint32 RootCellID)
+				{
+					Queue.Add(&Cells[RootCellID]);
+				});
+				return Queue;
+			}
+		}
+
+		RootCells.AllocatedIteration([&](const uint32* RootCellID)
+		{
+			const FSparseOctreeCell* RootCell = &Cells[*RootCellID];
+			if (GetCellBox(*RootCell, MaxExpandFactor).Contains(Point))
+			{
+				Queue.Add(&Cells[*RootCellID]);
+			}
+		});
+		return Queue;
+	}
+
+	// helper to find the root-level cells that could intersect with a given query object, specialized for range queries
+	TArray<const FSparseOctreeCell*, TInlineAllocator<32>> InitializeQueryQueue(const FAxisAlignedBox3d& Bounds) const
+	{
+		TArray<const FSparseOctreeCell*, TInlineAllocator<32>> Queue;
+
+		// Skip range iteration if there are not many root cells -- should be faster to just directly iterate all the root cells in that case
+		constexpr int32 MinCountForRangeQuery = 10;
+		if (RootCells.GetCount() > MinCountForRangeQuery)
+		{
+			FVector3d RootBoundExpand(RootDimension * MaxExpandFactor);
+			FVector3i RootMinIndex = PointToIndex(0, Bounds.Min - RootBoundExpand);
+			FVector3i RootMaxIndex = PointToIndex(0, Bounds.Max + RootBoundExpand);
+			// double-check that there are enough root cells that we are likely to save time by local iteration
+			FVector3i QuerySize = RootMaxIndex - RootMinIndex + FVector3i(1, 1, 1);
+			if (RootCells.GetCount() > QuerySize.X * QuerySize.Y * QuerySize.Z) // There are enough root cells that we could save time by local iteration
+			{
+				RootCells.RangeIteration(RootMinIndex, RootMaxIndex, [&](uint32 RootCellID)
+				{
+					// No need to check bound intersection since we're only iterating over cells that are in bounds
+					Queue.Add(&Cells[RootCellID]);
+				});
+				return Queue;
+			}
+		}
+
+		RootCells.AllocatedIteration([&](const uint32* RootCellID)
+		{
+			const FSparseOctreeCell* RootCell = &Cells[*RootCellID];
+			if (GetCellBox(*RootCell, MaxExpandFactor).Intersects(Bounds))
+			{
+				Queue.Add(&Cells[*RootCellID]);
+			}
+		});
+		return Queue;
+	}
 };
 
 
