Removed remaining TIsSame usage.

diff --git a/Engine/Source/Runtime/Engine/Public/MeshCardRepresentation.h b/Engine/Source/Runtime/Engine/Public/MeshCardRepresentation.h
--- a/Engine/Source/Runtime/Engine/Public/MeshCardRepresentation.h
+++ b/Engine/Source/Runtime/Engine/Public/MeshCardRepresentation.h
@@ -62,7 +62,7 @@ public:
 	}
 
 	// Conversion from other type.
-	template<typename FArg, TEMPLATE_REQUIRES(!TIsSame<T, FArg>::Value)>
+	template<typename FArg, TEMPLATE_REQUIRES(!std::is_same_v<T, FArg>)>
 	explicit TLumenCardOBB(const TLumenCardOBB<FArg>& From)
 	{
 		Origin = UE::Math::TVector<T>(From.Origin);
diff --git a/Engine/Source/Runtime/GeometryCore/Public/CompGeom/ExactPredicates.h b/Engine/Source/Runtime/GeometryCore/Public/CompGeom/ExactPredicates.h
--- a/Engine/Source/Runtime/GeometryCore/Public/CompGeom/ExactPredicates.h
+++ b/Engine/Source/Runtime/GeometryCore/Public/CompGeom/ExactPredicates.h
@@ -213,7 +213,7 @@ RealType InSphere3(const TVector<RealType>& A, const TVector<RealType>& B, const
 	double PC[3]{ (double)C.X, (double)C.Y, (double)C.Z };
 	double PD[3]{ (double)D.X, (double)D.Y, (double)D.Z };
 	double PE[3]{ (double)E.X, (double)E.Y, (double)E.Z };
-	if constexpr (TIsSame<RealType, double>::Value)
+	if constexpr (std::is_same_v<RealType, double>)
 	{
 		return InSphere(PA, PB, PC, PD, PE);
 	}
