add a merge-only api to convex decomposition

diff --git a/Engine/Source/Runtime/GeometryCore/Private/CompGeom/ConvexDecomposition3.cpp b/Engine/Source/Runtime/GeometryCore/Private/CompGeom/ConvexDecomposition3.cpp
--- a/Engine/Source/Runtime/GeometryCore/Private/CompGeom/ConvexDecomposition3.cpp
+++ b/Engine/Source/Runtime/GeometryCore/Private/CompGeom/ConvexDecomposition3.cpp
@@ -714,7 +714,10 @@ struct FPartialCutResult
 
 void FConvexDecomposition3::FConvexPart::Compact()
 {
-	ensure(!bIsCompact);
+	if (bIsCompact)
+	{
+		return;
+	}
 
 	// remove all triangles, and all verties that are not reference by the hull
 	// and remap HullTriangles to the new compact indices
@@ -1189,7 +1192,45 @@ void FConvexDecomposition3::UpdateProximitiesAfterSplit(int32 SplitIdx, int32 Ne
 	}
 }
 
-int32 FConvexDecomposition3::MergeBest(int32 TargetNumParts, double MaxErrorTolerance, double MinThicknessToleranceWorldSpace, bool bAllowCompact)
+void FConvexDecomposition3::InitializeFromHulls(int32 NumHulls, TFunctionRef<double(int32)> HullVolumes, TFunctionRef<int32(int32)> HullNumVertices, TFunctionRef<FVector3d(int32, int32)> HullVertices, TArrayView<const TPair<int32, int32>> Proximity)
+{
+	Decomposition.Empty(NumHulls);
+	for (int32 HullIdx = 0; HullIdx < NumHulls; ++HullIdx)
+	{
+		FConvexPart* Convex = new FConvexPart(true);
+		Convex->bGeometryVolumeUnreliable = false;
+		Convex->bMustMerge = false;
+		double Volume = HullVolumes(HullIdx);
+		Convex->GeoVolume = Volume;
+		Convex->HullVolume = Volume;
+		Convex->GeoCenter = FVector3d(0, 0, 0);
+		FAxisAlignedBox3d Bounds;
+		int32 NumVertices = HullNumVertices(HullIdx);
+		for (int32 Idx = 0; Idx < NumVertices; ++Idx)
+		{
+			FVector3d Vertex = HullVertices(HullIdx, Idx);
+			Convex->GeoCenter += Vertex;
+			Convex->InternalGeo.AppendVertex(Vertex);
+			Bounds.Contain(Vertex);
+		}
+		if (ensure(NumVertices > 0))
+		{
+			Convex->GeoCenter /= NumVertices;
+		}
+		Convex->Bounds = Bounds;
+		// Note: For merging, actual triangles + planes not required
+		Decomposition.Add(Convex);
+	}
+
+	for (const TPair<int32, int32>& Link : Proximity)
+	{
+		int32 ProxIdx = Proximities.Emplace(FIndex2i(Link.Key, Link.Value), FPlane3d(), false /*bPlaneSeparates*/);
+		DecompositionToProximity.Add(Link.Key, ProxIdx);
+		DecompositionToProximity.Add(Link.Value, ProxIdx);
+	}
+}
+
+int32 FConvexDecomposition3::MergeBest(int32 TargetNumParts, double MaxErrorTolerance, double MinThicknessToleranceWorldSpace, bool bAllowCompact, bool bRequireHullTriangles)
 {
 	// Support having a max error tolerance
 	double VolumeTolerance = ConvertDistanceToleranceToLocalVolumeTolerance(MaxErrorTolerance);
@@ -1248,6 +1289,11 @@ int32 FConvexDecomposition3::MergeBest(int32 TargetNumParts, double MaxErrorTole
 		for (int32 PartIdx = 0; PartIdx < Decomposition.Num(); PartIdx++)
 		{
 			FConvexPart& Part = Decomposition[PartIdx];
+			// The 'must merge' logic is not designed to handle pre-compacted hull inputs
+			if (Part.IsCompact())
+			{
+				continue;
+			}
 			Part.bMustMerge = IsPartBelowSizeTolerance(Part);
 			MustMergeCount += Part.bMustMerge;
 		}
@@ -1530,6 +1576,18 @@ int32 FConvexDecomposition3::MergeBest(int32 TargetNumParts, double MaxErrorTole
 			DecompositionToProximity.Add(ToLink, ProxIdx);
 		}
 	}
+	
+	// Compute the triangulation for any hulls that are vertex-only (only possible if InitializeFromHulls was called before MergeBest)
+	if (bRequireHullTriangles)
+	{
+		for (FConvexPart& ConvexPart : Decomposition)
+		{
+			if (ConvexPart.HullTriangles.IsEmpty())
+			{
+				ConvexPart.ComputeHull();
+			}
+		}
+	}
 
 	return MergeNum;
 }
diff --git a/Engine/Source/Runtime/GeometryCore/Public/CompGeom/ConvexDecomposition3.h b/Engine/Source/Runtime/GeometryCore/Public/CompGeom/ConvexDecomposition3.h
--- a/Engine/Source/Runtime/GeometryCore/Public/CompGeom/ConvexDecomposition3.h
+++ b/Engine/Source/Runtime/GeometryCore/Public/CompGeom/ConvexDecomposition3.h
@@ -43,6 +43,16 @@ public:
 		InitializeFromMesh(SourceMesh, bMergeEdges);
 	}
 
+	/**
+	 * Initialize from convex hulls allows the caller to only use the hull merging phase of the algorithm
+	 * @param NumHulls			Number of convex hulls in the initial decomposition
+	 * @param HullVolumes		Function from Hull Index -> Hull Volume
+	 * @param HullNumVertices	Function from Hull Index -> Hull Vertex Count
+	 * @param HullVertices		Function from Hull Index, Vertex Index -> Hull Vertex Position
+	 * @param Proximity			All the local proximities for the hulls. Hulls will not be merged unless they are connected by this proximity graph.
+	 */
+	void InitializeFromHulls(int32 NumHulls, TFunctionRef<double(int32)> HullVolumes, TFunctionRef<int32(int32)> HullNumVertices, TFunctionRef<FVector3d(int32, int32)> HullVertices, TArrayView<const TPair<int32, int32>> Proximity);
+
 	void InitializeFromMesh(const FDynamicMesh3& SourceMesh, bool bMergeEdges);
 
 	//
@@ -114,8 +124,9 @@ public:
 	// @param MinPartThickness		Optionally specify a minimum thickness (in cm) for convex parts; parts below this thickness will always be merged away. Overrides TargetNumParts and ErrorTolerance when needed.
 	//								Note: These parts may be further split so they can be merged into multiple hulls
 	// @param bAllowCompact			Allow the algorithm to discard underlying geometry once it will no longer be used, resulting in a smaller representation & faster merges
+	// @param bRequireHullTriangles	Require all hulls to have associated triangles after MergeBest is completed. (If using InitializeFromHulls, will need to triangulate any un-merged hulls.)
 	// @return						The number of merges performed
-	int32 MergeBest(int32 TargetNumParts, double ErrorTolerance = 0, double MinThicknessTolerance = 0, bool bAllowCompact = true);
+	int32 MergeBest(int32 TargetNumParts, double ErrorTolerance = 0, double MinThicknessTolerance = 0, bool bAllowCompact = true, bool bRequireHullTriangles = false);
 
 	// simple helper to convert an error tolerance expressed in world space to a local-space volume tolerance
 	inline double ConvertDistanceToleranceToLocalVolumeTolerance(double DistTolerance) const
@@ -210,6 +221,9 @@ public:
 	{
 		FConvexPart() {}
 		FConvexPart(const FDynamicMesh3& SourceMesh, bool bMergeEdges, FTransformSRT3d& TransformOut);
+		// Allow direct construction of a compact part (e.g. to allow construction of a pre-existing convex hull)
+		FConvexPart(bool bIsCompact) : bIsCompact(bIsCompact)
+		{}
 
 		void Reset()
 		{
