remove deprecated timestamp concept and fix const correctness for geometrycore curve and point set classes

diff --git a/Engine/Source/Runtime/GeometryCore/Public/Curve/GeneralPolygon2.h b/Engine/Source/Runtime/GeometryCore/Public/Curve/GeneralPolygon2.h
--- a/Engine/Source/Runtime/GeometryCore/Public/Curve/GeneralPolygon2.h
+++ b/Engine/Source/Runtime/GeometryCore/Public/Curve/GeneralPolygon2.h
@@ -141,7 +141,7 @@ public:
     }
 
 
-    double Perimeter()
+    double Perimeter() const
     {
 		double PerimSum = Outer.Perimeter();
 		for (const TPolygon2<T> &Hole : Holes)
@@ -152,7 +152,7 @@ public:
     }
 
 
-    TAxisAlignedBox2<T> Bounds()
+    TAxisAlignedBox2<T> Bounds() const
     {
 		TAxisAlignedBox2<T> Box = Outer.Bounds();
 		for (const TPolygon2<T> Hole : Holes)
diff --git a/Engine/Source/Runtime/GeometryCore/Public/DynamicMesh/DynamicPointSet3.h b/Engine/Source/Runtime/GeometryCore/Public/DynamicMesh/DynamicPointSet3.h
--- a/Engine/Source/Runtime/GeometryCore/Public/DynamicMesh/DynamicPointSet3.h
+++ b/Engine/Source/Runtime/GeometryCore/Public/DynamicMesh/DynamicPointSet3.h
@@ -48,10 +48,6 @@ protected:
 	/** List of vertex positions */
 	TDynamicVector<RealType> Vertices;
 
-	/** The timestamp is incremented any time a function that modifies the point set is called */
-	UE_DEPRECATED(5.1, "Timestamps for TDynamicPointSet3 were not being used and will be removed in the future")
-	int Timestamp = 0;
-
 
 	/** Base attribute set */
 	TDynamicAttributeSetBase<TDynamicPointSet3<RealType>>* BaseAttributeSet = nullptr;
@@ -64,17 +60,6 @@ protected:
 	//
 public:
 
-	// Note: We need explicit defaults because of the deprecated Timestamp member
-	// Once Timestamp is deleted, we can delete all of these except the constructor
-	PRAGMA_DISABLE_DEPRECATION_WARNINGS
-	TDynamicPointSet3() = default;
-	~TDynamicPointSet3() = default;
-	TDynamicPointSet3(const TDynamicPointSet3& Other) = default;
-	TDynamicPointSet3(TDynamicPointSet3&& Other) noexcept = default;
-	TDynamicPointSet3& operator=(const TDynamicPointSet3& Other) = default;
-	TDynamicPointSet3& operator=(TDynamicPointSet3&& Other) noexcept = default;
-	PRAGMA_ENABLE_DEPRECATION_WARNINGS
-
 	void SetExternallyManagedAttributes(TDynamicAttributeSetBase<TDynamicPointSet3<RealType>>* AttributeSet)
 	{
 		BaseAttributeSet = AttributeSet;
@@ -124,7 +109,6 @@ public:
 		Vertices.Clear();
 		VertexRefCounts = FRefCountVector();
 		BaseAttributeSet = nullptr;
-		IncrementDeprecatedTimestamp();
 	}
 
 
@@ -142,7 +126,6 @@ public:
 			GetBaseAttributeSet()->OnNewVertex(vid, false);
 		}
 
-		IncrementDeprecatedTimestamp();
 		return vid;
 	}
 
@@ -186,7 +169,6 @@ public:
 			GetBaseAttributeSet()->OnNewVertex(VertexID, true);
 		}
 
-		IncrementDeprecatedTimestamp();
 		return EMeshResult::Ok;
 	}
 
@@ -214,7 +196,6 @@ public:
 		Vertices[i] = vNewPos.X;
 		Vertices[i + 1] = vNewPos.Y;
 		Vertices[i + 2] = vNewPos.Z;
-		IncrementDeprecatedTimestamp();
 	}
 
 
@@ -253,15 +234,6 @@ public:
 		return ((RealType)VertexCount() / (RealType)MaxVertexID());
 	}
 
-	/** Timestamp is incremented any time any change is made to the point set */
-	UE_DEPRECATED(5.1, "Timestamps for TDynamicPointSet3 were not being used and will be removed in the future")
-	inline int GetTimestamp() const
-	{
-		PRAGMA_DISABLE_DEPRECATION_WARNINGS
-		return Timestamp;
-		PRAGMA_ENABLE_DEPRECATION_WARNINGS
-	}
-
 	//
 	// Geometric queries
 	//
@@ -321,30 +293,9 @@ public:
 			GetBaseAttributeSet()->OnRemoveVertex(VertexID);
 		}
 
-		IncrementDeprecatedTimestamp();
 		return EMeshResult::Ok;
 	}
 
-
-
-protected:
-
-	UE_DEPRECATED(5.1, "Timestamps for TDynamicPointSet3 were not being used and will be removed in the future")
-	void UpdateTimeStamp()
-	{
-		IncrementDeprecatedTimestamp();
-	}
-
-private:
-
-	// Note: this function will be removed when Timestamp is removed
-	inline void IncrementDeprecatedTimestamp()
-	{
-		PRAGMA_DISABLE_DEPRECATION_WARNINGS
-		Timestamp++;
-		PRAGMA_ENABLE_DEPRECATION_WARNINGS
-	}
-
 };
 
 
diff --git a/Engine/Source/Runtime/GeometryCore/Public/Polygon2.h b/Engine/Source/Runtime/GeometryCore/Public/Polygon2.h
--- a/Engine/Source/Runtime/GeometryCore/Public/Polygon2.h
+++ b/Engine/Source/Runtime/GeometryCore/Public/Polygon2.h
@@ -33,26 +33,12 @@ protected:
 	/** The list of vertices/corners of the polygon */
 	TArray<TVector2<T>> Vertices;
 
-	/** A counter that is incremented every time the polygon vertices are modified. */
-	UE_DEPRECATED(5.1, "Timestamps for TPolygon2 were not being used and will be removed in the future")
-	int Timestamp = 0;
-
 public:
 
 	TPolygon2()
 	{
 	}
 
-	// Note: We need all 5 of these explicitly defaulted just because of the deprecated Timestamp member
-	// Once Timestamp is deleted, we can delete these as well
-	PRAGMA_DISABLE_DEPRECATION_WARNINGS
-	~TPolygon2() = default;
-	TPolygon2(const TPolygon2& Other) = default;
-	TPolygon2(TPolygon2&& Other) noexcept = default;
-	TPolygon2& operator=(const TPolygon2& Other) = default;
-	TPolygon2& operator=(TPolygon2&& Other) noexcept = default;
-	PRAGMA_ENABLE_DEPRECATION_WARNINGS
-
 	/**
 	 * Construct polygon with given list of vertices
 	 */
@@ -83,15 +69,6 @@ public:
 		}
 	}
 
-	/** @return the Timestamp for the polygon, which is updated every time the polygon is modified */
-	UE_DEPRECATED(5.1, "Timestamps for TPolygon2 were not being used and will be removed in the future")
-	int GetTimestamp() const 
-	{
-		PRAGMA_DISABLE_DEPRECATION_WARNINGS
-		return Timestamp;
-		PRAGMA_ENABLE_DEPRECATION_WARNINGS
-	}
-
 	/**
 	 * Get the vertex at a given index
 	 */
@@ -102,7 +79,6 @@ public:
 
 	/**
 	 * Get the vertex at a given index
-	 * @warning changing the vertex via this operator does not update Timestamp!
 	 */
 	TVector2<T>& operator[](int Index)
 	{
@@ -140,7 +116,6 @@ public:
 	void AppendVertex(const TVector2<T>& Position)
 	{
 		Vertices.Add(Position);
-		IncrementDeprecatedTimestamp();
 	}
 
 	/**
@@ -149,7 +124,6 @@ public:
 	void AppendVertices(const TArray<TVector2<T>>& NewVertices)
 	{
 		Vertices.Append(NewVertices);
-		IncrementDeprecatedTimestamp();
 	}
 
 	/**
@@ -158,7 +132,6 @@ public:
 	void Set(int VertexIndex, const TVector2<T>& Position)
 	{
 		Vertices[VertexIndex] = Position;
-		IncrementDeprecatedTimestamp();
 	}
 
 	/**
@@ -167,7 +140,6 @@ public:
 	void RemoveVertex(int VertexIndex)
 	{
 		Vertices.RemoveAt(VertexIndex);
-		IncrementDeprecatedTimestamp();
 	}
 
 	/**
@@ -176,7 +148,6 @@ public:
 	void SetVertices(const TArray<TVector2<T>>& NewVertices)
 	{
 		Vertices = NewVertices;
-		IncrementDeprecatedTimestamp();
 	}
 
 
@@ -186,7 +157,6 @@ public:
 	void Reverse()
 	{
 		Algo::Reverse(Vertices);
-		IncrementDeprecatedTimestamp();
 	}
 
 
@@ -369,7 +339,7 @@ public:
 	 * @param bDegenerateIsConvex	What to return for degenerate input (less than 3 points or equivalent due to repeated points)
 	 * @return						true if polygon is convex
 	 */
-	bool IsConvex(T RadiansTolerance = TMathUtil<T>::ZeroTolerance, bool bDegenerateIsConvex = true)
+	bool IsConvex(T RadiansTolerance = TMathUtil<T>::ZeroTolerance, bool bDegenerateIsConvex = true) const
 	{
 		return CurveUtil::IsConvex2<T, TVector2<T>>(Vertices, RadiansTolerance, bDegenerateIsConvex);
 	}
@@ -694,7 +664,6 @@ public:
 		{
 			Vertices[i] += Translate;
 		}
-		IncrementDeprecatedTimestamp();
 		return *this;
 	}
 
@@ -709,7 +678,6 @@ public:
 		{
 			Vertices[i] = Scale * (Vertices[i] - Origin) + Origin;
 		}
-		IncrementDeprecatedTimestamp();
 		return *this;
 	}
 
@@ -725,7 +693,6 @@ public:
 		{
 			Vertices[i] = TransformFunc(Vertices[i]);
 		}
-		IncrementDeprecatedTimestamp();
 		return *this;
 	}
 
@@ -759,8 +726,6 @@ public:
 		{
 			Vertices[k] = NewVertices[k];
 		}
-
-		IncrementDeprecatedTimestamp();
 	}
 
 
@@ -794,8 +759,6 @@ public:
 		{
 			Vertices[k] = NewVertices[k];
 		}
-
-		IncrementDeprecatedTimestamp();
 	}
 
 
@@ -955,7 +918,6 @@ public:
 			}
 		}
 
-		IncrementDeprecatedTimestamp();
 		return;
 	}
 
@@ -1020,8 +982,6 @@ public:
 			Vertices.Add(next_cut);
 			iCur = iNext;
 		} while (iCur != 0);
-
-		IncrementDeprecatedTimestamp();
 	}
 
 
@@ -1092,15 +1052,6 @@ public:
 		return Circle;
 	}
 
-private:
-
-	// Note: this function will be removed when Timestamp is removed
-	inline void IncrementDeprecatedTimestamp()
-	{
-		PRAGMA_DISABLE_DEPRECATION_WARNINGS
-		Timestamp++;
-		PRAGMA_ENABLE_DEPRECATION_WARNINGS
-	}
 };
 
 typedef TPolygon2<double> FPolygon2d;
diff --git a/Engine/Source/Runtime/GeometryCore/Public/Polyline3.h b/Engine/Source/Runtime/GeometryCore/Public/Polyline3.h
--- a/Engine/Source/Runtime/GeometryCore/Public/Polyline3.h
+++ b/Engine/Source/Runtime/GeometryCore/Public/Polyline3.h
@@ -24,22 +24,8 @@ protected:
 	/** The list of vertices of the polyline */
 	TArray<TVector<T>> Vertices;
 
-	/** A counter that is incremented every time the polyline vertices are modified */
-	UE_DEPRECATED(5.1, "Timestamps for TPolyline3 were not being used and will be removed in the future")
-	int Timestamp = 0;
-
 public:
 
-	// Note: We need all 5 of these explicitly defaulted just because of the deprecated Timestamp member
-	// Once Timestamp is deleted, we can delete these as well
-	PRAGMA_DISABLE_DEPRECATION_WARNINGS
-	~TPolyline3() = default;
-	TPolyline3(const TPolyline3& Other) = default;
-	TPolyline3(TPolyline3&& Other) noexcept = default;
-	TPolyline3& operator=(const TPolyline3& Other) = default;
-	TPolyline3& operator=(TPolyline3&& Other) noexcept = default;
-	PRAGMA_ENABLE_DEPRECATION_WARNINGS
-
 
 	TPolyline3()
 	{
@@ -61,22 +47,6 @@ public:
 		Vertices.Add(Point1);
 	}
 
-	/** @return the Timestamp for the polyline, which is updated every time the polyline is modified */
-	UE_DEPRECATED(5.1, "Timestamps for TPolyline3 were not being used and will be removed in the future")
-	int GetTimestamp() const
-	{
-		PRAGMA_DISABLE_DEPRECATION_WARNINGS
-		return Timestamp;
-		PRAGMA_ENABLE_DEPRECATION_WARNINGS
-	}
-
-	/** Explicitly increment the Timestamp */
-	UE_DEPRECATED(5.1, "Timestamps for TPolyline3 were not being used and will be removed in the future")
-	void IncrementTimestamp()
-	{
-		IncrementDeprecatedTimestamp();
-	}
-
 	/**
 	 * Get the vertex at a given index
 	 */
@@ -87,7 +57,6 @@ public:
 
 	/**
 	 * Get the vertex at a given index
-	 * @warning changing the vertex via this operator does not update Timestamp!
 	 */
 	TVector<T>& operator[](int Index)
 	{
@@ -141,7 +110,6 @@ public:
 	void Clear()
 	{
 		Vertices.Reset();
-		IncrementDeprecatedTimestamp();
 	}
 
 	/**
@@ -150,7 +118,6 @@ public:
 	void AppendVertex(const TVector<T>& Position)
 	{
 		Vertices.Add(Position);
-		IncrementDeprecatedTimestamp();
 	}
 
 	/**
@@ -159,7 +126,6 @@ public:
 	void AppendVertices(const TArray<TVector<T>>& NewVertices)
 	{
 		Vertices.Append(NewVertices);
-		IncrementDeprecatedTimestamp();
 	}
 
 	/**
@@ -174,7 +140,6 @@ public:
 		{
 			Vertices.Append( (TVector<T>)NewVertices[k] );
 		}
-		IncrementDeprecatedTimestamp();
 	}
 
 	/**
@@ -183,7 +148,6 @@ public:
 	void Set(int VertexIndex, const TVector<T>& Position)
 	{
 		Vertices[VertexIndex] = Position;
-		IncrementDeprecatedTimestamp();
 	}
 
 	/**
@@ -192,7 +156,6 @@ public:
 	void RemoveVertex(int VertexIndex)
 	{
 		Vertices.RemoveAt(VertexIndex);
-		IncrementDeprecatedTimestamp();
 	}
 
 	/**
@@ -206,7 +169,6 @@ public:
 		{
 			Vertices[k] = NewVertices[k];
 		}
-		IncrementDeprecatedTimestamp();
 	}
 
 
@@ -220,7 +182,6 @@ public:
 		{
 			Swap(Vertices[VertexIndex], Vertices[j]);
 		}
-		IncrementDeprecatedTimestamp();
 	}
 
 
@@ -470,16 +431,6 @@ public:
 		}
 		NewPolyline.Vertices[k] = Vertices[N];
 	}
-
-private:
-
-	// Note: this function will be removed when Timestamp is removed
-	inline void IncrementDeprecatedTimestamp()
-	{
-		PRAGMA_DISABLE_DEPRECATION_WARNINGS
-		Timestamp++;
-		PRAGMA_ENABLE_DEPRECATION_WARNINGS
-	}
 };
 
 typedef TPolyline3<double> FPolyline3d;
