Added a new class to store per-bone data.
Added bone attributes (name, color, parent index, pose) to the dynamic mesh attribute set.

diff --git a/Engine/Source/Runtime/GeometryCore/Private/DynamicMesh/DynamicBoneAttribute.cpp b/Engine/Source/Runtime/GeometryCore/Private/DynamicMesh/DynamicBoneAttribute.cpp
new file mode 100644
--- /dev/null
+++ b/Engine/Source/Runtime/GeometryCore/Private/DynamicMesh/DynamicBoneAttribute.cpp
@@ -0,0 +1,102 @@
+// Copyright Epic Games, Inc. All Rights Reserved.
+
+#include "DynamicMesh/DynamicBoneAttribute.h"
+
+using namespace UE::Geometry;
+
+template<>
+bool TDynamicBoneAttributeBase<FDynamicMesh3, FTransform>::IsSameAs(const TDynamicBoneAttributeBase<FDynamicMesh3, FTransform>& Other) const
+{
+	if (Num() != Other.Num())
+	{
+		return false;
+	}
+
+	for (int Idx = 0; Idx < Num(); ++Idx)
+	{	
+		if (GetValue(Idx).Equals(Other.GetValue(Idx)) == false)
+		{
+			return false;
+		}
+	}
+
+	return true;
+}
+
+template<typename ParentType, typename AttribValueType>
+TDynamicBoneAttributeBase<ParentType, AttribValueType>::TDynamicBoneAttributeBase(ParentType* ParentIn, const int InNumBones)
+:
+Parent(ParentIn)
+{
+	Resize(InNumBones);
+}
+
+template<typename ParentType, typename AttribValueType>
+TDynamicBoneAttributeBase<ParentType, AttribValueType>::TDynamicBoneAttributeBase(ParentType* ParentIn, const int InNumBones, const AttribValueType& InitialValue)
+:
+Parent(ParentIn)
+{
+	Initialize(InNumBones, InitialValue);
+}
+
+template<typename ParentType, typename AttribValueType>
+void TDynamicBoneAttributeBase<ParentType, AttribValueType>::Initialize(const int32 InNumBones, const AttribValueType& InitialValue)
+{	
+	AttribValues.Init(InitialValue, InNumBones);
+}
+
+template<typename ParentType, typename AttribValueType>
+void TDynamicBoneAttributeBase<ParentType, AttribValueType>::Resize(const int32 InNumBones)
+{	
+	AttribValues.SetNumUninitialized(InNumBones);
+}
+
+template<typename ParentType, typename AttribValueType>
+void TDynamicBoneAttributeBase<ParentType, AttribValueType>::Copy(const TDynamicBoneAttributeBase<ParentType, AttribValueType>& Copy)
+{
+	TDynamicAttributeBase<ParentType>::CopyParentClassData(Copy);
+	AttribValues = Copy.AttribValues;
+}
+
+template<typename ParentType, typename AttribValueType>
+TDynamicAttributeBase<ParentType>* TDynamicBoneAttributeBase<ParentType, AttribValueType>::MakeCopy(ParentType* ParentIn) const
+{
+	TDynamicBoneAttributeBase<ParentType, AttribValueType>* ToFill = new TDynamicBoneAttributeBase<ParentType, AttribValueType>(ParentIn);
+	ToFill->Copy(*this);
+	return ToFill;
+}
+
+template<typename ParentType, typename AttribValueType>
+TDynamicAttributeBase<ParentType>* TDynamicBoneAttributeBase<ParentType, AttribValueType>::MakeNew(ParentType* ParentIn) const
+{
+	TDynamicBoneAttributeBase<ParentType, AttribValueType>* Matching = new TDynamicBoneAttributeBase<ParentType, AttribValueType>(ParentIn);
+	return Matching;
+}
+
+template<typename ParentType, typename AttribValueType>
+bool TDynamicBoneAttributeBase<ParentType, AttribValueType>::IsSameAs(const TDynamicBoneAttributeBase<ParentType, AttribValueType>& Other) const
+{
+	return AttribValues == Other.AttribValues;
+}
+
+template<typename ParentType, typename AttribValueType>
+void TDynamicBoneAttributeBase<ParentType, AttribValueType>::Append(const AttribValueType& InValue) 
+{
+	AttribValues.Add(InValue);
+}
+
+template<typename ParentType, typename AttribValueType>
+void TDynamicBoneAttributeBase<ParentType, AttribValueType>::Serialize(FArchive& Ar)
+{	
+	/** TODO: Not supported for now */
+	checkNoEntry();
+}
+
+// Instantiate for linker
+namespace UE::Geometry
+{
+	template class TDynamicBoneAttributeBase<FDynamicMesh3, FName>;
+	template class TDynamicBoneAttributeBase<FDynamicMesh3, int32>;
+	template class TDynamicBoneAttributeBase<FDynamicMesh3, FVector4f>;
+	template class TDynamicBoneAttributeBase<FDynamicMesh3, FTransform>;
+}
\ No newline at end of file
diff --git a/Engine/Source/Runtime/GeometryCore/Private/DynamicMesh/DynamicMeshAttributeSet.cpp b/Engine/Source/Runtime/GeometryCore/Private/DynamicMesh/DynamicMeshAttributeSet.cpp
--- a/Engine/Source/Runtime/GeometryCore/Private/DynamicMesh/DynamicMeshAttributeSet.cpp
+++ b/Engine/Source/Runtime/GeometryCore/Private/DynamicMesh/DynamicMeshAttributeSet.cpp
@@ -3,6 +3,7 @@
 #include "DynamicMesh/DynamicMeshAttributeSet.h"
 
 #include "DynamicMesh/DynamicVertexSkinWeightsAttribute.h"
+#include "DynamicMesh/DynamicBoneAttribute.h"
 #include "IndexTypes.h"
 #include "Async/Async.h"
 #include "Serialization/NameAsStringProxyArchive.h"
@@ -80,6 +81,8 @@ void FDynamicMeshAttributeSet::Copy(const FDynamicMeshAttributeSet& Copy)
 		AttachSkinWeightsAttribute(AttribPair.Key,
 			static_cast<FDynamicMeshVertexSkinWeightsAttribute *>(AttribPair.Value->MakeCopy(ParentMesh)));
 	}
+
+	CopyBoneAttributes(Copy);
 	
 	GenericAttributes.Reset();
 	for (const TPair<FName, TUniquePtr<FDynamicMeshAttributeBase>>& AttribPair : Copy.GenericAttributes)
@@ -173,6 +176,8 @@ void FDynamicMeshAttributeSet::CompactCopy(const FCompactMaps& CompactMaps, cons
 		AttachSkinWeightsAttribute(AttribPair.Key,
 			static_cast<FDynamicMeshVertexSkinWeightsAttribute *>(AttribPair.Value->MakeCompactCopy(CompactMaps, ParentMesh)));
 	}
+
+	CopyBoneAttributes(Copy);
 	
 	GenericAttributes.Reset();
 	for (const TPair<FName, TUniquePtr<FDynamicMeshAttributeBase>>& AttribPair : Copy.GenericAttributes)
@@ -406,9 +411,7 @@ void FDynamicMeshAttributeSet::EnableMatchingAttributes(const FDynamicMeshAttrib
 		}
 	}
 
-
-
-
+	EnableMatchingBoneAttributes(ToMatch, bClearExisting, bDiscardExtraAttributes);
 }
 
 
@@ -449,6 +452,26 @@ void FDynamicMeshAttributeSet::Reparent(FDynamicMesh3* NewParent)
 	{
 		AttribPair.Value.Get()->Reparent(NewParent);
 	}
+
+	if (BoneNameAttrib)
+	{
+		BoneNameAttrib->Reparent(NewParent);
+	}
+
+	if (BoneParentIndexAttrib)
+	{
+		BoneParentIndexAttrib->Reparent(NewParent);
+	}
+
+	if (BonePoseAttrib)
+	{
+		BonePoseAttrib->Reparent(NewParent);
+	}
+
+	if (BoneColorAttrib)
+	{
+		BoneColorAttrib->Reparent(NewParent);
+	}
 	
 	for (const TPair<FName, TUniquePtr<FDynamicMeshAttributeBase>>& AttribPair : GenericAttributes)
 	{
@@ -1215,6 +1238,11 @@ bool FDynamicMeshAttributeSet::IsSameAs(const FDynamicMeshAttributeSet& Other, b
 			++ItOther;
 		}
 	}
+
+	if (!IsSameBoneAttributesAs(Other))
+	{
+		return false;
+	}
 	
 	// TODO: Test GenericAttributes
 
@@ -1362,5 +1390,217 @@ void FDynamicMeshAttributeSet::Serialize(FArchive& Ar, const FCompactMaps* Compa
 		}
 	}
 
+
+	// TODO: Serialize bone attributes
+
 	//Ar << GenericAttributes; // TODO
 }
+
+bool FDynamicMeshAttributeSet::CheckValidity(bool bAllowNonmanifold, EValidityCheckFailMode FailMode) const
+{
+	bool bValid = FDynamicMeshAttributeSetBase::CheckValidity(bAllowNonmanifold, FailMode);
+	for (int UVLayerIndex = 0; UVLayerIndex < NumUVLayers(); UVLayerIndex++)
+	{
+		bValid = GetUVLayer(UVLayerIndex)->CheckValidity(bAllowNonmanifold, FailMode) && bValid;
+	}
+	bValid = PrimaryNormals()->CheckValidity(bAllowNonmanifold, FailMode) && bValid;
+	if (ColorLayer)
+	{
+		bValid = ColorLayer->CheckValidity(bAllowNonmanifold, FailMode) && bValid;
+	}
+	if (MaterialIDAttrib)
+	{
+		bValid = MaterialIDAttrib->CheckValidity(bAllowNonmanifold, FailMode) && bValid;
+	}
+	for (int PolygroupLayerIndex = 0; PolygroupLayerIndex < NumPolygroupLayers(); PolygroupLayerIndex++)
+	{
+		bValid = GetPolygroupLayer(PolygroupLayerIndex)->CheckValidity(bAllowNonmanifold, FailMode) && bValid;
+	}
+	for (int WeightLayerIndex = 0; WeightLayerIndex < NumWeightLayers(); WeightLayerIndex++)
+	{
+		bValid = GetWeightLayer(WeightLayerIndex)->CheckValidity(bAllowNonmanifold, FailMode) && bValid;
+	}
+	for (const TPair<FName, TUniquePtr<FDynamicMeshVertexSkinWeightsAttribute>>& Pair : SkinWeightAttributes)
+	{
+		if (Pair.Value)
+		{
+			bValid = Pair.Value->CheckValidity(bAllowNonmanifold, FailMode) && bValid;
+		}
+	}
+
+	bValid = CheckBoneValidity(FailMode) && bValid;
+
+	return bValid;
+}
+
+
+//
+// Bone Attributes Methods
+//
+
+namespace BoneAttributeHelpers
+{	
+	template <typename ParentType, typename AttribValueType>
+ 	void EnableIfMatching(FDynamicMesh3* Mesh,
+						  TUniquePtr<TDynamicBoneAttributeBase<ParentType, AttribValueType>>& Attribute,
+						  const TDynamicBoneAttributeBase<ParentType, AttribValueType>* ToMatch,
+						  const AttribValueType& InitialValue,
+						  bool bClearExisting, 
+						  bool bDiscardExtraAttributes)
+	{
+		const bool bToMatchIsNotNull = ToMatch != nullptr;
+		const bool bWantBoneNames = (bClearExisting || bDiscardExtraAttributes) ? bToMatchIsNotNull : (bToMatchIsNotNull || Attribute.Get());
+		if (bClearExisting || bWantBoneNames == false)
+		{
+			Attribute = nullptr;
+		}
+		if (bWantBoneNames)
+		{	
+			// If ToMatch has no bones, then enable the bone attribute with zero bones so that the caller can resize it manually
+			const int32 NumBones = bToMatchIsNotNull ? ToMatch->Num() : 0;
+			TDynamicBoneAttributeBase<ParentType, AttribValueType>* Ptr = static_cast<TDynamicBoneAttributeBase<ParentType, AttribValueType>*>(ToMatch->MakeNew(Mesh));
+			Attribute = TUniquePtr<TDynamicBoneAttributeBase<ParentType, AttribValueType>>(Ptr);
+			Attribute->Initialize(NumBones, InitialValue);
+		}
+	}
+
+	template <typename ParentType, typename AttribValueType>
+ 	void CopyBoneAttribute(FDynamicMesh3* Mesh,
+						   TUniquePtr<TDynamicBoneAttributeBase<ParentType, AttribValueType>>& Attribute,
+						   const TDynamicBoneAttributeBase<ParentType, AttribValueType>* Copy)
+	{
+		if (Copy)
+		{
+			TDynamicBoneAttributeBase<ParentType, AttribValueType>* Ptr = static_cast<TDynamicBoneAttributeBase<ParentType, AttribValueType>*>(Copy->MakeNew(Mesh));
+			Attribute = TUniquePtr<TDynamicBoneAttributeBase<ParentType, AttribValueType>>(Ptr);
+			Attribute->Copy(*(Copy));
+		}
+		else
+		{
+			Attribute = nullptr;
+		}
+	}
+}
+
+int32 FDynamicMeshAttributeSet::GetNumBones() const
+{
+	return HasBones() ? GetBoneNames()->Num() : 0;
+}
+
+void FDynamicMeshAttributeSet::CopyBoneAttributes(const FDynamicMeshAttributeSet& Copy)
+{
+	BoneAttributeHelpers::CopyBoneAttribute(ParentMesh, BoneNameAttrib, Copy.GetBoneNames());
+	BoneAttributeHelpers::CopyBoneAttribute(ParentMesh, BoneParentIndexAttrib, Copy.GetBoneParentIndices());
+	BoneAttributeHelpers::CopyBoneAttribute(ParentMesh, BonePoseAttrib, Copy.GetBonePoses());
+	BoneAttributeHelpers::CopyBoneAttribute(ParentMesh, BoneColorAttrib, Copy.GetBoneColors());
+}
+
+void FDynamicMeshAttributeSet::EnableMatchingBoneAttributes(const FDynamicMeshAttributeSet& ToMatch, bool bClearExisting, bool bDiscardExtraAttributes)
+{
+	BoneAttributeHelpers::EnableIfMatching(ParentMesh, BoneNameAttrib, ToMatch.GetBoneNames(), (FName)NAME_None, bClearExisting, bDiscardExtraAttributes);
+	BoneAttributeHelpers::EnableIfMatching(ParentMesh, BoneParentIndexAttrib, ToMatch.GetBoneParentIndices(), (int32)INDEX_NONE, bClearExisting, bDiscardExtraAttributes);
+	BoneAttributeHelpers::EnableIfMatching(ParentMesh, BonePoseAttrib, ToMatch.GetBonePoses(), FTransform::Identity, bClearExisting, bDiscardExtraAttributes);
+	BoneAttributeHelpers::EnableIfMatching(ParentMesh, BoneColorAttrib, ToMatch.GetBoneColors(), FVector4f::One(), bClearExisting, bDiscardExtraAttributes);
+}
+
+void FDynamicMeshAttributeSet::EnableBones(const int InBonesNum)
+{
+	if (HasBones() == false || GetNumBones() != InBonesNum)
+	{
+		BoneNameAttrib = MakeUnique<FDynamicMeshBoneNameAttribute>(ParentMesh, InBonesNum, NAME_None);
+		BoneParentIndexAttrib = MakeUnique<FDynamicMeshBoneParentIndexAttribute>(ParentMesh, InBonesNum, INDEX_NONE);
+		BonePoseAttrib = MakeUnique<FDynamicMeshBonePoseAttribute>(ParentMesh, InBonesNum, FTransform::Identity);
+		BoneColorAttrib = MakeUnique<FDynamicMeshBoneColorAttribute>(ParentMesh, InBonesNum, FVector4f::One());
+	}
+}
+
+void FDynamicMeshAttributeSet::DisableBones()
+{
+	BoneNameAttrib.Reset();
+	BoneParentIndexAttrib.Reset();
+	BonePoseAttrib.Reset();
+	BoneColorAttrib.Reset();
+}
+
+bool FDynamicMeshAttributeSet::IsSameBoneAttributesAs(const FDynamicMeshAttributeSet& Other) const
+{
+	if (HasBones() != Other.HasBones())
+	{
+		return false;
+	}
+	
+	if (HasBones())
+	{
+		if (!BoneNameAttrib->IsSameAs(*Other.BoneNameAttrib))
+		{
+			return false;
+		}
+
+		if (!BoneParentIndexAttrib->IsSameAs(*Other.BoneParentIndexAttrib))
+		{
+			return false;
+		}
+	}
+
+	return true;
+}
+
+bool FDynamicMeshAttributeSet::AppendBonesUnique(const FDynamicMeshAttributeSet& Other)
+{
+	if (!Other.CheckBoneValidity(EValidityCheckFailMode::ReturnOnly))
+	{
+		checkSlow(false);
+		return false; // don't append from an invalid bone data
+	}
+
+	if (!Other.HasBones())
+	{
+		return true;
+	}
+
+	const FDynamicMeshBoneNameAttribute* OtherBoneNameAttrib = Other.GetBoneNames();
+
+	if (!HasBones())
+	{
+		EnableBones(0);
+	}
+
+	// Used to check if bone with a given name in Other already exists
+	TSet<FName> HashSet;
+	HashSet.Append(BoneNameAttrib->GetAttribValues());
+
+	for (int BoneIdx = 0; BoneIdx < OtherBoneNameAttrib->Num(); ++BoneIdx)
+	{
+		if (HashSet.Contains(OtherBoneNameAttrib->GetValue(BoneIdx)) == false)
+		{
+			BoneNameAttrib->Append(OtherBoneNameAttrib->GetValue(BoneIdx));
+			BoneParentIndexAttrib->Append(Other.GetBoneParentIndices()->GetValue(BoneIdx));
+			BonePoseAttrib->Append(Other.GetBonePoses()->GetValue(BoneIdx));
+			BoneColorAttrib->Append(Other.GetBoneColors()->GetValue(BoneIdx));
+		}
+	}
+
+	return true;
+}
+
+bool FDynamicMeshAttributeSet::CheckBoneValidity(EValidityCheckFailMode FailMode) const
+{
+	const int32 NumBones = HasBones() ? GetNumBones() : 0;
+
+	bool bValid = true;
+	
+	bValid = (BoneParentIndexAttrib->Num() == NumBones || BoneParentIndexAttrib->IsEmpty()) && bValid;
+	bValid = (BoneColorAttrib->Num() == NumBones || BoneColorAttrib->IsEmpty()) && bValid;
+	bValid = (BonePoseAttrib->Num() == NumBones || BonePoseAttrib->IsEmpty()) && bValid;
+
+	if (FailMode == EValidityCheckFailMode::Check)
+	{
+		checkf(bValid, TEXT("TDynamicSkeletalBonesAttribute::CheckValidity failed!"));
+	}
+	else if (FailMode == EValidityCheckFailMode::Ensure)
+	{ 
+		ensureMsgf(bValid, TEXT("TDynamicSkeletalBonesAttribute::CheckValidity failed!"));
+	}
+
+	return bValid;
+}
diff --git a/Engine/Source/Runtime/GeometryCore/Public/DynamicMesh/DynamicBoneAttribute.h b/Engine/Source/Runtime/GeometryCore/Public/DynamicMesh/DynamicBoneAttribute.h
new file mode 100644
--- /dev/null
+++ b/Engine/Source/Runtime/GeometryCore/Public/DynamicMesh/DynamicBoneAttribute.h
@@ -0,0 +1,155 @@
+// Copyright Epic Games, Inc. All Rights Reserved.
+
+#pragma once
+
+#include "DynamicMesh/DynamicAttribute.h"
+#include "Containers/Array.h"
+
+//
+// Forward declarations
+//
+namespace UE::Geometry { class FDynamicMesh3; }
+
+namespace UE
+{
+namespace Geometry
+{
+
+/**
+ * TDynamicAttributeBase is a base class for storing per-bone data. 
+ */
+template<typename ParentType, typename AttribValueType>
+class TDynamicBoneAttributeBase
+: 
+public TDynamicAttributeBase<ParentType>
+{
+protected:
+	friend class FDynamicMeshAttributeSet;
+	
+	/** The parent object (e.g. mesh, point set) this attribute belongs to. */
+	ParentType* Parent = nullptr;
+
+	/** List of per-bone values. */
+	TArray<AttribValueType> AttribValues;
+public:
+	
+	/** Create an empty attribute. */
+	TDynamicBoneAttributeBase() = default;
+
+	/** Create an attribute for the given parent and set the number of unintialized values. */
+	TDynamicBoneAttributeBase(ParentType* ParentIn, const int InNumBones = 0);
+
+	/** Create an attribute for the given parent and number of bones and set all values to the InitialValue. */
+	TDynamicBoneAttributeBase(ParentType* ParentIn, const int InNumBones, const AttribValueType& InitialValue);
+
+	virtual ~TDynamicBoneAttributeBase() = default;
+
+	/** Resize the array of values to the given max bone number and set all the values to the InitialValue. */
+	void Initialize(const int32 InNumBones, const AttribValueType& InitialValue);
+
+	/** Resize the array of values to the given max bone number. */
+	void Resize(const int32 InNumBones);
+
+	 /** Set this attribute to contain the same arrays as the copy attribute. */
+	void Copy(const TDynamicBoneAttributeBase<ParentType, AttribValueType>& Copy);
+
+ 	/** @return true if the attribute has no bone values stored */
+	inline bool IsEmpty() const
+	{
+		return AttribValues.IsEmpty();
+	}
+	
+	/** @return the number of the values stored. */
+	inline int32 Num() const 
+	{
+		return AttribValues.Num();
+	}
+
+	/** 
+	 * @return true, if this attribute is the same as the Other.
+	 * @note assumes that the AttribValueType support equality operator. If not, create a template specialization.
+	 */
+	virtual bool IsSameAs(const TDynamicBoneAttributeBase<ParentType, AttribValueType>& Other) const;
+
+	inline const TArray<AttribValueType>& GetAttribValues() const
+	{ 
+		return AttribValues; 
+	}
+
+	inline void SetValue(const int32 InBoneID, const AttribValueType& InValue)
+	{
+		checkSlow(InBoneID >= 0 && InBoneID < AttribValues.Num());
+		AttribValues[InBoneID] = InValue;
+	}
+
+	inline const AttribValueType& GetValue(const int32 InBoneID) const
+	{
+		checkSlow(InBoneID >= 0 && InBoneID < AttribValues.Num())
+		return AttribValues[InBoneID];
+	}
+
+	/** Add value at the end of the value array. */
+	void Append(const AttribValueType& InValue);
+
+public:
+
+    //
+    // TDynamicAttributeBase interface
+    //
+
+    virtual TDynamicAttributeBase<ParentType>* MakeCopy(ParentType* ParentIn) const override;
+
+	virtual TDynamicAttributeBase<ParentType>* MakeNew(ParentType* ParentIn) const override;
+
+    virtual void CompactInPlace(const FCompactMaps& CompactMaps) override
+	{
+		// Compact maps don't affect the bone attribute.
+	}
+
+    virtual void Reparent(ParentType* NewParent) override 
+    { 
+        Parent = NewParent;  
+    }
+
+	virtual bool CopyThroughMapping(const TDynamicAttributeBase<ParentType>* Source, const FMeshIndexMappings& Mapping) override
+	{
+		// Index remapping is not supported by the bone attribute. We need to be careful with the order of the 
+		// bones since it affects the indexing of the skin weight attributes.
+		return true;
+	}
+
+	virtual bool CopyOut(int RawID, void* Buffer, int BufferSize) const override
+	{
+		return false; // Not supported for bones
+	}
+
+	virtual bool CopyIn(int RawID, void* Buffer, int BufferSize) override
+	{
+		return false; // Not supported for bones
+	}
+
+	virtual TUniquePtr<TDynamicAttributeChangeBase<ParentType>> NewBlankChange() const override
+	{
+		return MakeUnique<TDynamicAttributeChangeBase<ParentType>>();
+	}
+
+public:
+    
+    //
+    // Serialization
+    //
+
+    void Serialize(FArchive& Ar);
+};
+
+/** FTransform doesn't support equality operator. */
+template<>
+bool TDynamicBoneAttributeBase<FDynamicMesh3, FTransform>::IsSameAs(const TDynamicBoneAttributeBase<FDynamicMesh3, FTransform>& Other) const;
+
+using FDynamicMeshBoneNameAttribute = TDynamicBoneAttributeBase<FDynamicMesh3, FName>;
+using FDynamicMeshBoneParentIndexAttribute = TDynamicBoneAttributeBase<FDynamicMesh3, int32>;
+using FDynamicMeshBoneColorAttribute = TDynamicBoneAttributeBase<FDynamicMesh3, FVector4f>;
+using FDynamicMeshBonePoseAttribute = TDynamicBoneAttributeBase<FDynamicMesh3, FTransform>;
+
+} // namespace Geometry
+} // namespace UE
\ No newline at end of file
diff --git a/Engine/Source/Runtime/GeometryCore/Public/DynamicMesh/DynamicMeshAttributeSet.h b/Engine/Source/Runtime/GeometryCore/Public/DynamicMesh/DynamicMeshAttributeSet.h
--- a/Engine/Source/Runtime/GeometryCore/Public/DynamicMesh/DynamicMeshAttributeSet.h
+++ b/Engine/Source/Runtime/GeometryCore/Public/DynamicMesh/DynamicMeshAttributeSet.h
@@ -57,6 +57,16 @@ class TDynamicVertexSkinWeightsAttribute;
 
 using FDynamicMeshVertexSkinWeightsAttribute = TDynamicVertexSkinWeightsAttribute<FDynamicMesh3>;
 	
+/** Bone Attributes */
+template<typename ParentType, typename AttribValueType>
+class TDynamicBoneAttributeBase;
+
+using FDynamicMeshBoneNameAttribute = TDynamicBoneAttributeBase<FDynamicMesh3, FName>;
+using FDynamicMeshBoneParentIndexAttribute = TDynamicBoneAttributeBase<FDynamicMesh3, int32>;
+using FDynamicMeshBoneColorAttribute = TDynamicBoneAttributeBase<FDynamicMesh3, FVector4f>;
+using FDynamicMeshBonePoseAttribute = TDynamicBoneAttributeBase<FDynamicMesh3, FTransform>;
+
+
 /**
  * FDynamicMeshAttributeSet manages a set of extended attributes for a FDynamicMesh3.
  * This includes UV and Normal overlays, etc.
@@ -380,7 +390,90 @@ public:
 	{
 		return SkinWeightAttributes;
 	}
+
+
+
+	//
+	// Bone Attributes
+	//
+
+	void CopyBoneAttributes(const FDynamicMeshAttributeSet& Copy);
+	
+	void EnableMatchingBoneAttributes(const FDynamicMeshAttributeSet& ToMatch, bool bClearExisting, bool bDiscardExtraAttributes);
+
+	/** @note Only compares bone names and parent indices. */
+	bool IsSameBoneAttributesAs(const FDynamicMeshAttributeSet& Other) const;
+
+	/**
+	 * The attribute is valid if either all attributes are empty or if the bone name attribute is not empty then all the other
+	 * attributes must be either empty or their size is equal to the bone name attribute size.
+	 */
+	bool CheckBoneValidity(EValidityCheckFailMode FailMode) const;
+
+	
+	/** 
+	 * Append bone attributes from another set. When appending, the bone attributes in the Other set will only be 
+	 * appended if a bone with the same name does not exist in the current bone name attribute. Hence, the 
+	 * order of the bones in this set is preserved.
+	 * 
+	 * @return true, if append was successful
+	 */
+	bool AppendBonesUnique(const FDynamicMeshAttributeSet& Other);
+
+	/** Enable all bone attributes and intialize their size to the given bone number. */
+	void EnableBones(const int InBonesNum);
+
+	/** Disable all bone attributes. */
+	void DisableBones();
+	
+	/** Get number of bones. */
+	int32 GetNumBones() const;
+
+	bool HasBones() const 
+	{ 
+		return !!BoneNameAttrib; 
+	}
+
+	FDynamicMeshBoneNameAttribute* GetBoneNames()
+	{
+		return BoneNameAttrib.Get();
+	}
+
+	const FDynamicMeshBoneNameAttribute* GetBoneNames() const
+	{
+		return BoneNameAttrib.Get();
+	}
 	
+	FDynamicMeshBoneParentIndexAttribute* GetBoneParentIndices()
+	{
+		return BoneParentIndexAttrib.Get();
+	}
+
+	const FDynamicMeshBoneParentIndexAttribute* GetBoneParentIndices() const
+	{
+		return BoneParentIndexAttrib.Get();
+	}
+
+	FDynamicMeshBonePoseAttribute* GetBonePoses()
+	{
+		return BonePoseAttrib.Get();
+	}
+
+	const FDynamicMeshBonePoseAttribute* GetBonePoses() const
+	{
+		return BonePoseAttrib.Get();
+	}
+
+	FDynamicMeshBoneColorAttribute* GetBoneColors()
+	{
+		return BoneColorAttrib.Get();
+	}
+
+	const FDynamicMeshBoneColorAttribute* GetBoneColors() const
+	{
+		return BoneColorAttrib.Get();
+	}
+
 	// Attach a new attribute (and transfer ownership of it to the attribute set)
 	void AttachAttribute(FName AttribName, FDynamicMeshAttributeBase* Attribute)
 	{
@@ -469,6 +562,12 @@ protected:
 	using SkinWeightAttributesMap = TMap<FName, TUniquePtr<FDynamicMeshVertexSkinWeightsAttribute>>;
 	SkinWeightAttributesMap SkinWeightAttributes;
 
+	// Bone attributes
+	TUniquePtr<FDynamicMeshBoneNameAttribute> BoneNameAttrib;
+	TUniquePtr<FDynamicMeshBoneParentIndexAttribute> BoneParentIndexAttrib;
+	TUniquePtr<FDynamicMeshBonePoseAttribute> BonePoseAttrib;
+	TUniquePtr<FDynamicMeshBoneColorAttribute> BoneColorAttrib;
+
 	using GenericAttributesMap = TMap<FName, TUniquePtr<FDynamicMeshAttributeBase>>;
 	GenericAttributesMap GenericAttributes;
 	
@@ -511,32 +610,7 @@ protected:
 	 * @param bAllowNonmanifold Accept non-manifold topology as valid. Note that this should almost always be true for attributes; non-manifold overlays are generally valid.
 	 * @param FailMode Desired behavior if mesh is found invalid
 	 */
-	virtual bool CheckValidity(bool bAllowNonmanifold, EValidityCheckFailMode FailMode) const
-	{
-		bool bValid = FDynamicMeshAttributeSetBase::CheckValidity(bAllowNonmanifold, FailMode);
-		for (int UVLayerIndex = 0; UVLayerIndex < NumUVLayers(); UVLayerIndex++)
-		{
-			bValid = GetUVLayer(UVLayerIndex)->CheckValidity(bAllowNonmanifold, FailMode) && bValid;
-		}
-		bValid = PrimaryNormals()->CheckValidity(bAllowNonmanifold, FailMode) && bValid;
-		if (ColorLayer)
-		{
-			bValid = ColorLayer->CheckValidity(bAllowNonmanifold, FailMode) && bValid;
-		}
-		if (MaterialIDAttrib)
-		{
-			bValid = MaterialIDAttrib->CheckValidity(bAllowNonmanifold, FailMode) && bValid;
-		}
-		for (int PolygroupLayerIndex = 0; PolygroupLayerIndex < NumPolygroupLayers(); PolygroupLayerIndex++)
-		{
-			bValid = GetPolygroupLayer(PolygroupLayerIndex)->CheckValidity(bAllowNonmanifold, FailMode) && bValid;
-		}
-		for (int WeightLayerIndex = 0; WeightLayerIndex < NumWeightLayers(); WeightLayerIndex++)
-		{
-			bValid = GetWeightLayer(WeightLayerIndex)->CheckValidity(bAllowNonmanifold, FailMode) && bValid;
-		}
-		return bValid;
-	}
+	virtual bool CheckValidity(bool bAllowNonmanifold, EValidityCheckFailMode FailMode) const;
 };
 
 
