fix validity check for boneless meshes
(fixes geometry processing unit tests)

diff --git a/Engine/Source/Runtime/GeometryCore/Private/DynamicMesh/DynamicMeshAttributeSet.cpp b/Engine/Source/Runtime/GeometryCore/Private/DynamicMesh/DynamicMeshAttributeSet.cpp
--- a/Engine/Source/Runtime/GeometryCore/Private/DynamicMesh/DynamicMeshAttributeSet.cpp
+++ b/Engine/Source/Runtime/GeometryCore/Private/DynamicMesh/DynamicMeshAttributeSet.cpp
@@ -1585,13 +1585,21 @@ bool FDynamicMeshAttributeSet::AppendBonesUnique(const FDynamicMeshAttributeSet&
 
 bool FDynamicMeshAttributeSet::CheckBoneValidity(EValidityCheckFailMode FailMode) const
 {
-	const int32 NumBones = HasBones() ? GetNumBones() : 0;
-
 	bool bValid = true;
 	
-	bValid = (BoneParentIndexAttrib->Num() == NumBones || BoneParentIndexAttrib->IsEmpty()) && bValid;
-	bValid = (BoneColorAttrib->Num() == NumBones || BoneColorAttrib->IsEmpty()) && bValid;
-	bValid = (BonePoseAttrib->Num() == NumBones || BonePoseAttrib->IsEmpty()) && bValid;
+	if (!HasBones())
+	{
+		// if boneless, no bone-related attributes should be set
+		bValid = !BoneNameAttrib && !BoneParentIndexAttrib && !BoneColorAttrib && !BonePoseAttrib;
+	}
+	else
+	{
+		const int32 NumBones = GetNumBones();
+	
+		bValid = (BoneParentIndexAttrib->Num() == NumBones || BoneParentIndexAttrib->IsEmpty()) && bValid;
+		bValid = (BoneColorAttrib->Num() == NumBones || BoneColorAttrib->IsEmpty()) && bValid;
+		bValid = (BonePoseAttrib->Num() == NumBones || BonePoseAttrib->IsEmpty()) && bValid;
+	}
 
 	if (FailMode == EValidityCheckFailMode::Check)
 	{
