Fix static analysis error.

diff --git a/Engine/Source/Runtime/GeometryCore/Private/DynamicMesh/DynamicMeshAttributeSet.cpp b/Engine/Source/Runtime/GeometryCore/Private/DynamicMesh/DynamicMeshAttributeSet.cpp
--- a/Engine/Source/Runtime/GeometryCore/Private/DynamicMesh/DynamicMeshAttributeSet.cpp
+++ b/Engine/Source/Runtime/GeometryCore/Private/DynamicMesh/DynamicMeshAttributeSet.cpp
@@ -1449,16 +1449,15 @@ namespace BoneAttributeHelpers
 						  bool bDiscardExtraAttributes)
 	{
 		const bool bToMatchIsNotNull = ToMatch != nullptr;
-		const bool bWantBoneNames = (bClearExisting || bDiscardExtraAttributes) ? bToMatchIsNotNull : (bToMatchIsNotNull || Attribute.Get());
-		if (bClearExisting || bWantBoneNames == false)
+		const bool bWantAttrib = (bClearExisting || bDiscardExtraAttributes) ? bToMatchIsNotNull : (bToMatchIsNotNull || Attribute.Get());
+		if (bClearExisting || bWantAttrib == false)
 		{
-			Attribute = nullptr;
+			Attribute.Reset();
 		}
-		if (bWantBoneNames)
+		if (bWantAttrib)
 		{	
-			// If ToMatch has no bones, then enable the bone attribute with zero bones so that the caller can resize it manually
-			const int32 NumBones = bToMatchIsNotNull ? ToMatch->Num() : 0;
-			TDynamicBoneAttributeBase<ParentType, AttribValueType>* Ptr = static_cast<TDynamicBoneAttributeBase<ParentType, AttribValueType>*>(ToMatch->MakeNew(Mesh));
+			const int32 NumBones = ToMatch ? ToMatch->Num() : 0;
+			TDynamicBoneAttributeBase<ParentType, AttribValueType>* Ptr = new TDynamicBoneAttributeBase<ParentType, AttribValueType>(Mesh);
 			Attribute = TUniquePtr<TDynamicBoneAttributeBase<ParentType, AttribValueType>>(Ptr);
 			Attribute->Initialize(NumBones, InitialValue);
 		}
@@ -1471,13 +1470,13 @@ namespace BoneAttributeHelpers
 	{
 		if (Copy)
 		{
-			TDynamicBoneAttributeBase<ParentType, AttribValueType>* Ptr = static_cast<TDynamicBoneAttributeBase<ParentType, AttribValueType>*>(Copy->MakeNew(Mesh));
+			TDynamicBoneAttributeBase<ParentType, AttribValueType>* Ptr = new TDynamicBoneAttributeBase<ParentType, AttribValueType>(Mesh);
 			Attribute = TUniquePtr<TDynamicBoneAttributeBase<ParentType, AttribValueType>>(Ptr);
 			Attribute->Copy(*(Copy));
 		}
 		else
 		{
-			Attribute = nullptr;
+			Attribute.Reset();
 		}
 	}
 }
@@ -1585,21 +1584,13 @@ bool FDynamicMeshAttributeSet::AppendBonesUnique(const FDynamicMeshAttributeSet&
 
 bool FDynamicMeshAttributeSet::CheckBoneValidity(EValidityCheckFailMode FailMode) const
 {
+	const int32 NumBones = HasBones() ? GetNumBones() : 0;
+
 	bool bValid = true;
 	
-	if (!HasBones())
-	{
-		// if boneless, no bone-related attributes should be set
-		bValid = !BoneNameAttrib && !BoneParentIndexAttrib && !BoneColorAttrib && !BonePoseAttrib;
-	}
-	else
-	{
-		const int32 NumBones = GetNumBones();
-	
-		bValid = (BoneParentIndexAttrib->Num() == NumBones || BoneParentIndexAttrib->IsEmpty()) && bValid;
-		bValid = (BoneColorAttrib->Num() == NumBones || BoneColorAttrib->IsEmpty()) && bValid;
-		bValid = (BonePoseAttrib->Num() == NumBones || BonePoseAttrib->IsEmpty()) && bValid;
-	}
+	bValid = (BoneParentIndexAttrib->Num() == NumBones || BoneParentIndexAttrib->IsEmpty()) && bValid;
+	bValid = (BoneColorAttrib->Num() == NumBones || BoneColorAttrib->IsEmpty()) && bValid;
+	bValid = (BonePoseAttrib->Num() == NumBones || BonePoseAttrib->IsEmpty()) && bValid;
 
 	if (FailMode == EValidityCheckFailMode::Check)
 	{
