Add a relatively fast heuristic Merge() function, to combine two OrientedBox3s

The orientation of the merged box will be the lowest-volume choice of the two input box orientations, or a heuristically-oriented box.  The heuristic orientation will be either (A) an interpolation of the orientations, if the boxes are close together, or (B) a box lined up with the line connecting the two box origins, if the boxes are far apart (relative to their extents).

diff --git a/Engine/Source/Runtime/GeometryCore/Public/OrientedBoxTypes.h b/Engine/Source/Runtime/GeometryCore/Public/OrientedBoxTypes.h
--- a/Engine/Source/Runtime/GeometryCore/Public/OrientedBoxTypes.h
+++ b/Engine/Source/Runtime/GeometryCore/Public/OrientedBoxTypes.h
@@ -22,7 +22,7 @@ using namespace UE::Math;
 template<typename RealType>
 struct TOrientedBox3
 {
-	// available for porting: ContainPoint, MergeBoxes()
+	// available for porting: ContainPoint
 
 
 	/** 3D position (center) and orientation (axes) of the box */
@@ -227,6 +227,176 @@ struct TOrientedBox3
 	 }
 
 
+	 // Create a merged TOrientedBox, encompassing this box and one other.
+	 // Uses heuristics to test a few possible orientations and pick the smallest-volume result. Not an optimal fit.
+	 // @param Other						The box to merge with
+	 // @param bOnlyConsiderExistingAxes	If true, the merged box will always use the Rotation of one of the input boxes
+	 TOrientedBox3<RealType> Merge(const TOrientedBox3<RealType>& Other, bool bOnlyConsiderExistingAxes = false) const
+	 {
+		 TQuaternion<RealType> RotInv = Frame.Rotation.Inverse();
+		 TQuaternion<RealType> RotInvOtherRot = RotInv * Other.Frame.Rotation;
+		 TVector<RealType> O2mO1 = Other.Frame.Origin - Frame.Origin;
+		 TVector<RealType> TranslateIntoFrame1 = RotInv * O2mO1;
+		 TVector<RealType> TranslateIntoFrame2 = Other.Frame.Rotation.InverseMultiply(-O2mO1);
+		 TAxisAlignedBox3<RealType> LocalBox1(-Extents, Extents);
+		 TAxisAlignedBox3<RealType> LocalBox2(-Other.Extents, Other.Extents);
+		 // Corners of this box, in its local coordinate frame
+		 TVector<RealType> Corners1[8]
+		 {
+			 LocalBox1.Min,
+			 LocalBox1.Max,
+			 TVector<RealType>(LocalBox1.Min.X, LocalBox1.Min.Y, LocalBox1.Max.Z),
+			 TVector<RealType>(LocalBox1.Min.X, LocalBox1.Max.Y, LocalBox1.Min.Z),
+			 TVector<RealType>(LocalBox1.Max.X, LocalBox1.Min.Y, LocalBox1.Min.Z),
+			 TVector<RealType>(LocalBox1.Min.X, LocalBox1.Max.Y, LocalBox1.Max.Z),
+			 TVector<RealType>(LocalBox1.Max.X, LocalBox1.Min.Y, LocalBox1.Max.Z),
+			 TVector<RealType>(LocalBox1.Max.X, LocalBox1.Max.Y, LocalBox1.Min.Z),
+		 };
+		 // Cordinates of Other, in its local coordinate frame
+		 TVector<RealType> Corners2[8]
+		 {
+			 LocalBox2.Min,
+			 LocalBox2.Max,
+			 TVector<RealType>(LocalBox2.Min.X, LocalBox2.Min.Y, LocalBox2.Max.Z),
+			 TVector<RealType>(LocalBox2.Min.X, LocalBox2.Max.Y, LocalBox2.Min.Z),
+			 TVector<RealType>(LocalBox2.Max.X, LocalBox2.Min.Y, LocalBox2.Min.Z),
+			 TVector<RealType>(LocalBox2.Min.X, LocalBox2.Max.Y, LocalBox2.Max.Z),
+			 TVector<RealType>(LocalBox2.Max.X, LocalBox2.Min.Y, LocalBox2.Max.Z),
+			 TVector<RealType>(LocalBox2.Max.X, LocalBox2.Max.Y, LocalBox2.Min.Z),
+		 };
+		 // Expand each local AABB to contain the other box's points
+		 for (int32 Idx = 0; Idx < 8; ++Idx)
+		 {
+			 LocalBox1.Contain(TranslateIntoFrame1 + RotInvOtherRot * Corners2[Idx]);
+			 LocalBox2.Contain(TranslateIntoFrame2 + RotInvOtherRot.InverseMultiply(Corners1[Idx]));
+		 }
+		 RealType CandidateVolumes[3]{ LocalBox1.Volume(), LocalBox2.Volume(), TMathUtil<RealType>::MaxReal };
+		 TVector<RealType> Axes1[3], Axes2[3], FoundAxes[3];
+		 TAxisAlignedBox3<RealType> FoundBounds;
+		 if (!bOnlyConsiderExistingAxes)
+		 {
+			 // Compute the two local frames
+			 Frame.Rotation.GetAxes(Axes1[0], Axes1[1], Axes1[2]);
+			 Other.Frame.Rotation.GetAxes(Axes2[0], Axes2[1], Axes2[2]);
+			 // Compute metrics re the size of the boxes and how far apart they are
+			 RealType MaxExtentsSizeSq = FMath::Max(Extents.SizeSquared(), Other.Extents.SizeSquared());
+			 RealType CenterDiffSizeSq = O2mO1.SizeSquared();
+			 // For boxes that are far apart relative to their size, fit a box along the line connecting their centers
+			 if (CenterDiffSizeSq > MaxExtentsSizeSq)
+			 {
+				 // find the most-perpendicular original box axis as our second axis
+				 // (Note we could alternatively try interpolating axes, or picking the axis with the largest projected extent ...)
+				 FoundAxes[0] = O2mO1.GetSafeNormal(0.0);
+				 TVector<RealType> MostPerpendicular;
+				 RealType SmallestDot = TMathUtil<RealType>::MaxReal;
+				 for (int32 AxisIdx = 0; AxisIdx < 3; ++AxisIdx)
+				 {
+					 RealType Dot1 = FMath::Abs(Axes1[AxisIdx].Dot(FoundAxes[0]));
+					 RealType Dot2 = FMath::Abs(Axes2[AxisIdx].Dot(FoundAxes[0]));
+					 if (Dot1 < Dot2)
+					 {
+						 if (Dot1 < SmallestDot)
+						 {
+							 SmallestDot = Dot1;
+							 MostPerpendicular = Axes1[AxisIdx];
+						 }
+					 }
+					 else if (Dot2 < SmallestDot)
+					 {
+						 SmallestDot = Dot2;
+						 MostPerpendicular = Axes2[AxisIdx];
+					 }
+				 }
+				 FoundAxes[2] = FoundAxes[0].Cross(MostPerpendicular).GetSafeNormal(0);
+				 FoundAxes[1] = FoundAxes[2].Cross(FoundAxes[0]);
+			 }
+			 else // otherwise, for boxes that are closer together, interpolate their orientations
+			 {
+				 // Find which axes on the Other.Frame.Rotation are closest to the Frame.Rotation's X and Y axes
+				 int32 Axis1X_IdxInAxis2 = 0;
+				 int32 Axis1Y_IdxInAxis2 = 0;
+				 RealType BestXAlignment = FMath::Abs(Axes1[0].Dot(Axes2[0]));
+				 RealType BestYAlignment = FMath::Abs(Axes1[1].Dot(Axes2[0]));
+				 for (int32 Axis2Idx = 1; Axis2Idx < 3; ++Axis2Idx)
+				 {
+					 RealType XAlignment = FMath::Abs(Axes1[0].Dot(Axes2[Axis2Idx]));
+					 RealType YAlignment = FMath::Abs(Axes1[1].Dot(Axes2[Axis2Idx]));
+					 if (XAlignment > BestXAlignment)
+					 {
+						 Axis1X_IdxInAxis2 = Axis2Idx;
+						 BestXAlignment = XAlignment;
+					 }
+					 if (YAlignment > BestYAlignment)
+					 {
+						 Axis1Y_IdxInAxis2 = Axis2Idx;
+						 BestYAlignment = YAlignment;
+					 }
+				 }
+				 if (Axis1X_IdxInAxis2 == Axis1Y_IdxInAxis2)
+				 {
+					 // don't let the closest axis to Y be the same as the closest to X
+					 // (this could happen in a 45-degree case where two axes are tied; should be very rare)
+					 Axis1Y_IdxInAxis2 = (1 + Axis1Y_IdxInAxis2) % 3;
+				 }
+				 RealType SignX = Axes1[0].Dot(Axes2[Axis1X_IdxInAxis2]) < 0 ? -1 : 1;
+				 RealType SignY = Axes1[1].Dot(Axes2[Axis1Y_IdxInAxis2]) < 0 ? -1 : 1;
+				 FoundAxes[0] = (Axes1[0] + SignX * Axes2[Axis1X_IdxInAxis2]).GetSafeNormal();
+				 FoundAxes[1] = (Axes1[1] + SignY * Axes2[Axis1Y_IdxInAxis2]);
+				 // Make sure the axes form an orthonormal frame
+				 FoundAxes[2] = (FoundAxes[0].Cross(FoundAxes[1])).GetSafeNormal();
+				 FoundAxes[1] = FoundAxes[2].Cross(FoundAxes[0]).GetSafeNormal();
+			 }
+
+			 // Fit a local AABB in the Found space
+			 for (int32 CornerIdx = 0; CornerIdx < 8; ++CornerIdx)
+			 {
+				 TVector<RealType> WorldCorner1 = Frame.FromFramePoint(Corners1[CornerIdx]);
+				 TVector<RealType> AxesCorner1(WorldCorner1.Dot(FoundAxes[0]), WorldCorner1.Dot(FoundAxes[1]), WorldCorner1.Dot(FoundAxes[2]));
+				 FoundBounds.Contain(AxesCorner1);
+				 TVector<RealType> WorldCorner2 = Other.Frame.FromFramePoint(Corners2[CornerIdx]);
+				 TVector<RealType> AxesCorner2(WorldCorner2.Dot(FoundAxes[0]), WorldCorner2.Dot(FoundAxes[1]), WorldCorner2.Dot(FoundAxes[2]));
+				 FoundBounds.Contain(AxesCorner2);
+			 }
+			 CandidateVolumes[2] = FoundBounds.Volume();
+		 }
+
+		 // Choose the coordinate space with the smallest bounding box, by volume
+		 int BestIdx = 0;
+		 RealType BestVolume = CandidateVolumes[0];
+		 int AxesToConsider = bOnlyConsiderExistingAxes ? 2 : 3;
+		 for (int Idx = 1; Idx < AxesToConsider; ++Idx)
+		 {
+			 if (CandidateVolumes[Idx] < BestVolume)
+			 {
+				 BestIdx = Idx;
+				 BestVolume = CandidateVolumes[Idx];
+			 }
+		 }
+		 TOrientedBox3<RealType> Result;
+		 if (BestIdx == 2)
+		 {
+			 // Convert the Found Axes into a new oriented box
+			 Result.Frame.Rotation.SetFromRotationMatrix(TMatrix3<RealType>(FoundAxes[0], FoundAxes[1], FoundAxes[2], false));
+			 TVector<RealType> FoundCenter = FoundBounds.Center();
+			 Result.Frame.Origin = FoundAxes[0] * FoundCenter.X + FoundAxes[1] * FoundCenter.Y + FoundAxes[2] * FoundCenter.Z;
+			 Result.Extents = FoundBounds.Extents();
+		 }
+		 else if (BestIdx == 0)
+		 {
+			 Result.Frame.Rotation = Frame.Rotation;
+			 Result.Frame.Origin = Frame.FromFramePoint(LocalBox1.Center());
+			 Result.Extents = LocalBox1.Extents();
+		 }
+		 else // BestIdx == 1
+		 {
+			 Result.Frame.Rotation = Other.Frame.Rotation;
+			 Result.Frame.Origin = Other.Frame.FromFramePoint(LocalBox2.Center());
+			 Result.Extents = LocalBox2.Extents();
+		 }
+		 // Add a small tolerance so that the input boxes are more likely to be properly contained in the merged result, after floating point error
+		 Result.Extents += FVector3d(FMathd::ZeroTolerance, FMathd::ZeroTolerance, FMathd::ZeroTolerance);
+		 return Result;
+	 }
 
 };
 
