-DynamicMeshEditor.cpp:  Handle bone data when appending two meshes and use it to reindex the the combined skinning weights.
-DynamicMeshAttributeSet.cpp:  Fixed a bug where we were resetting the bone attributes when matching.
-DynamicVertexSkinWeightsAttribute.h: Added a function that allows to reindex bone indices with respect to a new skeleton.
-SkeletalMeshLODModle.cpp: Handle bone attributes when FSkeletalMeshLODModel::GetMeshDescription is called.

diff --git a/Engine/Plugins/Runtime/GeometryProcessing/Source/DynamicMesh/Private/DynamicMeshEditor.cpp b/Engine/Plugins/Runtime/GeometryProcessing/Source/DynamicMesh/Private/DynamicMeshEditor.cpp
--- a/Engine/Plugins/Runtime/GeometryProcessing/Source/DynamicMesh/Private/DynamicMeshEditor.cpp
+++ b/Engine/Plugins/Runtime/GeometryProcessing/Source/DynamicMesh/Private/DynamicMeshEditor.cpp
@@ -7,6 +7,7 @@
 #include "MeshRegionBoundaryLoops.h"
 #include "DynamicSubmesh3.h"
 #include "DynamicMesh/DynamicVertexSkinWeightsAttribute.h"
+#include "DynamicMesh/DynamicBoneAttribute.h"
 #include "DynamicMesh/MeshNormals.h"
 #include "MeshQueries.h"
 #include "Selections/MeshConnectedComponents.h"
@@ -1705,12 +1706,48 @@ void FDynamicMeshEditor::AppendMesh(const FDynamicMesh3* AppendMesh,
 			}
 		}
 
-		for (const TPair<FName, TUniquePtr<FDynamicMeshVertexSkinWeightsAttribute>>& AttribPair : AppendMesh->Attributes()->GetSkinWeightsAttributes())
+		if (AppendMesh->Attributes()->HasBones() && Mesh->Attributes()->HasBones())
 		{
-			FDynamicMeshVertexSkinWeightsAttribute* ToAttrib = Mesh->Attributes()->GetSkinWeightsAttribute(AttribPair.Key);
-			if (ToAttrib)
+			const bool bSameSkeletons = AppendMesh->Attributes()->GetBoneNames()->IsSameAs(*Mesh->Attributes()->GetBoneNames());
+
+			if (!bSameSkeletons)
 			{
-				ToAttrib->CopyThroughMapping(AttribPair.Value.Get(), IndexMapsOut);
+				Mesh->Attributes()->AppendBonesUnique(*AppendMesh->Attributes());
+			}
+			
+			for (const TPair<FName, TUniquePtr<FDynamicMeshVertexSkinWeightsAttribute>>& AttribPair : AppendMesh->Attributes()->GetSkinWeightsAttributes())
+			{
+				FDynamicMeshVertexSkinWeightsAttribute* ToAttrib = Mesh->Attributes()->GetSkinWeightsAttribute(AttribPair.Key);
+				if (ToAttrib)
+				{	
+					if (bSameSkeletons)
+					{
+						// If the skeletons are the same then we do not need to re-index and we simply copy
+						ToAttrib->CopyThroughMapping(AttribPair.Value.Get(), IndexMapsOut);
+					}
+					else
+					{
+						// Make a copy of the append mesh skinning weights and reindex them with respect to the new skeleton
+						FDynamicMeshVertexSkinWeightsAttribute CopyAppendMeshAttrib;
+						CopyAppendMeshAttrib.Copy(*AttribPair.Value.Get());
+						CopyAppendMeshAttrib.ReindexBoneIndicesToSkeleton(AppendMesh->Attributes()->GetBoneNames()->GetAttribValues(),
+						 											  	  Mesh->Attributes()->GetBoneNames()->GetAttribValues());
+
+						// Now copy re-indexed weights to the mesh we are appending to
+						ToAttrib->CopyThroughMapping(&CopyAppendMeshAttrib, IndexMapsOut);
+					}
+				}
+			}
+		}
+		else
+		{
+			for (const TPair<FName, TUniquePtr<FDynamicMeshVertexSkinWeightsAttribute>>& AttribPair : AppendMesh->Attributes()->GetSkinWeightsAttributes())
+			{
+				FDynamicMeshVertexSkinWeightsAttribute* ToAttrib = Mesh->Attributes()->GetSkinWeightsAttribute(AttribPair.Key);
+				if (ToAttrib)
+				{
+					ToAttrib->CopyThroughMapping(AttribPair.Value.Get(), IndexMapsOut);
+				}
 			}
 		}
 		
diff --git a/Engine/Source/Runtime/Engine/Private/SkeletalMeshLODModel.cpp b/Engine/Source/Runtime/Engine/Private/SkeletalMeshLODModel.cpp
--- a/Engine/Source/Runtime/Engine/Private/SkeletalMeshLODModel.cpp
+++ b/Engine/Source/Runtime/Engine/Private/SkeletalMeshLODModel.cpp
@@ -16,6 +16,7 @@
 #include "Serialization/MemoryWriter.h"
 #include "UObject/FortniteMainBranchObjectVersion.h"
 #include "SkeletalMeshAttributes.h"
+#include "ReferenceSkeleton.h"
 
 /*-----------------------------------------------------------------------------
 FSkelMeshImportedMeshInfo
@@ -1384,12 +1385,20 @@ void FSkeletalMeshLODModel::GetMeshDescription(FMeshDescription& MeshDescription
 
 	TPolygonGroupAttributesRef<FName> PolygonGroupMaterialSlotNames = MeshAttributes.GetPolygonGroupMaterialSlotNames();
 	
+	FSkeletalMeshAttributes::FBoneNameAttributesRef BoneNames = MeshAttributes.GetBoneNames();
+	FSkeletalMeshAttributes::FBoneParentIndexAttributesRef BoneParentIndices = MeshAttributes.GetBoneParentIndices();
+	FSkeletalMeshAttributes::FBonePoseAttributesRef BonePoses = MeshAttributes.GetBonePoses();
+
 	const int32 NumTriangles = IndexBuffer.Num() / 3;
 
+	const FReferenceSkeleton& RefSkeleton = Owner->GetRefSkeleton();
+	const int NumBones = RefSkeleton.GetRawBoneNum();
+
 	MeshDescription.ReserveNewPolygonGroups(Sections.Num());
 	MeshDescription.ReserveNewTriangles(NumTriangles);
 	MeshDescription.ReserveNewVertexInstances(NumTriangles * 3);
 	MeshDescription.ReserveNewVertices(static_cast<int32>(NumVertices));
+	MeshAttributes.ReserveNewBones(NumBones);
 
 	TArray<FVertexID> VertexIDs;
 	VertexIDs.Reserve(NumVertices);
@@ -1482,6 +1491,19 @@ void FSkeletalMeshLODModel::GetMeshDescription(FMeshDescription& MeshDescription
 			MeshDescription.CreateTriangle(PolygonGroupID, TriangleVertexInstanceIDs);
 		}
 	}
+
+	// Set Bone Attributes
+	for (int Idx = 0; Idx < NumBones; ++Idx)
+	{
+		const FMeshBoneInfo& BoneInfo = RefSkeleton.GetRawRefBoneInfo()[Idx];
+		const FTransform& BoneTransform = RefSkeleton.GetRawRefBonePose()[Idx];
+
+		const FBoneID BoneID = MeshAttributes.CreateBone();
+		
+		BoneNames.Set(BoneID, BoneInfo.Name);
+		BoneParentIndices.Set(BoneID, BoneInfo.ParentIndex);
+		BonePoses.Set(BoneID, BoneTransform);
+	}
 }
 
 #endif // WITH_EDITOR
diff --git a/Engine/Source/Runtime/GeometryCore/Private/DynamicMesh/DynamicMeshAttributeSet.cpp b/Engine/Source/Runtime/GeometryCore/Private/DynamicMesh/DynamicMeshAttributeSet.cpp
--- a/Engine/Source/Runtime/GeometryCore/Private/DynamicMesh/DynamicMeshAttributeSet.cpp
+++ b/Engine/Source/Runtime/GeometryCore/Private/DynamicMesh/DynamicMeshAttributeSet.cpp
@@ -1457,9 +1457,12 @@ namespace BoneAttributeHelpers
 		if (bWantAttrib)
 		{	
 			const int32 NumBones = ToMatch ? ToMatch->Num() : 0;
-			TDynamicBoneAttributeBase<ParentType, AttribValueType>* Ptr = new TDynamicBoneAttributeBase<ParentType, AttribValueType>(Mesh);
-			Attribute = TUniquePtr<TDynamicBoneAttributeBase<ParentType, AttribValueType>>(Ptr);
-			Attribute->Initialize(NumBones, InitialValue);
+			if (!Attribute)
+			{
+				TDynamicBoneAttributeBase<ParentType, AttribValueType>* Ptr = new TDynamicBoneAttributeBase<ParentType, AttribValueType>(Mesh);
+				Attribute = TUniquePtr<TDynamicBoneAttributeBase<ParentType, AttribValueType>>(Ptr);
+				Attribute->Initialize(NumBones, InitialValue);
+			}
 		}
 	}
 
diff --git a/Engine/Source/Runtime/GeometryCore/Public/DynamicMesh/DynamicVertexSkinWeightsAttribute.h b/Engine/Source/Runtime/GeometryCore/Public/DynamicMesh/DynamicVertexSkinWeightsAttribute.h
--- a/Engine/Source/Runtime/GeometryCore/Public/DynamicMesh/DynamicVertexSkinWeightsAttribute.h
+++ b/Engine/Source/Runtime/GeometryCore/Public/DynamicMesh/DynamicVertexSkinWeightsAttribute.h
@@ -48,6 +48,7 @@ class TDynamicVertexSkinWeightsAttribute final :
 {
 public:
 	using FBoneWeights = UE::AnimationCore::FBoneWeights;
+	using FBoneWeight  = UE::AnimationCore::FBoneWeight;
 
 protected:
 	friend class FDynamicMeshAttributeSet;
@@ -158,7 +159,98 @@ public:
 		VertexBoneWeights.InsertAt(InBoneWeights, InNewVertexID);
 	}
 
+	
+	/** 
+	 * Re-index bone indices of the skin weights from one reference skeleton (represented by an array of bone names) 
+	 * to another.
+	 * 
+	 * @param FromRefSkeleton The reference skeleton that the current skin weights are indexed against.
+	 * @param ToRefSkeleton The new reference skeleton that should be used to reindex the weights.
+	 * 
+	 * @return true on success, false otherwise
+	 * 
+	 * @note Reasons for failure are:
+	 * 		 - bone index of one of the bone weights is larger than the number of bones in FromRefSkeleton, meaning
+	 * 		   the weights most likely were not computed with respect to the FromRefSkeleton.
+	 * 		 - one of the bone weights references a bone in FromRefSkeleton that doesn't exist in ToRefSkeleton.
+	 * 		 - ToRefSkeleton contains duplicates.
+	 *
+	 * @note FromRefSkeleton can contain bones that are not present in ToRefSkeleton as long as those bones are not 
+			 referenced by any of the weights.
+	 */
+	bool ReindexBoneIndicesToSkeleton(const TArray<FName>& FromRefSkeleton, const TArray<FName>& ToRefSkeleton)
+	{	
+		if (VertexBoneWeights.Num() == 0)
+		{
+			return true; // empty weights, nothing to do
+		}
+
+		if (FromRefSkeleton == ToRefSkeleton)
+		{
+			return true; // skeletons are the same, nothing to do
+		}
+
+		// Create a map from bone name to bone index for the ToRefSkeleton so we can
+		// check if bone with a given name in FromRefSkeleton exists in ToRefSkeleton 
+		// and at what index
+		TMap<FName, int32> NameToIndex;
+		NameToIndex.Reserve(ToRefSkeleton.Num());
+		for (int BoneID = 0; BoneID < ToRefSkeleton.Num(); ++BoneID)
+		{
+			const FName& BoneName = ToRefSkeleton[BoneID];
+			if (NameToIndex.Contains(BoneName))
+			{
+				return false; // there should be no duplicates in ToRefSkeleton
+			}
+			NameToIndex.Add(BoneName, BoneID);
+		}
+
+		// Store temporary results of reindexing in a new dynamic vector
+		TDynamicVector<FBoneWeights> NewVertexBoneWeights;
+		NewVertexBoneWeights.Resize(VertexBoneWeights.Num());
+
+		// Since we are only changing indicies, make sure we don't do any re-normalization on the weights
+		UE::AnimationCore::FBoneWeightsSettings BoneSettings;
+		BoneSettings.SetNormalizeType(UE::AnimationCore::EBoneWeightNormalizeType::None); 
+
+		// Iterate over every skin weight of every vertex and remap it's bone index from FromRefSkeleton to ToRefSkeleton
+		for (int VertexID = 0; VertexID < VertexBoneWeights.Num(); ++VertexID)
+		{
+			FBoneWeights BoneWeights = VertexBoneWeights[VertexID];
+			if (BoneWeights.Num() == 0)
+			{
+				continue; // vertex has no weights
+			}
+
+			FBoneWeights NewWeights;
+			for (int32 Idx = 0; Idx < BoneWeights.Num(); ++Idx)
+			{
+				const FBoneWeight& BoneWeight = BoneWeights[Idx];
+				if (BoneWeight.GetBoneIndex() >= FromRefSkeleton.Num())
+				{
+					return false; // invalid FromRefSkeleton
+				}
+				
+				const FName& FromName = FromRefSkeleton[BoneWeight.GetBoneIndex()];
+				const uint16 FromWeight = BoneWeight.GetRawWeight();
+
+				const int32* ToIdx = NameToIndex.Find(FromName);
+				if (ToIdx)
+				{
+					NewWeights.SetBoneWeight(FBoneWeight(static_cast<FBoneIndexType>(*ToIdx), FromWeight), BoneSettings);
+					NewVertexBoneWeights[VertexID] = NewWeights;
+				}
+				else
+				{
+					return false; // ToRefSkeleton must be a superset of all the bones referenced in FromRefSkeleton by the skinning weights
+				}
+			}
+		}
+
+		VertexBoneWeights = MoveTemp(NewVertexBoneWeights);
 
+		return true;
+	}
 
 	//
 	// Accessors/Queries
