fix invalid mesh state caused by not merging some neighboring edges in FDynamicMesh3's EdgeMerge, when mesh has bowties
+ add a test to cover this case

diff --git a/Engine/Plugins/Runtime/GeometryProcessing/Source/DynamicMesh/Private/DynamicMeshEditor.cpp b/Engine/Plugins/Runtime/GeometryProcessing/Source/DynamicMesh/Private/DynamicMeshEditor.cpp
--- a/Engine/Plugins/Runtime/GeometryProcessing/Source/DynamicMesh/Private/DynamicMeshEditor.cpp
+++ b/Engine/Plugins/Runtime/GeometryProcessing/Source/DynamicMesh/Private/DynamicMeshEditor.cpp
@@ -233,6 +233,10 @@ bool FDynamicMeshEditor::WeldVertexLoops(const TArray<int32>& Loop1, const TArra
 			{
 				SkipEdges.Add(MergeInfo.ExtraRemovedEdges.B);
 			}
+			for (int RemovedEID : MergeInfo.BowtiesRemovedEdges)
+			{
+				SkipEdges.Add(RemovedEID);
+			}
 		}
 	}
 
diff --git a/Engine/Source/Runtime/GeometryCore/Private/DynamicMesh/DynamicMesh3_Edits.cpp b/Engine/Source/Runtime/GeometryCore/Private/DynamicMesh/DynamicMesh3_Edits.cpp
--- a/Engine/Source/Runtime/GeometryCore/Private/DynamicMesh/DynamicMesh3_Edits.cpp
+++ b/Engine/Source/Runtime/GeometryCore/Private/DynamicMesh/DynamicMesh3_Edits.cpp
@@ -1686,6 +1686,7 @@ EMeshResult FDynamicMesh3::MergeEdges(int eKeep, int eDiscard, FMergeEdgesInfo&
 	// 'link condition' in edge collapses.
 	// Note that we have to catch cases where both edges to the shared vertex are
 	// boundary edges, in that case we will also merge this edge later on
+	int MaxAdjBoundaryMerges[2]{ 0,0 };
 	if (a != c)
 	{
 		int ea = 0, ec = 0, other_v = (b == d) ? b : -1;
@@ -1698,6 +1699,10 @@ EMeshResult FDynamicMesh3::MergeEdges(int eKeep, int eDiscard, FMergeEdgesInfo&
 				{
 					return EMeshResult::Failed_InvalidNeighbourhood;
 				}
+				else
+				{
+					MaxAdjBoundaryMerges[0]++;
+				}
 			}
 		}
 	}
@@ -1713,6 +1718,10 @@ EMeshResult FDynamicMesh3::MergeEdges(int eKeep, int eDiscard, FMergeEdgesInfo&
 				{
 					return EMeshResult::Failed_InvalidNeighbourhood;
 				}
+				else
+				{
+					MaxAdjBoundaryMerges[1]++;
+				}
 			}
 		}
 	}
@@ -1812,10 +1821,8 @@ EMeshResult FDynamicMesh3::MergeEdges(int eKeep, int eDiscard, FMergeEdgesInfo&
 	// to either a or b that are connected to the same vertex on their 'other' side.
 	// So we now have two boundary edges connecting the same two vertices - disaster!
 	// We need to find and merge these edges.
-	// Q: I don't think it is possible to have multiple such edge-pairs at a or b
-	//    But I am not certain...is a bit tricky to handle because we modify edges_v...
 	MergeInfo.ExtraRemovedEdges = FIndex2i(InvalidID, InvalidID);
-	MergeInfo.ExtraKeptEdges = MergeInfo.ExtraRemovedEdges;
+	MergeInfo.ExtraKeptEdges = FIndex2i(InvalidID, InvalidID);
 	for (int vi = 0; vi < 2; ++vi)
 	{
 		int v1 = a, v2 = c;   // vertices of merged edge
@@ -1831,14 +1838,14 @@ EMeshResult FDynamicMesh3::MergeEdges(int eKeep, int eDiscard, FMergeEdgesInfo&
 		TArray<int> edges_v;
 		GetVertexEdgesList(v1, edges_v);
 		int Nedges = (int)edges_v.Num();
-		bool found = false;
+		int FoundNum = 0;
 		// in this loop, we compare 'other' vert_1 and vert_2 of edges around v1.
 		// problem case is when vert_1 == vert_2  (ie two edges w/ same other vtx).
-		//restart_merge_loop:
-		for (int i = 0; i < Nedges && found == false; ++i)
+		for (int i = 0; i < Nedges && FoundNum < MaxAdjBoundaryMerges[vi]; ++i)
 		{
 			int edge_1 = edges_v[i];
-			if (IsBoundaryEdge(edge_1) == false)
+			// Skip any non-boundary edge, or edge we've already removed via merging
+			if (!EdgeRefCounts.IsValidUnsafe(edge_1) || !IsBoundaryEdge(edge_1))
 			{
 				continue;
 			}
@@ -1846,8 +1853,13 @@ EMeshResult FDynamicMesh3::MergeEdges(int eKeep, int eDiscard, FMergeEdgesInfo&
 			for (int j = i + 1; j < Nedges; ++j)
 			{
 				int edge_2 = edges_v[j];
+				// Skip any non-boundary edge, or edge we've already removed via merging
+				if (!EdgeRefCounts.IsValidUnsafe(edge_2) || !IsBoundaryEdge(edge_2))
+				{
+					continue;
+				}
 				int vert_2 = GetOtherEdgeVertex(edge_2, v1);
-				if (vert_1 == vert_2 && IsBoundaryEdge(edge_2))  // if ! boundary here, we are in deep trouble...
+				if (vert_1 == vert_2)
 				{
 					// replace edge_2 w/ edge_1 in tri, update edge and vtx-edge-nbr lists
 					int tri_1 = Edges[edge_1].Tri[0];
@@ -1857,14 +1869,19 @@ EMeshResult FDynamicMesh3::MergeEdges(int eKeep, int eDiscard, FMergeEdgesInfo&
 					VertexEdgeLists.Remove(v1, edge_2);
 					VertexEdgeLists.Remove(vert_1, edge_2);
 					EdgeRefCounts.Decrement(edge_2);
-					MergeInfo.ExtraRemovedEdges[vi] = edge_2;
-					MergeInfo.ExtraKeptEdges[vi] = edge_1;
-
-					//edges_v = VertexEdgeLists_list(v1);      // this code allows us to continue checking, ie in case we had
-					//Nedges = edges_v.Count;               // multiple such edges. but I don't think it's possible.
-					//goto restart_merge_loop;
-					found = true;			  // exit outer i loop
-					break;					  // exit inner j loop
+					if (FoundNum == 0)
+					{
+						MergeInfo.ExtraRemovedEdges[vi] = edge_2;
+						MergeInfo.ExtraKeptEdges[vi] = edge_1;
+					}
+					else
+					{
+						MergeInfo.BowtiesRemovedEdges.Add(edge_2);
+						MergeInfo.BowtiesKeptEdges.Add(edge_1);
+					}
+
+					FoundNum++;				  // exit outer i loop if we've found all possible merge edges
+					break;					  // exit inner j loop; we won't merge anything else to edge_1
 				}
 			}
 		}
diff --git a/Engine/Source/Runtime/GeometryCore/Public/DynamicMesh/InfoTypes.h b/Engine/Source/Runtime/GeometryCore/Public/DynamicMesh/InfoTypes.h
--- a/Engine/Source/Runtime/GeometryCore/Public/DynamicMesh/InfoTypes.h
+++ b/Engine/Source/Runtime/GeometryCore/Public/DynamicMesh/InfoTypes.h
@@ -203,6 +203,10 @@ struct FMergeEdgesInfo
 
 	FIndex2i ExtraRemovedEdges; // extra removed edges, see description below. Either may be or InvalidID
 	FIndex2i ExtraKeptEdges;	// extra kept edges, paired with ExtraRemovedEdges
+
+	// Even more Removed and Kept edges, in cases where there were multiple such edges on one or both sides of the merged edge
+	// Only possible if the pre-merge mesh had non-manifold vertices (aka bowties), in almost all meshes these arrays will be empty
+	TArray<int, TInlineAllocator<4>> BowtiesRemovedEdges, BowtiesKeptEdges;
 };
 
 /** Information about mesh elements modified/created by PokeTriangle() */
