GeometryCore: added FindAllHitTriangles function to get all triangles intersected by a ray (not just the closest one)

diff --git a/Engine/Source/Runtime/GeometryCore/Public/Spatial/MeshAABBTree3.h b/Engine/Source/Runtime/GeometryCore/Public/Spatial/MeshAABBTree3.h
--- a/Engine/Source/Runtime/GeometryCore/Public/Spatial/MeshAABBTree3.h
+++ b/Engine/Source/Runtime/GeometryCore/Public/Spatial/MeshAABBTree3.h
@@ -458,6 +458,33 @@ public:
 		FindHitTriangle(RootIndex, Ray, NearestT, TID, BaryCoords, Options);
 		return TID != IndexConstants::InvalidID;
 	}
+	
+	virtual bool FindAllHitTriangles(
+		const FRay3d& Ray, TArray<MeshIntersection::FHitIntersectionResult>& OutHits,
+		const FQueryOptions& Options = FQueryOptions()) const override
+	{
+		if ( ensure(IsValid(Options.bAllowUnsafeModifiedMeshQueries)) == false )
+		{
+			return false;
+		}
+
+		using HitResult = MeshIntersection::FHitIntersectionResult;
+		
+		FindHitTriangles(RootIndex, Ray, OutHits, Options);
+
+		if (OutHits.IsEmpty())
+		{
+			return false;
+		}
+
+		// sort all hits by distance		
+		OutHits.Sort([](const HitResult& Hit0, const HitResult& Hit1)
+		{
+			return Hit0.Distance < Hit1.Distance;
+		});
+		
+		return true;
+	}
 
 protected:
 	void FindHitTriangle(
@@ -537,6 +564,82 @@ protected:
 		}
 	}
 
+	void FindHitTriangles(
+		int IBox, const FRay3d& Ray, TArray<MeshIntersection::FHitIntersectionResult>& Intersections,
+		const FQueryOptions& Options = FQueryOptions()) const
+	{
+		// Note: using TNumericLimits<float>::Max() here because we need to use <= to compare Box hit
+        //   to MaxDistance, and Box hit returns TNumericLimits<double>::Max() on no-hit. So, if we set
+        //   MaxDistance to TNumericLimits<double>::Max(), then we will test all boxes (!)
+        const double MaxDistance = (Options.MaxDistance < TNumericLimits<float>::Max()) ? Options.MaxDistance : TNumericLimits<float>::Max();
+		static constexpr double e = FMathd::ZeroTolerance;
+		
+		int idx = BoxToIndex[IBox];
+		if (idx < TrianglesEnd)
+		{ // triangle-list case, array is [N t1 t2 ... tN]
+			FTriangle3d Triangle;
+			int num_tris = IndexList[idx];
+			for (int i = 1; i <= num_tris; ++i)
+			{
+				int ti = IndexList[idx + i];
+				if (Options.TriangleFilterF != nullptr && Options.TriangleFilterF(ti) == false)
+				{
+					continue;
+				}
+
+				Mesh->GetTriVertices(ti, Triangle.V[0], Triangle.V[1], Triangle.V[2]);
+				FIntrRay3Triangle3d Query = FIntrRay3Triangle3d(Ray, Triangle);
+				if (Query.Find() && Query.RayParameter <= MaxDistance + e)
+				{
+					MeshIntersection::FHitIntersectionResult NewIntersection({ti, Query.RayParameter, Query.TriangleBaryCoords});
+					Intersections.Emplace(MoveTemp(NewIntersection));
+				}
+			}
+		}
+		else
+		{ // internal node, either 1 or 2 child boxes
+			int iChild1 = IndexList[idx];
+			if (iChild1 < 0)
+			{ // 1 child, descend if nearer than cur min-dist
+				iChild1 = (-iChild1) - 1;
+				double fChild1T = box_ray_intersect_t(iChild1, Ray);
+				if (fChild1T <= MaxDistance + e)
+				{
+					FindHitTriangles(iChild1, Ray, Intersections, Options);
+				}
+			}
+			else
+			{ // 2 children, descend closest first
+				iChild1 = iChild1 - 1;
+				int iChild2 = IndexList[idx + 1] - 1;
+
+				double fChild1T = box_ray_intersect_t(iChild1, Ray);
+				double fChild2T = box_ray_intersect_t(iChild2, Ray);
+				if (fChild1T < fChild2T)
+				{
+					if (fChild1T <= MaxDistance + e)
+					{
+						FindHitTriangles(iChild1, Ray, Intersections, Options);
+						if (fChild2T <= MaxDistance + e)
+						{
+							FindHitTriangles(iChild2, Ray, Intersections, Options);
+						}
+					}
+				}
+				else
+				{
+					if (fChild2T <= MaxDistance + e)
+					{
+						FindHitTriangles(iChild2, Ray, Intersections, Options);
+						if (fChild1T <= MaxDistance + e)
+						{
+							FindHitTriangles(iChild1, Ray, Intersections, Options);
+						}
+					}
+				}
+			}
+		}
+	}
 
 public:
 	/**
diff --git a/Engine/Source/Runtime/GeometryCore/Public/Spatial/SpatialInterfaces.h b/Engine/Source/Runtime/GeometryCore/Public/Spatial/SpatialInterfaces.h
--- a/Engine/Source/Runtime/GeometryCore/Public/Spatial/SpatialInterfaces.h
+++ b/Engine/Source/Runtime/GeometryCore/Public/Spatial/SpatialInterfaces.h
@@ -7,6 +7,16 @@
 #include "Math/Ray.h"
 #include "VectorTypes.h"
 
+namespace MeshIntersection
+{
+	struct FHitIntersectionResult
+	{
+		int TriangleId;
+		double Distance;
+		FVector3d BaryCoords;
+	};
+}
+
 namespace UE
 {
 namespace Geometry
@@ -124,7 +134,14 @@ public:
 	 */
 	virtual bool FindNearestHitTriangle(const FRay3d& Ray, double& NearestT, int& TID, FVector3d& BaryCoords, const FQueryOptions& Options = FQueryOptions()) const = 0;
 
-
+	/**
+	 * Find all triangles intersected by the given ray sorted by distance
+	 * @param Ray query ray
+	 * @param OutHits returned-by-reference hit infos sorted by distance
+	 * @param Options Query options (ex. max distance)
+	 * @return true if hit, false if no hit found
+	 */
+	virtual bool FindAllHitTriangles(const FRay3d& Ray, TArray<MeshIntersection::FHitIntersectionResult>& OutHits, const FQueryOptions& Options = FQueryOptions()) const = 0;
 };
 
 
