Fix texture filter kernels used in Baking Tools

diff --git a/Engine/Plugins/Runtime/GeometryProcessing/Source/DynamicMesh/Private/Sampling/MeshMapBaker.cpp b/Engine/Plugins/Runtime/GeometryProcessing/Source/DynamicMesh/Private/Sampling/MeshMapBaker.cpp
--- a/Engine/Plugins/Runtime/GeometryProcessing/Source/DynamicMesh/Private/Sampling/MeshMapBaker.cpp
+++ b/Engine/Plugins/Runtime/GeometryProcessing/Source/DynamicMesh/Private/Sampling/MeshMapBaker.cpp
@@ -592,9 +592,7 @@ void FMeshMapBaker::BakeSample(
 			// Note: There will be no contribution if the sample and pixel are on different UV charts
 			float FilterWeight = Weight * static_cast<float>(SampleUVChart == BufferTilePixelUVChart);
 			{
-				FVector2d TexelDistance = Dimensions.GetTexelUV(SourceCoords) - UVPosition;
-				TexelDistance.X *= Dimensions.GetWidth();
-				TexelDistance.Y *= Dimensions.GetHeight();
+				FVector2d TexelDistance = Dimensions.GetTexelDistance(Dimensions.GetTexelUV(SourceCoords), UVPosition);
 				FilterWeight *= TextureFilterEval(TexelDistance);
 			}
 
diff --git a/Engine/Source/Runtime/GeometryCore/Public/Image/BCSplineFilter.h b/Engine/Source/Runtime/GeometryCore/Public/Image/BCSplineFilter.h
--- a/Engine/Source/Runtime/GeometryCore/Public/Image/BCSplineFilter.h
+++ b/Engine/Source/Runtime/GeometryCore/Public/Image/BCSplineFilter.h
@@ -26,14 +26,16 @@ template <EBCSplineType SplineType, bool bIsRadial>
 class FBCSplineFilter
 {
 private:
+	// Measured in Texel units i.e., Radius of 1 <=> Radius of 1 texel side length
 	float Radius = 1.0f;
+
 	float Coeff1[3] = {0.0f, 0.0f, 0.0f};
 	float Coeff2[4] = {0.0f, 0.0f, 0.0f, 0.0f};
 
 public:
 	FBCSplineFilter(const float RadiusIn);
 
-	/** @return the filter weight given a 2D distance vector. */
+	/** @return the filter weight given a 2D distance vector, in Texel units. */
 	float GetWeight(const FVector2d& Dist) const;
 };
 
@@ -43,12 +45,12 @@ FBCSplineFilter<SplineType, bIsRadial>::FBCSplineFilter(const float RadiusIn)
 {
 	auto ComputeCoeffs = [this](const float& B, const float& C)
 	{
-		// [x < 1]:       1/6 * ((12 - 9B - 6C)x^3 + (-18 + 12B + 6C)x^2 + (6 - 2B))
+		// |x| < 1:        1/6 * ((12 - 9B - 6C)x^3 + (-18 + 12B + 6C)x^2 + (6 - 2B))
 		Coeff1[0] = (12.0f - 9.0f * B - 6.0f * C) / 6.0f;
 		Coeff1[1] = (-18.0f + 12.0f * B + 6.0f * C) / 6.0f;
 		Coeff1[2] = (6.0f - 2.0f * B) / 6.0f;
 
-		// [1 < x < 2]:   1/6 * ((-B - 6C)x^3 + (6B + 30C)x^2 + (-12B - 48C)x + (8B + 24C))
+		// 1 <= |x| < 2:   1/6 * ((-B - 6C)x^3 + (6B + 30C)x^2 + (-12B - 48C)x + (8B + 24C))
 		Coeff2[0] = (-B - 6.0f * C) / 6.0f;
 		Coeff2[1] = (6.0f * B + 30.0f * C) / 6.0f;
 		Coeff2[2] = (-12.0f * B - 48.0f * C) / 6.0f;
@@ -80,15 +82,17 @@ float FBCSplineFilter<SplineType, bIsRadial>::GetWeight(const FVector2d& Dist) c
 {
 	auto ComputeWeight = [this](float X) -> float
 	{
-		const float AbsX = FMathf::Abs(X);
+		// Scale the filter kernel f(X) in the X-direction by making the substitution f(X/Radius)
+		const float AbsX = FMathf::Abs(X / Radius);
 		const float AbsX2 = AbsX * AbsX;
 		const float AbsX3 = AbsX * AbsX2;
+
 		float Weight = 0.0f;
-		if(AbsX < Radius)
+		if(AbsX < 1.0f)
 		{
 			Weight = Coeff1[0] * AbsX3 + Coeff1[1] * AbsX2 + Coeff1[2];
 		}
-		else if(AbsX < Radius * 2.0f)
+		else if(AbsX < 2.0f)
 		{
 			Weight = Coeff2[0] * AbsX3 + Coeff2[1] * AbsX2 + Coeff2[2] * AbsX + Coeff2[3];
 		}
@@ -97,7 +101,7 @@ float FBCSplineFilter<SplineType, bIsRadial>::GetWeight(const FVector2d& Dist) c
 
 	if constexpr(bIsRadial)
 	{
-		return ComputeWeight(Dist.Length());
+		return ComputeWeight((float)Dist.Length());
 	}
 	else
 	{
diff --git a/Engine/Source/Runtime/GeometryCore/Public/Image/BoxFilter.h b/Engine/Source/Runtime/GeometryCore/Public/Image/BoxFilter.h
--- a/Engine/Source/Runtime/GeometryCore/Public/Image/BoxFilter.h
+++ b/Engine/Source/Runtime/GeometryCore/Public/Image/BoxFilter.h
@@ -13,6 +13,7 @@ namespace Geometry
 class FBoxFilter
 {
 private:
+	// Measured in Texel units i.e., Radius of 1 <=> Radius of 1 texel side length
 	float Radius = 0.5f;
 	
 public:
@@ -21,11 +22,20 @@ public:
 	{		
 	}
 
-	/** @return the filter weight given a 2D distance vector. */
+	/** @return the filter weight given a 2D distance vector, in Texel units. */
 	float GetWeight(const FVector2d& Dist) const
 	{
-		// [-Radius, Radius)
-		return Dist.X >= -Radius && Dist.X < Radius && Dist.Y >= -Radius && Dist.Y < Radius; 
+		// Returns 1 if Dist is within the region [-Radius, Radius]x[-Radius, Radius] and 0 otherwise.
+		//
+		// Note: Including the entire boundary of the box region is motivated by considering the filter
+		// weight at the center pixel of a 3x3 image with 1 sample per texel. First, a consisten policy
+		// of including/excluding the entire boundary preserves mirror symmetry under texture filtering
+		// (the sample points in the corner texels all have the same contribution to the center texel),
+		// and second, by making this policy inclusive means that when the Radius is exactly 1 the
+		// center texel weight changes compared to the weight for 0 < Radius < 1, which seems more
+		// intuitive for users.
+		//
+		return Dist.X >= -Radius && Dist.X <= Radius && Dist.Y >= -Radius && Dist.Y <= Radius; 
 	}
 };
 
diff --git a/Engine/Source/Runtime/GeometryCore/Public/Image/ImageDimensions.h b/Engine/Source/Runtime/GeometryCore/Public/Image/ImageDimensions.h
--- a/Engine/Source/Runtime/GeometryCore/Public/Image/ImageDimensions.h
+++ b/Engine/Source/Runtime/GeometryCore/Public/Image/ImageDimensions.h
@@ -113,6 +113,15 @@ public:
 		return GetTexelUV(GetCoords(LinearIndex));
 	}
 
+	/** @return The distance, measured in texels, between points P and Q which should be given in the default UV-space [0..1]^2 */
+	FVector2d GetTexelDistance(const FVector2d& P, const FVector2d Q) const
+	{
+		FVector2d TexelDistance = P - Q;
+		TexelDistance.X *= Width;
+		TexelDistance.Y *= Height;
+		return TexelDistance;
+	}
+
 	/** @return integer XY coordinates for real-valued XY coordinates (ie texel that contains value, if texel origin is in bottom-left */
 	FVector2i PixelToCoords(const FVector2d& PixelPosition) const
 	{
