Fix C++20 compile error. Make TPolygon2::SegmentIterator's comparison operators const to address ambiguity issues.

diff --git a/Engine/Source/Runtime/GeometryCore/Public/Polygon2.h b/Engine/Source/Runtime/GeometryCore/Public/Polygon2.h
--- a/Engine/Source/Runtime/GeometryCore/Public/Polygon2.h
+++ b/Engine/Source/Runtime/GeometryCore/Public/Polygon2.h
@@ -227,8 +227,8 @@ public:
 			i++;
 			return copy;
 		}
-		inline bool operator==(const SegmentIterator & i2) { return i2.polygon == polygon && i2.i == i; }
-		inline bool operator!=(const SegmentIterator & i2) { return i2.polygon != polygon || i2.i != i; }
+		inline bool operator==(const SegmentIterator & i2) const { return i2.polygon == polygon && i2.i == i; }
+		inline bool operator!=(const SegmentIterator & i2) const { return i2.polygon != polygon || i2.i != i; }
 	protected:
 		const TPolygon2 * polygon;
 		int i;
