Fix an incorrect filter kernel and remove the corresponding test which was causing the Geometry Processing Tests target to fail

diff --git a/Engine/Source/Runtime/GeometryCore/Public/Image/BCSplineFilter.h b/Engine/Source/Runtime/GeometryCore/Public/Image/BCSplineFilter.h
--- a/Engine/Source/Runtime/GeometryCore/Public/Image/BCSplineFilter.h
+++ b/Engine/Source/Runtime/GeometryCore/Public/Image/BCSplineFilter.h
@@ -21,13 +21,16 @@ enum class EBCSplineType
  * @param SplineType The BC-spline filter type (predefined B/C parameters)
  * @param bIsRadial When true, use Euclidean distance to compute filter weights: Weight = Filter(sqrt(X*X + Y*Y)).
  *                  Otherwise use linear 1D distance in X/Y separately: WX = Filter(X), WY = Filter(Y), Weight = WX*WY 
+ * 
+ * TODO The Unreal naming convention uses the T prefix for template functions. This would require deprecation.
+ * TODO Consider removing the Abs from GetWeight since distances are non-negative. This would require deprecation.
+ * TODO Consider using constexpr polynomial coefficients to make this class smaller
  */
 template <EBCSplineType SplineType, bool bIsRadial>
 class FBCSplineFilter
 {
 private:
-	// Measured in Texel units i.e., Radius of 1 <=> Radius of 1 texel side length
-	float Radius = 1.0f;
+	float InvRadius = 1.0f;
 
 	float Coeff1[3] = {0.0f, 0.0f, 0.0f};
 	float Coeff2[4] = {0.0f, 0.0f, 0.0f, 0.0f};
@@ -35,25 +38,27 @@ private:
 public:
 	FBCSplineFilter(const float RadiusIn);
 
-	/** @return the filter weight given a 2D distance vector, in Texel units. */
+	/** @return The filter weight given a 2D distance vector
+	 *  @note   Works on general 2D point sets i.e., does not assume Cartesian/pixel grids
+	 */
 	float GetWeight(const FVector2d& Dist) const;
 
-	/** @return true if the given 2D distance vector, in Texel units, is in the region where the filter is defined and false otherwise */
+	/** @return true if the given 2D distance vector is in the region where the filter is defined and false otherwise */
 	bool IsInFilterRegion(const FVector2d& Dist) const;
 };
 
 template<EBCSplineType SplineType, bool bIsRadial>
 FBCSplineFilter<SplineType, bIsRadial>::FBCSplineFilter(const float RadiusIn)
-	: Radius(RadiusIn)
+	: InvRadius(1.f/RadiusIn)
 {
 	auto ComputeCoeffs = [this](const float& B, const float& C)
 	{
-		// |x| < 1:        1/6 * ((12 - 9B - 6C)x^3 + (-18 + 12B + 6C)x^2 + (6 - 2B))
+		// x < 1:        1/6 * ((12 - 9B - 6C)x^3 + (-18 + 12B + 6C)x^2 + (6 - 2B))
 		Coeff1[0] = (12.0f - 9.0f * B - 6.0f * C) / 6.0f;
 		Coeff1[1] = (-18.0f + 12.0f * B + 6.0f * C) / 6.0f;
 		Coeff1[2] = (6.0f - 2.0f * B) / 6.0f;
 
-		// 1 <= |x| < 2:   1/6 * ((-B - 6C)x^3 + (6B + 30C)x^2 + (-12B - 48C)x + (8B + 24C))
+		// 1 <= x < 2:   1/6 * ((-B - 6C)x^3 + (6B + 30C)x^2 + (-12B - 48C)x + (8B + 24C))
 		Coeff2[0] = (-B - 6.0f * C) / 6.0f;
 		Coeff2[1] = (6.0f * B + 30.0f * C) / 6.0f;
 		Coeff2[2] = (-12.0f * B - 48.0f * C) / 6.0f;
@@ -86,7 +91,7 @@ float FBCSplineFilter<SplineType, bIsRadial>::GetWeight(const FVector2d& Dist) c
 	auto ComputeWeight = [this](float X) -> float
 	{
 		// Scale the filter kernel f(X) in the X-direction by making the substitution f(X/Radius)
-		const float AbsX = FMathf::Abs(X / Radius);
+		const float AbsX = FMathf::Abs(X * InvRadius);
 		const float AbsX2 = AbsX * AbsX;
 		const float AbsX3 = AbsX * AbsX2;
 
@@ -119,9 +124,9 @@ bool FBCSplineFilter<SplineType, bIsRadial>::IsInFilterRegion(const FVector2d& D
 {
 	if constexpr(bIsRadial)
 	{
-		return (FMathf::Abs((float)Dist.Length() / Radius) < 2.0f);
+		return (FMathf::Abs((float)Dist.Length() * InvRadius) < 2.0f);
 	}
-	return (FMathf::Abs((float)Dist.X / Radius) < 2.0f) && (FMathf::Abs((float)Dist.Y / Radius) < 2.0f);
+	return (FMathf::Abs((float)Dist.X * InvRadius) < 2.0f) && (FMathf::Abs((float)Dist.Y * InvRadius) < 2.0f);
 }
 
 using FBSplineFilter = FBCSplineFilter<EBCSplineType::BSpline, /*bIsRadial*/ false>;
diff --git a/Engine/Source/Runtime/GeometryCore/Public/Image/BoxFilter.h b/Engine/Source/Runtime/GeometryCore/Public/Image/BoxFilter.h
--- a/Engine/Source/Runtime/GeometryCore/Public/Image/BoxFilter.h
+++ b/Engine/Source/Runtime/GeometryCore/Public/Image/BoxFilter.h
@@ -13,7 +13,6 @@ namespace Geometry
 class FBoxFilter
 {
 private:
-	// Measured in Texel units i.e., Radius of 1 <=> Radius of 1 texel side length
 	float Radius = 0.5f;
 	
 public:
@@ -22,26 +21,19 @@ public:
 	{		
 	}
 
-	/** @return the filter weight given a 2D distance vector, in Texel units. */
+	/** @return The filter weight given a 2D distance vector
+	 *  @note   Works on general 2D point sets i.e., does not assume Cartesian/pixel grids
+	 */
 	float GetWeight(const FVector2d& Dist) const
 	{
-		// Returns 1 if Dist is within the region [-Radius, Radius]x[-Radius, Radius] and 0 otherwise.
-		//
-		// Note: Including the entire boundary of the box region is motivated by considering the filter
-		// weight at the center pixel of a 3x3 image with 1 sample per texel. First, a consisten policy
-		// of including/excluding the entire boundary preserves mirror symmetry under texture filtering
-		// (the sample points in the corner texels all have the same contribution to the center texel),
-		// and second, by making this policy inclusive means that when the Radius is exactly 1 the
-		// center texel weight changes compared to the weight for 0 < Radius < 1, which seems more
-		// intuitive for users.
-		//
-		return Dist.X >= -Radius && Dist.X <= Radius && Dist.Y >= -Radius && Dist.Y <= Radius;
+		// Returns 1 if Dist is within the region [-Radius, Radius)x[-Radius, Radius) and 0 otherwise.
+		return -Radius <= Dist.X && Dist.X < Radius && -Radius <= Dist.Y && Dist.Y < Radius;
 	}
 
-	/** @return true if the given 2D distance vector, in Texel units, is in the region where the filter is defined and false otherwise */
+	/** @return true if the given 2D distance vector is in the region where the filter is defined and false otherwise */
 	bool IsInFilterRegion(const FVector2d& Dist) const
 	{
-		return Dist.X >= -Radius && Dist.X <= Radius && Dist.Y >= -Radius && Dist.Y <= Radius; 
+		return -Radius <= Dist.X && Dist.X < Radius && -Radius <= Dist.Y && Dist.Y < Radius;
 	}
 };
 
