make fbox/fbox2D conversion to/from geometrycore AABB types correctly handle the FBox IsValid flag

diff --git a/Engine/Source/Runtime/GeometryCore/Public/BoxTypes.h b/Engine/Source/Runtime/GeometryCore/Public/BoxTypes.h
--- a/Engine/Source/Runtime/GeometryCore/Public/BoxTypes.h
+++ b/Engine/Source/Runtime/GeometryCore/Public/BoxTypes.h
@@ -331,12 +331,21 @@ struct TAxisAlignedBox3
 
 	explicit operator FBox() const
 	{
-		return FBox((FVector)Min, (FVector)Max);
+		FBox ToRet((FVector)Min, (FVector)Max);
+		ToRet.IsValid = !IsEmpty();
+		return ToRet;
 	}
 	TAxisAlignedBox3(const FBox& Box)
 	{
-		Min = TVector<RealType>(Box.Min);
-		Max = TVector<RealType>(Box.Max);
+		if (Box.IsValid)
+		{
+			Min = TVector<RealType>(Box.Min);
+			Max = TVector<RealType>(Box.Max);
+		}
+		else
+		{
+			*this = Empty();
+		}
 	}
 
 	/**
@@ -644,12 +653,21 @@ struct TAxisAlignedBox2
 
 	explicit operator FBox2D() const
 	{
-		return FBox2D((FVector2D)Min, (FVector2D)Max);
+		FBox2D ToRet((FVector2D)Min, (FVector2D)Max);
+		ToRet.bIsValid = !IsEmpty();
+		return ToRet;
 	}
 	TAxisAlignedBox2(const FBox2D& Box)
 	{
-		Min = Box.Min;
-		Max = Box.Max;
+		if (Box.bIsValid)
+		{
+			Min = Box.Min;
+			Max = Box.Max;
+		}
+		else
+		{
+			*this = Empty();
+		}
 	}
 
 	static TAxisAlignedBox2<RealType> Empty()
