GeometryCore: Fix potential buffer overrun in FSmallListSet

diff --git a/Engine/Source/Runtime/GeometryCore/Private/Util/SmallListSet.cpp b/Engine/Source/Runtime/GeometryCore/Private/Util/SmallListSet.cpp
--- a/Engine/Source/Runtime/GeometryCore/Private/Util/SmallListSet.cpp
+++ b/Engine/Source/Runtime/GeometryCore/Private/Util/SmallListSet.cpp
@@ -468,7 +468,7 @@ void FSmallListSet::Serialize(FArchive& Ar, bool bCompactData, bool bUseCompress
 				const int32 ListCount = Buffer[BufferIndex++];
 				Resize(ListCount);
 
-				for (size_t ListIndex = 0; ListIndex < ListCount; ++ListIndex)
+				for (int32 ListIndex = 0; ListIndex < ListCount; ++ListIndex)
 				{
 					const int32 ListValueCount = Buffer[BufferIndex++];
 					if (ListValueCount > 0)
@@ -495,7 +495,7 @@ void FSmallListSet::Serialize(FArchive& Ar, bool bCompactData, bool bUseCompress
 			}
 			else
 			{
-				const size_t ListCount = ListHeads.Num();
+				const int32 ListCount = ListHeads.Num();
 				size_t BufferSize = 0;
 				size_t BufferIndex = 0;
 
@@ -503,7 +503,7 @@ void FSmallListSet::Serialize(FArchive& Ar, bool bCompactData, bool bUseCompress
 				Buffer.Resize(++BufferSize);
 				Buffer[BufferIndex++] = ListCount;
 
-				for (size_t ListIndex = 0; ListIndex < ListCount; ++ListIndex)
+				for (int32 ListIndex = 0; ListIndex < ListCount; ++ListIndex)
 				{
 					const int32 ListValueCount = GetCount(ListIndex);
 
