Fixed mismatched copy constructor/assignment warning in TQuaternion.

diff --git a/Engine/Source/Runtime/GeometryCore/Public/Quaternion.h b/Engine/Source/Runtime/GeometryCore/Public/Quaternion.h
--- a/Engine/Source/Runtime/GeometryCore/Public/Quaternion.h
+++ b/Engine/Source/Runtime/GeometryCore/Public/Quaternion.h
@@ -38,7 +38,10 @@ struct TQuaternion
 	TQuaternion();
 	TQuaternion(RealType X, RealType Y, RealType Z, RealType W);
 	explicit TQuaternion(const RealType* Values);
-	TQuaternion(const TQuaternion& Copy);
+	PRAGMA_DISABLE_DEPRECATION_WARNINGS
+	TQuaternion(const TQuaternion& Copy) = default;
+	TQuaternion& operator=(const TQuaternion& Copy) = default;
+	PRAGMA_ENABLE_DEPRECATION_WARNINGS
 	template<typename RealType2>
 	explicit TQuaternion(const TQuaternion<RealType2>& Copy);
 	TQuaternion(const TVector<RealType>& Axis, RealType Angle, bool bAngleIsDegrees);
@@ -190,15 +193,6 @@ TQuaternion<RealType>::TQuaternion(const TQuaternion<RealType2>& Copy)
 	W = (RealType)Copy.W;
 }
 
-template<typename RealType>
-TQuaternion<RealType>::TQuaternion(const TQuaternion& Copy)
-{
-	X = Copy.X;
-	Y = Copy.Y;
-	Z = Copy.Z;
-	W = Copy.W;
-}
-
 template<typename RealType>
 TQuaternion<RealType>::TQuaternion(const TVector<RealType>& Axis, RealType Angle, bool bAngleIsDegrees)
 {
