Fixed mismatched copy constructor.assignment operator in TAxisAlignedBox2 and TAxisAlignedBox3.

diff --git a/Engine/Source/Runtime/GeometryCore/Public/BoxTypes.h b/Engine/Source/Runtime/GeometryCore/Public/BoxTypes.h
--- a/Engine/Source/Runtime/GeometryCore/Public/BoxTypes.h
+++ b/Engine/Source/Runtime/GeometryCore/Public/BoxTypes.h
@@ -262,8 +262,6 @@ struct TAxisAlignedBox3
 			TMathUtil<RealType>::Max3(A.Z, B.Z, C.Z));
 	}
 
-	TAxisAlignedBox3(const TAxisAlignedBox3& OtherBox) = default;
-
 	template<typename OtherRealType>
 	explicit TAxisAlignedBox3(const TAxisAlignedBox3<OtherRealType>& OtherBox)
 	{
@@ -615,8 +613,6 @@ struct TAxisAlignedBox2
 	{
 	}
 
-	TAxisAlignedBox2(const TAxisAlignedBox2& OtherBox) = default;
-
 	template<typename OtherRealType>
 	explicit TAxisAlignedBox2(const TAxisAlignedBox2<OtherRealType>& OtherBox)
 	{
