Fixed mismatched copy constructor/assignment operator in TQuadricError, TVolPresQuadricError and TAttrBasedQuadricError.

diff --git a/Engine/Source/Runtime/GeometryCore/Public/QuadricError.h b/Engine/Source/Runtime/GeometryCore/Public/QuadricError.h
--- a/Engine/Source/Runtime/GeometryCore/Public/QuadricError.h
+++ b/Engine/Source/Runtime/GeometryCore/Public/QuadricError.h
@@ -125,22 +125,6 @@ struct TQuadricError
 		Subtract(b);
 	}
 
-	TQuadricError& operator=(const TQuadricError& b)
-	{
-		Axx = b.Axx;
-		Axy = b.Axy;
-		Axz = b.Axz;
-		Ayy = b.Ayy;
-		Ayz = b.Ayz;
-		Azz = b.Azz;
-		bx = b.bx;
-		by = b.by;
-		bz = b.bz;
-		c = b.c;
-
-		return *this;
-	}
-
 	/** 
 	* Scale this quadric by the weight 'w'
 	*/
@@ -374,13 +358,6 @@ public:
 	}
 
 
-	TVolPresQuadricError& operator=(const TVolPresQuadricError& other)
-	{
-		BaseStruct::operator=(other);
-		PlaneData = other.PlaneData;
-		return *this;
-	}
-
 	/**
 	* The optimal point minimizing the quadric error with respect to a volume conserving constraint
 	*/
@@ -576,24 +553,6 @@ public:
 	}
 
 
-	TAttrBasedQuadricError& operator=(const TAttrBasedQuadricError& other)
-	{
-		BaseClass::operator=(other);
-
-		a = other.a;
-		attrweight = other.attrweight;
-
-		grad[0] = other.grad[0];
-		grad[1] = other.grad[1];
-		grad[2] = other.grad[2];
-
-		d[0] = other.d[0];
-		d[1] = other.d[1];
-		d[2] = other.d[2];
-
-		return *this;
-	}
-
 	/**
 	* The optimal point minimizing the quadric error with respect to a volume conserving constraint
 	*/
