GeometryCore: add ability to store and fetch 'debug' FDynamicMesh3 via UE::Geometry::Debug::StashDebugMesh() and FetchDebugMesh(). CVar geometry.DynamicMesh.EnableDebugMeshes controls this support, defaults to disabled. Console command geometry.DynamicMesh.ClearDebugMeshes added to allow clearing entire debug mesh set.
GeometryScript: expose above capability to Geometry Scripting

diff --git a/Engine/Plugins/Experimental/GeometryScripting/Source/GeometryScriptingEditor/Private/EditorDynamicMeshUtilityFunctions.cpp b/Engine/Plugins/Experimental/GeometryScripting/Source/GeometryScriptingEditor/Private/EditorDynamicMeshUtilityFunctions.cpp
--- a/Engine/Plugins/Experimental/GeometryScripting/Source/GeometryScriptingEditor/Private/EditorDynamicMeshUtilityFunctions.cpp
+++ b/Engine/Plugins/Experimental/GeometryScripting/Source/GeometryScriptingEditor/Private/EditorDynamicMeshUtilityFunctions.cpp
@@ -86,3 +86,30 @@ UDynamicMesh* UGeometryScriptLibrary_EditorDynamicMeshFunctions::EmitTrackedMesh
 	return TargetMesh;
 }
 
+
+
+UDynamicMesh* UGeometryScriptLibrary_EditorDynamicMeshFunctions::StashDebugMesh(UDynamicMesh* TargetMesh, FString DebugMeshName)
+{
+	TargetMesh->ProcessMesh([&](const FDynamicMesh3& ReadMesh)
+	{
+		UE::Geometry::Debug::StashDebugMesh(ReadMesh, DebugMeshName);
+	});
+	return TargetMesh;
+}
+
+
+UDynamicMesh* UGeometryScriptLibrary_EditorDynamicMeshFunctions::FetchDebugMesh(FString DebugMeshName, UDynamicMesh* ToTargetMesh, bool bClearDebugMesh, bool& bDebugMeshExists)
+{
+	bDebugMeshExists = false;
+	FDynamicMesh3 ResultMesh;
+	if (UE::Geometry::Debug::FetchDebugMesh(DebugMeshName, ResultMesh, bClearDebugMesh))
+	{
+		ToTargetMesh->EditMesh([&](FDynamicMesh3& EditMesh)
+		{
+			EditMesh = MoveTemp(ResultMesh);
+		});
+		bDebugMeshExists = true;
+	}
+
+	return ToTargetMesh;
+}
\ No newline at end of file
diff --git a/Engine/Plugins/Experimental/GeometryScripting/Source/GeometryScriptingEditor/Public/GeometryScript/EditorDynamicMeshUtilityFunctions.h b/Engine/Plugins/Experimental/GeometryScripting/Source/GeometryScriptingEditor/Public/GeometryScript/EditorDynamicMeshUtilityFunctions.h
--- a/Engine/Plugins/Experimental/GeometryScripting/Source/GeometryScriptingEditor/Public/GeometryScript/EditorDynamicMeshUtilityFunctions.h
+++ b/Engine/Plugins/Experimental/GeometryScripting/Source/GeometryScriptingEditor/Public/GeometryScript/EditorDynamicMeshUtilityFunctions.h
@@ -62,4 +62,24 @@ public:
 	static UPARAM(DisplayName = "Dynamic Mesh") UDynamicMesh* 
 	EmitTrackedMeshChange(UDynamicMesh* TargetMesh, UPARAM(ref) FDynamicMeshChangeContainer& ChangeTracker);
 
+
+	/**
+	 * Store a copy of TargetMesh with name DebugMeshName.
+	 * The mesh can later be recovered via FetchDebugMesh.
+	 * @warning This function stores the mesh in a global data structure, the caller must take care to avoid storing large numbers of debug meshes
+	 */
+	UFUNCTION(BlueprintCallable, Category = "DynamicMesh|Changes")
+	static UPARAM(DisplayName = "Dynamic Mesh") UDynamicMesh* 
+	StashDebugMesh(UDynamicMesh* TargetMesh, FString DebugMeshName);
+
+
+	/**
+	 * Fetch a debug FDynamicMesh3 saved with DebugMeshName from the global debug mesh storage and copy to ToTargetMesh.
+	 * If DebugMeshName does not exist, a cube will be returned.
+	 * @param bClearDebugMesh if true, debug mesh will be removed from global storage
+	 * @param bDebugMeshExists will return as true if DebugMeshName existed
+	 */
+	UFUNCTION(BlueprintCallable, Category = "DynamicMesh|Changes")
+	static UPARAM(DisplayName = "Dynamic Mesh") UDynamicMesh* 
+	FetchDebugMesh(FString DebugMeshName, UDynamicMesh* ToTargetMesh, bool bClearDebugMesh, bool& bDebugMeshExists);
 };
\ No newline at end of file
diff --git a/Engine/Source/Runtime/GeometryCore/Private/DynamicMesh/DynamicMesh3.cpp b/Engine/Source/Runtime/GeometryCore/Private/DynamicMesh/DynamicMesh3.cpp
--- a/Engine/Source/Runtime/GeometryCore/Private/DynamicMesh/DynamicMesh3.cpp
+++ b/Engine/Source/Runtime/GeometryCore/Private/DynamicMesh/DynamicMesh3.cpp
@@ -4,6 +4,7 @@
 #include "DynamicMesh/DynamicMeshAttributeSet.h"
 #include "Generators/MeshShapeGenerator.h"
 #include "Templates/UniquePtr.h"
+#include "HAL/IConsoleManager.h"	// required for cvars
 
 using namespace UE::Geometry;
 
@@ -1345,3 +1346,85 @@ int FDynamicMesh3::FindEdgeFromTriPair(int TriA, int TriB) const
 }
 
 
+
+
+
+static TAutoConsoleVariable<bool> CVarDynamicMeshDebugMeshesEnabled(
+	TEXT("geometry.DynamicMesh.EnableDebugMeshes"),
+	false,
+	TEXT("Enable/Disable FDynamicMesh3 Global Debug Mesh support. Debug Mesh support is only available in the Editor."));
+
+static FAutoConsoleCommand DynamicMeshClearDebugMeshesCmd(
+	TEXT("geometry.DynamicMesh.ClearDebugMeshes"),
+	TEXT("Discard all debug meshes currently stored in the FDynamicMesh3 Global Debug Mesh set. This command only works in the Editor."),
+	FConsoleCommandDelegate::CreateStatic(UE::Geometry::Debug::ClearAllDebugMeshes) );
+
+
+namespace UELocal
+{
+#if WITH_EDITOR
+	TMap<FString, TUniquePtr<FDynamicMesh3>> GlobalDebugMeshes;
+#endif
+}
+
+void UE::Geometry::Debug::ClearAllDebugMeshes()
+{
+#if WITH_EDITOR
+	if (CVarDynamicMeshDebugMeshesEnabled.GetValueOnAnyThread() == false )
+	{
+		UE_LOG(LogGeometry, Warning, TEXT("ClearAllDebugMeshes() called but geometry.DynamicMesh.EnableDebugMeshes CVar is disabled"));
+		return;
+	}
+
+	UELocal::GlobalDebugMeshes.Reset();
+#else
+	UE_LOG(LogGeometry, Warning, TEXT("DynamicMesh3 Global Debug Mesh support is only available in-Editor"));
+#endif
+}
+
+void UE::Geometry::Debug::StashDebugMesh(const FDynamicMesh3& Mesh, FString DebugMeshName)
+{
+#if WITH_EDITOR
+	if (CVarDynamicMeshDebugMeshesEnabled.GetValueOnAnyThread() == false )
+	{
+		UE_LOG(LogGeometry, Warning, TEXT("StashDebugMesh() called but geometry.DynamicMesh.EnableDebugMeshes CVar is disabled"))
+		return;
+	}
+
+	UELocal::GlobalDebugMeshes.Add( DebugMeshName, MakeUnique<FDynamicMesh3>(Mesh) );
+#else
+	UE_LOG(LogGeometry, Warning, TEXT("DynamicMesh3 Global Debug Mesh support is only available in-Editor"));
+#endif
+}
+
+bool UE::Geometry::Debug::FetchDebugMesh(FString DebugMeshName, FDynamicMesh3& MeshOut, bool bClear)
+{
+#if WITH_EDITOR
+	if (CVarDynamicMeshDebugMeshesEnabled.GetValueOnAnyThread() == false )
+	{
+		UE_LOG(LogGeometry, Warning, TEXT("FetchDebugMesh() called but geometry.DynamicMesh.EnableDebugMeshes CVar is disabled"));
+		return false;
+	}
+
+	TUniquePtr<FDynamicMesh3>* FoundMesh = UELocal::GlobalDebugMeshes.Find(DebugMeshName);
+	if (FoundMesh == nullptr)
+	{
+		return false;
+	}
+
+	if (bClear)
+	{
+		MeshOut = MoveTemp(**FoundMesh);
+		UELocal::GlobalDebugMeshes.Remove(DebugMeshName);
+	}
+	else
+	{
+		MeshOut = **FoundMesh;
+	}
+	return true;
+#else
+	UE_LOG(LogGeometry, Warning, TEXT("DynamicMesh3 Global Debug Mesh support is only available in-Editor"));
+	return false;
+#endif
+}
+
diff --git a/Engine/Source/Runtime/GeometryCore/Public/DynamicMesh/DynamicMesh3.h b/Engine/Source/Runtime/GeometryCore/Public/DynamicMesh/DynamicMesh3.h
--- a/Engine/Source/Runtime/GeometryCore/Public/DynamicMesh/DynamicMesh3.h
+++ b/Engine/Source/Runtime/GeometryCore/Public/DynamicMesh/DynamicMesh3.h
@@ -1538,6 +1538,29 @@ protected:
 
 
 
+
+namespace Debug
+{
+	/*
+	 * Store a copy of Mesh in a global Debug Meshes TMap<DebugMeshName, Mesh>.
+	 * The mesh will exist until the Engine exits, unless FetchDebugMesh or ClearAllDebugMeshes removes it
+	 */
+	GEOMETRYCORE_API void StashDebugMesh(const FDynamicMesh3& Mesh, FString DebugMeshName);
+
+	/**
+	 * If a mesh was previously stored via StashDebugMesh with the given DebugMeshName, 
+	 * copy it to MeshOut, and optionally remove it from the Debug Meshes map
+	 */
+	GEOMETRYCORE_API bool FetchDebugMesh(FString DebugMeshName, FDynamicMesh3& MeshOut, bool bClear);
+
+	/**
+	 * Discard all debug meshes stored in the global Debug Meshes map
+	 */
+	GEOMETRYCORE_API void ClearAllDebugMeshes();
+}
+
+
+
 } // end namespace UE::Geometry
 } // end namespace UE
 
