Add center of mass to GeometryCore MeshQueries and Geometry Script Mesh Query Functions

diff --git a/Engine/Plugins/Experimental/GeometryScripting/Source/GeometryScriptingCore/Private/MeshQueryFunctions.cpp b/Engine/Plugins/Experimental/GeometryScripting/Source/GeometryScriptingCore/Private/MeshQueryFunctions.cpp
--- a/Engine/Plugins/Experimental/GeometryScripting/Source/GeometryScriptingCore/Private/MeshQueryFunctions.cpp
+++ b/Engine/Plugins/Experimental/GeometryScripting/Source/GeometryScriptingCore/Private/MeshQueryFunctions.cpp
@@ -91,6 +91,20 @@ void UGeometryScriptLibrary_MeshQueryFunctions::GetMeshVolumeArea( UDynamicMesh*
 	});
 }
 
+void UGeometryScriptLibrary_MeshQueryFunctions::GetMeshVolumeAreaCenter(UDynamicMesh* TargetMesh, float& SurfaceArea, float& Volume, FVector& CenterOfMass)
+{
+	SurfaceArea = Volume = 0;
+	CenterOfMass = FVector::ZeroVector;
+	SimpleMeshQuery<bool>(TargetMesh, false, [&](const FDynamicMesh3& Mesh) {
+		FVector3d CoM;
+		FVector2d VolArea = TMeshQueries<FDynamicMesh3>::GetVolumeAreaCenter(Mesh, CoM);
+		Volume = VolArea.X;
+		SurfaceArea = VolArea.Y;
+		CenterOfMass = (FVector)CoM;
+		return true;
+	});
+}
+
 bool UGeometryScriptLibrary_MeshQueryFunctions::GetIsClosedMesh( UDynamicMesh* TargetMesh )
 {
 	return SimpleMeshQuery<bool>(TargetMesh, true, [&](const FDynamicMesh3& Mesh) { return Mesh.IsClosed(); });
diff --git a/Engine/Plugins/Experimental/GeometryScripting/Source/GeometryScriptingCore/Public/GeometryScript/MeshQueryFunctions.h b/Engine/Plugins/Experimental/GeometryScripting/Source/GeometryScriptingCore/Public/GeometryScript/MeshQueryFunctions.h
--- a/Engine/Plugins/Experimental/GeometryScripting/Source/GeometryScriptingCore/Public/GeometryScript/MeshQueryFunctions.h
+++ b/Engine/Plugins/Experimental/GeometryScripting/Source/GeometryScriptingCore/Public/GeometryScript/MeshQueryFunctions.h
@@ -47,6 +47,12 @@ public:
 	UFUNCTION(BlueprintPure, Category = "GeometryScript|MeshQueries", meta=(ScriptMethod))
 	static void GetMeshVolumeArea( UDynamicMesh* TargetMesh, float& SurfaceArea, float& Volume );
 
+	/**
+	 * Computes the volume, area and center-of-mass of the mesh.
+	 */
+	UFUNCTION(BlueprintCallable, Category = "GeometryScript|MeshQueries", meta = (ScriptMethod))
+	static void GetMeshVolumeAreaCenter(UDynamicMesh* TargetMesh, float& SurfaceArea, float& Volume, FVector& CenterOfMass);
+
 	/**
 	* Returns true if the mesh is closed, such as no topological boundary edges.
 	*/
diff --git a/Engine/Source/Runtime/GeometryCore/Public/MeshQueries.h b/Engine/Source/Runtime/GeometryCore/Public/MeshQueries.h
--- a/Engine/Source/Runtime/GeometryCore/Public/MeshQueries.h
+++ b/Engine/Source/Runtime/GeometryCore/Public/MeshQueries.h
@@ -160,6 +160,62 @@ public:
 		return FVector2d(Volume * (1.0/6.0), Area * .5f);
 	}
 
+	static FVector2d GetVolumeAreaCenter(const TriangleMeshType& Mesh, FVector3d& OutCenterOfMass)
+	{
+		double Volume = 0.0;
+		double Area = 0;
+		OutCenterOfMass = FVector3d::ZeroVector;
+
+		// Compute quantities relative to the first vertex for more stable computation
+		FVector3d RefVert = FVector3d::ZeroVector;
+		for (int VertIdx = 0; VertIdx < Mesh.MaxVertexID(); ++VertIdx)
+		{
+			if (Mesh.IsVertex(VertIdx))
+			{
+				RefVert = Mesh.GetVertex(VertIdx);
+				break;
+			}
+		}
+		for (int TriIdx = 0; TriIdx < Mesh.MaxTriangleID(); TriIdx++)
+		{
+			if (!Mesh.IsTriangle(TriIdx))
+			{
+				continue;
+			}
+
+			FVector3d V0, V1, V2;
+			Mesh.GetTriVertices(TriIdx, V0, V1, V2);
+			// Subtract reference vertex for stability
+			V0 -= RefVert;
+			V1 -= RefVert;
+			V2 -= RefVert;
+
+			// Get cross product of edges and (un-normalized) normal vector.
+			FVector3d V1mV0 = V1 - V0;
+			FVector3d V2mV0 = V2 - V0;
+			FVector3d N = V2mV0.Cross(V1mV0);
+			Area += N.Length();
+
+			FVector3d F1 = V0 + V1 + V2;
+			FVector3d F2(
+				V0.X * V0.X + V1.X * (V0.X + V1.X) + V2.X * F1.X,
+				V0.Y * V0.Y + V1.Y * (V0.Y + V1.Y) + V2.Y * F1.Y,
+				V0.Z * V0.Z + V1.Z * (V0.Z + V1.Z) + V2.Z * F1.Z);
+
+			Volume += N.X * F1.X;
+			OutCenterOfMass += N * F2;
+		}
+
+		if (Volume != 0.0)
+		{
+			OutCenterOfMass /= (Volume * 4.0);
+		}
+
+		OutCenterOfMass += RefVert;
+
+		return FVector2d(Volume * (1.0 / 6.0), Area * .5);
+	}
+
 
 	static FVector2d GetVolumeArea(const TriangleMeshType& Mesh, const TArray<int>& TriIndices)
 	{
