Fix up error handling for Geometry Script polygon triangulation function
+ Add a generalized winding triangulation mode for Delaunay2 as a useful fallback for inputs where some edges failed to insert + make the blueprint general polygon triangulation use it as a fallback

diff --git a/Engine/Plugins/Experimental/GeometryScripting/Source/GeometryScriptingCore/Private/MeshPrimitiveFunctions.cpp b/Engine/Plugins/Experimental/GeometryScripting/Source/GeometryScriptingCore/Private/MeshPrimitiveFunctions.cpp
--- a/Engine/Plugins/Experimental/GeometryScripting/Source/GeometryScriptingCore/Private/MeshPrimitiveFunctions.cpp
+++ b/Engine/Plugins/Experimental/GeometryScripting/Source/GeometryScriptingCore/Private/MeshPrimitiveFunctions.cpp
@@ -1318,6 +1318,8 @@ UDynamicMesh* UGeometryScriptLibrary_MeshPrimitiveFunctions::AppendPolygonListTr
 	FGeometryScriptPrimitiveOptions PrimitiveOptions,
 	FTransform Transform,
 	FGeometryScriptGeneralPolygonList PolygonList,
+	FGeometryScriptPolygonsTriangulationOptions TriangulationOptions,
+	bool& bTriangulationError,
 	UGeometryScriptDebug* Debug)
 {
 	if (TargetMesh == nullptr)
@@ -1336,14 +1338,22 @@ UDynamicMesh* UGeometryScriptLibrary_MeshPrimitiveFunctions::AppendPolygonListTr
 	TriangulationMeshGen.bReverseOrientation = true; // because FDelaunay2 generates triangles with reversed orientation vs what we want
 
 	// Triangulate and append each polygon separately
-	bool bAnyFailed = false;
+	bTriangulationError = false;
 	for (const UE::Geometry::FGeneralPolygon2d& Polygon : *PolygonList.Polygons)
 	{
 		UE::Geometry::FDelaunay2 Delaunay;
 		Delaunay.bAutomaticallyFixEdgesToDuplicateVertices = true;
 		TArray<FIndex3i> Triangles;
 		TArray<FVector2d> Vertices;
-		Delaunay.Triangulate(Polygon, &Triangles, &Vertices);
+		const bool bSuccess = Delaunay.Triangulate(Polygon, &Triangles, &Vertices, true);
+		if (!bSuccess)
+		{
+			bTriangulationError = true;
+			if (!TriangulationOptions.bStillAppendOnTriangulationError)
+			{
+				return TargetMesh;
+			}
+		}
 		int32 VertStart = TriangulationMeshGen.Vertices2D.Num();
 		TriangulationMeshGen.Vertices2D.Append(Vertices);
 		TriangulationMeshGen.Triangles2D.Reserve(TriangulationMeshGen.Triangles2D.Num() + Triangles.Num());
@@ -1356,10 +1366,6 @@ UDynamicMesh* UGeometryScriptLibrary_MeshPrimitiveFunctions::AppendPolygonListTr
 	{
 		AppendPrimitive(TargetMesh, &TriangulationMeshGen.Generate(), Transform, PrimitiveOptions);
 	}
-	if (bAnyFailed)
-	{
-		UE::Geometry::AppendWarning(Debug, EGeometryScriptErrorType::OperationFailed, LOCTEXT("PrimitiveFunctions_AppendPolygonListTriangulation_TriangulationFailed", "AppendPolygonListTriangulation: Triangulation failed, but may still contain a partial result."));
-	}
 
 	return TargetMesh;
 }
diff --git a/Engine/Plugins/Experimental/GeometryScripting/Source/GeometryScriptingCore/Public/GeometryScript/MeshPrimitiveFunctions.h b/Engine/Plugins/Experimental/GeometryScripting/Source/GeometryScriptingCore/Public/GeometryScript/MeshPrimitiveFunctions.h
--- a/Engine/Plugins/Experimental/GeometryScripting/Source/GeometryScriptingCore/Public/GeometryScript/MeshPrimitiveFunctions.h
+++ b/Engine/Plugins/Experimental/GeometryScripting/Source/GeometryScriptingCore/Public/GeometryScript/MeshPrimitiveFunctions.h
@@ -135,6 +135,19 @@ public:
 };
 
 
+USTRUCT(BlueprintType)
+struct GEOMETRYSCRIPTINGCORE_API FGeometryScriptPolygonsTriangulationOptions
+{
+	GENERATED_BODY()
+public:
+
+	/** Whether to still append the triangulation in error cases -- typically, cases where the input contained intersecting edges. Resulting triangulation likely will appear correct except at the intersecting edges. */
+	UPROPERTY(EditAnywhere, BlueprintReadWrite, Category = Options)
+	bool bStillAppendOnTriangulationError = true;
+
+};
+
+
 UCLASS(meta = (ScriptName = "GeometryScript_Primitives"))
 class GEOMETRYSCRIPTINGCORE_API UGeometryScriptLibrary_MeshPrimitiveFunctions : public UBlueprintFunctionLibrary
 {
@@ -530,6 +543,8 @@ public:
 
 	/**
 	* Generates a Delaunay Triangulation of the provided Polygon List, and appends it to the Target Mesh.
+	* @param bTriangulationError Reports whether the triangulation contains errors, typically due to intersecting edges in the input. Consider pre-processing the PolygonsList with PolygonsUnion to resolve intersections and prevent this error.
+	* @param bStillAppendOnError Whether to still append a best-effort triangulation in error cases. Often this will be a triangulation that does not quite match the polygon shape near intersecting edges in the input, but otherwise is as-expected.
 	*/
 	UFUNCTION(BlueprintCallable, Category = "GeometryScript|Primitives")
 	static UPARAM(DisplayName = "Target Mesh") UDynamicMesh*
@@ -538,6 +553,8 @@ public:
 		FGeometryScriptPrimitiveOptions PrimitiveOptions,
 		FTransform Transform,
 		FGeometryScriptGeneralPolygonList PolygonList,
+		FGeometryScriptPolygonsTriangulationOptions TriangulationOptions,
+		bool& bTriangulationError,
 		UGeometryScriptDebug* Debug = nullptr);
 
 
diff --git a/Engine/Source/Runtime/GeometryCore/Private/CompGeom/Delaunay2.cpp b/Engine/Source/Runtime/GeometryCore/Private/CompGeom/Delaunay2.cpp
--- a/Engine/Source/Runtime/GeometryCore/Private/CompGeom/Delaunay2.cpp
+++ b/Engine/Source/Runtime/GeometryCore/Private/CompGeom/Delaunay2.cpp
@@ -4,6 +4,9 @@
 
 #include "CompGeom/ExactPredicates.h"
 #include "Spatial/ZOrderCurvePoints.h"
+#include "MathUtil.h"
+
+#include "Algo/RemoveIf.h"
 
 #include "Async/ParallelFor.h"
 
@@ -334,6 +337,51 @@ struct FDelaunay2Connectivity
 		return bWellDefinedResult;
 	}
 
+	template<typename RealType>
+	bool GetFilledTrianglesGeneralizedWinding(TArray<FIndex3i>& TrianglesOut, TArrayView<const TVector2<RealType>> Vertices, TArrayView<const FIndex2i> BoundaryEdges, FDelaunay2::EFillMode FillMode) const
+	{
+		if (FillMode == FDelaunay2::EFillMode::Solid)
+		{
+			UE_LOG(LogGeometry, Warning, TEXT("Generalized Winding-based fill does not support the Solid fill mode -- instead using the exact fill path."));
+			GetFilledTriangles(TrianglesOut, BoundaryEdges, FillMode);
+			return false;
+		}
+		
+		TrianglesOut = GetTriangles();
+		TArray<bool> FillTri;
+		FillTri.SetNumUninitialized(TrianglesOut.Num());
+		ParallelFor(TrianglesOut.Num(), [&](int32 TriIdx)
+			{
+				const FIndex3i Tri = TrianglesOut[TriIdx];
+				const TVector2<RealType> Centroid = (Vertices[Tri.A] + Vertices[Tri.B] + Vertices[Tri.C]) / (RealType)3;
+				RealType ApproxWinding = (RealType)0;
+				for (const FIndex2i& Edge : BoundaryEdges)
+				{
+					TVector2<RealType> A = Vertices[Edge.A] - Centroid;
+					TVector2<RealType> B = Vertices[Edge.B] - Centroid;
+					ApproxWinding += TMathUtil<RealType>::Atan2(A.X * B.Y - A.Y * B.X, A.X * B.X + A.Y * B.Y);
+				}
+				const int32 Winding = FMath::RoundToInt32(ApproxWinding * TMathUtil<RealType>::InvTwoPi);
+				FillTri[TriIdx] = (
+					(FillMode == FDelaunay2::EFillMode::NegativeWinding && Winding < 0) ||
+					(FillMode == FDelaunay2::EFillMode::NonZeroWinding && Winding != 0) ||
+					(FillMode == FDelaunay2::EFillMode::PositiveWinding && Winding > 0) ||
+					(FillMode == FDelaunay2::EFillMode::OddWinding && (Winding % 2) != 0)
+				);
+			});
+
+		for (int32 TriIdx = 0; TriIdx < TrianglesOut.Num(); ++TriIdx)
+		{
+			if (!FillTri[TriIdx])
+			{
+				FillTri.RemoveAtSwap(TriIdx, 1, false);
+				TrianglesOut.RemoveAtSwap(TriIdx, 1, false);
+				TriIdx--; // re-consider the index w/ the newly swapped element
+			}
+		}
+		return true;
+	}
+
 	bool GetFilledTriangles(TArray<FIndex3i>& TrianglesOut, TArrayView<const FIndex2i> BoundaryEdges, TArrayView<const FIndex2i> HoleEdges) const
 	{
 		TSet<FIndex2i> BoundarySet, HoleSet;
@@ -1505,6 +1553,24 @@ bool FDelaunay2::GetFilledTriangles(TArray<FIndex3i>& TrianglesOut, TArrayView<c
 	return Connectivity->GetFilledTriangles(TrianglesOut, BoundaryEdges, HoleEdges);
 }
 
+bool FDelaunay2::GetFilledTrianglesGeneralizedWinding(TArray<FIndex3i>& TrianglesOut, TArrayView<const TVector2<double>> Vertices, TArrayView<const FIndex2i> Edges, EFillMode FillMode) const
+{
+	if (!ensure(Connectivity.IsValid()))
+	{
+		return false;
+	}
+	return Connectivity->GetFilledTrianglesGeneralizedWinding<double>(TrianglesOut, Vertices, Edges, FillMode);
+}
+
+bool FDelaunay2::GetFilledTrianglesGeneralizedWinding(TArray<FIndex3i>& TrianglesOut, TArrayView<const TVector2<float>> Vertices, TArrayView<const FIndex2i> Edges, EFillMode FillMode) const
+{
+	if (!ensure(Connectivity.IsValid()))
+	{
+		return false;
+	}
+	return Connectivity->GetFilledTrianglesGeneralizedWinding<float>(TrianglesOut, Vertices, Edges, FillMode);
+}
+
 bool FDelaunay2::IsDelaunay(TArrayView<const FVector2f> Vertices, TArrayView<const FIndex2i> SkipEdges) const
 {
 	if (ensure(Connectivity.IsValid()))
diff --git a/Engine/Source/Runtime/GeometryCore/Public/CompGeom/Delaunay2.h b/Engine/Source/Runtime/GeometryCore/Public/CompGeom/Delaunay2.h
--- a/Engine/Source/Runtime/GeometryCore/Public/CompGeom/Delaunay2.h
+++ b/Engine/Source/Runtime/GeometryCore/Public/CompGeom/Delaunay2.h
@@ -109,7 +109,7 @@ public:
 	// Note: TrianglesOut may be empty or incomplete if the input is self-intersecting.
 	// Note this clears any previously-held triangulation data, and triangulates the passed-in general polygon from scratch
 	template<typename RealType>
-	bool Triangulate(const TGeneralPolygon2<RealType>& GeneralPolygon, TArray<FIndex3i>* TrianglesOut = nullptr, TArray<TVector2<RealType>>* VerticesOut = nullptr)
+	bool Triangulate(const TGeneralPolygon2<RealType>& GeneralPolygon, TArray<FIndex3i>* TrianglesOut = nullptr, TArray<TVector2<RealType>>* VerticesOut = nullptr, bool bFallbackToGeneralizedWinding = false)
 	{
 		TArray<TVector2<RealType>> AllVertices;
 		TArray<FIndex2i> AllEdges;
@@ -139,7 +139,15 @@ public:
 		bool bSuccess = Triangulate(AllVertices, AllEdges);
 		if (TrianglesOut)
 		{
-			*TrianglesOut = GetFilledTriangles(AllEdges, GeneralPolygon.OuterIsClockwise() ? EFillMode::NegativeWinding : EFillMode::PositiveWinding);
+			EFillMode FillMode = GeneralPolygon.OuterIsClockwise() ? EFillMode::NegativeWinding : EFillMode::PositiveWinding;
+			if (bSuccess || !bFallbackToGeneralizedWinding)
+			{
+				*TrianglesOut = GetFilledTriangles(AllEdges, FillMode);
+			}
+			else
+			{
+				GetFilledTrianglesGeneralizedWinding(*TrianglesOut, AllVertices, AllEdges, FillMode);
+			}
 		}
 		if (VerticesOut)
 		{
@@ -196,6 +204,23 @@ public:
 	 */
 	bool GetFilledTriangles(TArray<FIndex3i>& TrianglesOut, TArrayView<const FIndex2i> Edges, EFillMode FillMode = EFillMode::PositiveWinding) const;
 
+	/**
+	 * Get (by reference) the triangles that are inside the given edges, using a generalized winding number method to determine which triangles are inside
+	 * Not valid for EFillMode::Solid, will fall back to the above GetFilledTriangles method in that case.
+	 * 
+	 * This method is for triangulations where the edges did not define a fully-closed shape, either by construction or due to un-resolved edge intersections in the input.
+	 * Note that it is slower than the standard GetFilledTriangles, and is worse for closed polygons where all edges were successfully constrained.
+	 *
+	 * @param TrianglesOut				Will be filled with a subset of the triangulation that is 'inside' the given edges, as defined by the FillMode
+	 * @param Vertices					The array of vertices that were used in the triangulation (by the previous call to Triangulate)
+	 * @param Edges						The array of edges that were already constrained in the triangulation (by a previous call to Triangulate or ConstrainEdges)
+	 * @param FillMode					Strategy to use to define which triangles to include in the output
+	 * @return							true if the result was well defined and consistent; false otherwise. Solid fill mode always has a well defined result; Winding-number-based fills do not if the edges have open spans.
+	 *									Note the triangulation will still be filled by best-effort even if the function returns false.
+	 */
+	bool GetFilledTrianglesGeneralizedWinding(TArray<FIndex3i>& TrianglesOut, TArrayView<const TVector2<double>> Vertices, TArrayView<const FIndex2i> Edges, EFillMode FillMode = EFillMode::PositiveWinding) const;
+	bool GetFilledTrianglesGeneralizedWinding(TArray<FIndex3i>& TrianglesOut, TArrayView<const TVector2<float>> Vertices, TArrayView<const FIndex2i> Edges, EFillMode FillMode = EFillMode::PositiveWinding) const;
+
 	/**
 	 * Get (by reference) the triangles that are inside the given edges, removing the outside-boundary triangles and the inside-hole triangles
 	 * 
