Prevent VoxBlend and VoxWrap adding vertices that are outside the spatial bounds as seeds, which happens for meshes that have isolated vertices that are unused by triangles.
 + do not checkSlow() on an out-of-bounds seed, as it is not really a problem (the point will just be skipped over)

diff --git a/Engine/Source/Runtime/GeometryCore/Public/Generators/MarchingCubes.h b/Engine/Source/Runtime/GeometryCore/Public/Generators/MarchingCubes.h
--- a/Engine/Source/Runtime/GeometryCore/Public/Generators/MarchingCubes.h
+++ b/Engine/Source/Runtime/GeometryCore/Public/Generators/MarchingCubes.h
@@ -267,7 +267,6 @@ protected:
 	}
 	FVector3i cell_index(const TVector<double>& Pos)
 	{
-		checkSlow(Bounds.Contains(Pos));
 		return FVector3i(
 			(int)((Pos.X - Bounds.Min.X) / CubeSize),
 			(int)((Pos.Y - Bounds.Min.Y) / CubeSize),
diff --git a/Engine/Source/Runtime/GeometryCore/Public/Implicit/Blend.h b/Engine/Source/Runtime/GeometryCore/Public/Implicit/Blend.h
--- a/Engine/Source/Runtime/GeometryCore/Public/Implicit/Blend.h
+++ b/Engine/Source/Runtime/GeometryCore/Public/Implicit/Blend.h
@@ -231,7 +231,12 @@ protected:
 					continue;
 				}
 
-				Seeds.Add(Source->GetVertex(VID));
+				FVector3d Seed = Source->GetVertex(VID);
+				// Only add vertices that are inside the spatial bounds (only vertices that are not on any triangles will be outside)
+				if (MarchingCubes.Bounds.Contains(Seed))
+				{
+					Seeds.Add(Seed);
+				}
 			}
 		}
 
diff --git a/Engine/Source/Runtime/GeometryCore/Public/Implicit/Solidify.h b/Engine/Source/Runtime/GeometryCore/Public/Implicit/Solidify.h
--- a/Engine/Source/Runtime/GeometryCore/Public/Implicit/Solidify.h
+++ b/Engine/Source/Runtime/GeometryCore/Public/Implicit/Solidify.h
@@ -134,7 +134,12 @@ public:
 		TArray<FVector3d> MCSeeds;
 		for ( int32 VertIdx : Source->VertexIndicesItr() )
 		{
-			MCSeeds.Add(Source->GetVertex(VertIdx));
+			FVector3d Vertex = Source->GetVertex(VertIdx);
+			// Only add vertices that are inside the spatial bounds (only vertices that are not on any triangles will be outside)
+			if (MarchingCubes.Bounds.Contains(Vertex))
+			{
+				MCSeeds.Add(Vertex);
+			}
 		}
 		MarchingCubes.GenerateContinuation(MCSeeds);
 
