Add normal transforms and plane transforms to geometry plane types and vector util (correct normal transform is already in TTransformSRT3, but this lets us more easily do correct normal transforms without converting away from standard FTransform)

diff --git a/Engine/Source/Runtime/GeometryCore/Public/PlaneTypes.h b/Engine/Source/Runtime/GeometryCore/Public/PlaneTypes.h
--- a/Engine/Source/Runtime/GeometryCore/Public/PlaneTypes.h
+++ b/Engine/Source/Runtime/GeometryCore/Public/PlaneTypes.h
@@ -7,6 +7,7 @@
 #include "VectorTypes.h"
 #include "VectorUtil.h"
 #include "Math/Plane.h"
+#include "Math/Transform.h"
 
 namespace UE
 {
@@ -53,6 +54,29 @@ struct TPlane3
 		return FPlane(Normal.X, Normal.Y, Normal.Z, Constant);
 	}
 
+	/**
+	 * Transform the plane by the given transform
+	 */
+	void Transform(const TTransform<RealType>& Tr)
+	{
+		TVector<RealType> TransformedOrigin = Tr.TransformPosition(Normal * Constant);
+		// Transform the normal
+		Normal = VectorUtil::TransformNormal(Tr, Normal);
+		// Update the plane constant using the transformed normal and origin
+		Constant = Normal.Dot(TransformedOrigin);
+	}
+
+	/**
+	 * Transform the plane by the inverse of the given transform
+	 */
+	void InverseTransform(const TTransform<RealType>& Tr)
+	{
+		TVector<RealType> TransformedOrigin = Tr.InverseTransformPosition(Normal * Constant);
+		// Inverse transform the normal
+		Normal = VectorUtil::InverseTransformNormal(Tr, Normal);
+		// Update the plane constant using the transformed normal and origin
+		Constant = Normal.Dot(TransformedOrigin);
+	}
 
 	/**
 	 * Compute d = Dot(N,P)-c where N is the plane normal and c is the plane
diff --git a/Engine/Source/Runtime/GeometryCore/Public/VectorUtil.h b/Engine/Source/Runtime/GeometryCore/Public/VectorUtil.h
--- a/Engine/Source/Runtime/GeometryCore/Public/VectorUtil.h
+++ b/Engine/Source/Runtime/GeometryCore/Public/VectorUtil.h
@@ -4,7 +4,7 @@
 
 #include "MathUtil.h"
 #include "VectorTypes.h"
-
+#include "Math/Transform.h"
 
 enum class EIntersectionResult
 {
@@ -528,6 +528,23 @@ namespace VectorUtil
 	}
 
 
+	template<typename RealType>
+	inline TVector<RealType> TransformNormal(const TTransform<RealType>& Transform, const TVector<RealType>& Normal)
+	{
+		const TVector<RealType>& S = Transform.GetScale3D();
+		RealType DetSign = FMathd::SignNonZero(S.X * S.Y * S.Z); // we only need to multiply by the sign of the determinant, rather than divide by it, since we normalize later anyway
+		TVector<RealType> SafeInvS(S.Y * S.Z * DetSign, S.X * S.Z * DetSign, S.X * S.Y * DetSign);
+		return Transform.TransformVectorNoScale(Normalized(SafeInvS * Normal));
+	}
+
+
+	template<typename RealType>
+	inline TVector<RealType> InverseTransformNormal(const TTransform<RealType>& Transform, const TVector<RealType>& Normal)
+	{
+		return Normalized(Transform.GetScale3D() * Transform.InverseTransformVectorNoScale(Normal));
+	}
+
+
 }; // namespace VectorUtil
 
 
