add FSphereCover as an optional volume representation to protect negative space in convex decomposition's merge algorithm
+ expose this to dataflow
+ make cluster-convex-from-children algorithm use depth-first node ordering, to get more consistent results

diff --git a/Engine/Plugins/Experimental/Fracture/Source/FractureEngine/Private/FractureEngineConvex.cpp b/Engine/Plugins/Experimental/Fracture/Source/FractureEngine/Private/FractureEngineConvex.cpp
--- a/Engine/Plugins/Experimental/Fracture/Source/FractureEngine/Private/FractureEngineConvex.cpp
+++ b/Engine/Plugins/Experimental/Fracture/Source/FractureEngine/Private/FractureEngineConvex.cpp
@@ -4,10 +4,56 @@
 #include "Chaos/Convex.h"
 #include "DynamicMesh/DynamicMesh3.h"
 #include "DynamicMesh/DynamicMeshAABBTree3.h"
+#include "Spatial/FastWinding.h"
 #include "ProjectionTargets.h"
 #include "MeshSimplification.h"
 #include "GeometryCollection/GeometryCollection.h"
 #include "GeometryCollection/GeometryCollectionConvexUtility.h"
+#include "GeometryCollection/Facades/CollectionTransformFacade.h"
+#include "GeometryCollection/Facades/CollectionTransformSelectionFacade.h"
+
+namespace
+{
+	// Local helpers for converting convex hulls to dynamic meshes, used to run geometry processing tasks on convex hulls (e.g., simplification, computing negative space)
+
+	static void AppendConvexHullToCompactDynamicMesh(const ::Chaos::FConvex* InConvexHull, UE::Geometry::FDynamicMesh3& Mesh, FTransform* OptionalTransform = nullptr)
+	{
+		check(Mesh.IsCompact());
+
+		const ::Chaos::FConvexStructureData& ConvexStructure = InConvexHull->GetStructureData();
+		const int32 NumV = InConvexHull->NumVertices();
+		const int32 NumP = InConvexHull->NumPlanes();
+		int32 StartV = Mesh.MaxVertexID();
+		for (int32 VIdx = 0; VIdx < NumV; ++VIdx)
+		{
+			FVector3d V = (FVector3d)InConvexHull->GetVertex(VIdx);
+			if (OptionalTransform)
+			{
+				V = OptionalTransform->TransformPosition(V);
+			}
+			int32 MeshVIdx = Mesh.AppendVertex(V);
+			checkSlow(MeshVIdx == VIdx + StartV); // Must be true because the mesh is compact
+		}
+		for (int32 PIdx = 0; PIdx < NumP; ++PIdx)
+		{
+			const int32 NumFaceV = ConvexStructure.NumPlaneVertices(PIdx);
+			const int32 V0 = StartV + ConvexStructure.GetPlaneVertex(PIdx, 0);
+			for (int32 SubIdx = 1; SubIdx + 1 < NumFaceV; ++SubIdx)
+			{
+				const int32 V1 = StartV + ConvexStructure.GetPlaneVertex(PIdx, SubIdx);
+				const int32 V2 = StartV + ConvexStructure.GetPlaneVertex(PIdx, SubIdx + 1);
+				Mesh.AppendTriangle(UE::Geometry::FIndex3i(V0, V1, V2));
+			}
+		}
+	}
+
+	static UE::Geometry::FDynamicMesh3 ConvexHullToDynamicMesh(const ::Chaos::FConvex* InConvexHull)
+	{
+		UE::Geometry::FDynamicMesh3 Mesh;
+		AppendConvexHullToCompactDynamicMesh(InConvexHull, Mesh);
+		return Mesh;
+	}
+}
 
 namespace UE::FractureEngine::Convex
 {
@@ -68,7 +114,6 @@ namespace UE::FractureEngine::Convex
 		}
 
 		const ::Chaos::FConvexStructureData& ConvexStructure = InConvexHull->GetStructureData();
-		const int32 NumV = InConvexHull->NumVertices();
 		const int32 NumP = InConvexHull->NumPlanes();
 
 		// Check if no simplification required, and skip simplification in that case
@@ -88,23 +133,7 @@ namespace UE::FractureEngine::Convex
 		}
 	
 		// Convert to DynamicMesh to run simplifier
-		UE::Geometry::FDynamicMesh3 Mesh;
-		for (int32 VIdx = 0; VIdx < NumV; ++VIdx)
-		{
-			int32 MeshVIdx = Mesh.AppendVertex((FVector3d)InConvexHull->GetVertex(VIdx));
-			checkSlow(MeshVIdx == VIdx);
-		}
-		for (int32 PIdx = 0; PIdx < NumP; ++PIdx)
-		{
-			const int32 NumFaceV = ConvexStructure.NumPlaneVertices(PIdx);
-			const int32 V0 = ConvexStructure.GetPlaneVertex(PIdx, 0);
-			for (int32 SubIdx = 1; SubIdx + 1 < NumFaceV; ++SubIdx)
-			{
-				const int32 V1 = ConvexStructure.GetPlaneVertex(PIdx, SubIdx);
-				const int32 V2 = ConvexStructure.GetPlaneVertex(PIdx, SubIdx + 1);
-				Mesh.AppendTriangle(UE::Geometry::FIndex3i(V0, V1, V2));
-			}
-		}
+		UE::Geometry::FDynamicMesh3 Mesh = ConvexHullToDynamicMesh(InConvexHull);
 		
 		// Run simplification
 		UE::Geometry::FVolPresMeshSimplification Simplifier(&Mesh);
@@ -148,4 +177,60 @@ namespace UE::FractureEngine::Convex
 
 		return true;
 	}
+
+	bool ComputeConvexHullsNegativeSpace(FManagedArrayCollection& Collection, UE::Geometry::FSphereCovering& OutNegativeSpace, const UE::Geometry::FNegativeSpaceSampleSettings& Settings, bool bRestrictToSelection, const TArrayView<const int32> TransformSelection)
+	{
+		if (!FGeometryCollectionConvexUtility::HasConvexHullData(&Collection))
+		{
+			return false;
+		}
+
+		TManagedArray<TSet<int32>>& TransformToConvexInds = Collection.ModifyAttribute<TSet<int32>>("TransformToConvexIndices", FTransformCollection::TransformGroup);
+		TManagedArray<TUniquePtr<::Chaos::FConvex>>& ConvexHulls = Collection.ModifyAttribute<TUniquePtr<::Chaos::FConvex>>("ConvexHull", "Convex");
+
+		UE::Geometry::FDynamicMesh3 CombinedMesh;
+
+		GeometryCollection::Facades::FCollectionTransformFacade TransformFacade(Collection);
+		GeometryCollection::Facades::FCollectionTransformSelectionFacade SelectionFacade(Collection);
+		TArray<int> RigidSelection;
+		if (!bRestrictToSelection)
+		{
+			RigidSelection = SelectionFacade.SelectLeaf();
+		}
+		else
+		{
+			RigidSelection.Append(TransformSelection);
+			SelectionFacade.ConvertSelectionToRigidNodes(RigidSelection);
+		}
+		
+		TArray<FTransform> GlobalTransformArray = TransformFacade.ComputeCollectionSpaceTransforms();
+		
+		auto ProcessBone = [&TransformToConvexInds, &ConvexHulls, &CombinedMesh, &GlobalTransformArray](int32 BoneIdx) -> bool
+		{
+			if (BoneIdx < 0 || BoneIdx >= TransformToConvexInds.Num())
+			{
+				// invalid bone index
+				return false;
+			}
+			bool bNoFailures = true;
+			for (int32 ConvexIdx : TransformToConvexInds[BoneIdx])
+			{
+				AppendConvexHullToCompactDynamicMesh(ConvexHulls[ConvexIdx].Get(), CombinedMesh, &GlobalTransformArray[BoneIdx]);
+			}
+			return bNoFailures;
+		};
+
+		bool bNoFailures = true;
+		for (int32 BoneIdx : RigidSelection)
+		{
+			bool bSuccess = ProcessBone(BoneIdx);
+			bNoFailures = bNoFailures && bSuccess;
+		}
+
+		UE::Geometry::FDynamicMeshAABBTree3 Tree(&CombinedMesh, true);
+		UE::Geometry::TFastWindingTree<UE::Geometry::FDynamicMesh3> Winding(&Tree, true);
+		OutNegativeSpace.AddNegativeSpace(Winding, Settings);
+
+		return bNoFailures;
+	}
 }
diff --git a/Engine/Plugins/Experimental/Fracture/Source/FractureEngine/Public/FractureEngineConvex.h b/Engine/Plugins/Experimental/Fracture/Source/FractureEngine/Public/FractureEngineConvex.h
--- a/Engine/Plugins/Experimental/Fracture/Source/FractureEngine/Public/FractureEngineConvex.h
+++ b/Engine/Plugins/Experimental/Fracture/Source/FractureEngine/Public/FractureEngineConvex.h
@@ -3,6 +3,7 @@
 #pragma once
 
 #include "Containers/ArrayView.h"
+#include "CompGeom/ConvexDecomposition3.h"
 
 namespace Chaos { class FConvex; }
 struct FManagedArrayCollection;
@@ -33,5 +34,11 @@ namespace UE::FractureEngine::Convex
 	 */
 	bool FRACTUREENGINE_API SimplifyConvexHull(const ::Chaos::FConvex* InConvexHull, ::Chaos::FConvex* OutConvexHull, const FSimplifyHullSettings& Settings);
 
+	/**
+	 * Compute negative space for the convex hulls on a geometry collection
+	 * @return true if any negative space was found
+	 */
+	bool FRACTUREENGINE_API ComputeConvexHullsNegativeSpace(FManagedArrayCollection& Collection, UE::Geometry::FSphereCovering& OutNegativeSpace, const UE::Geometry::FNegativeSpaceSampleSettings& Settings, bool bRestrictToSelection = false, const TArrayView<const int32> TransformSelection = TArrayView<const int32>());
+
 }
 
diff --git a/Engine/Plugins/Experimental/GeometryCollectionPlugin/Source/GeometryCollectionNodes/Private/Dataflow/GeometryCollectionUtilityNodes.cpp b/Engine/Plugins/Experimental/GeometryCollectionPlugin/Source/GeometryCollectionNodes/Private/Dataflow/GeometryCollectionUtilityNodes.cpp
--- a/Engine/Plugins/Experimental/GeometryCollectionPlugin/Source/GeometryCollectionNodes/Private/Dataflow/GeometryCollectionUtilityNodes.cpp
+++ b/Engine/Plugins/Experimental/GeometryCollectionPlugin/Source/GeometryCollectionNodes/Private/Dataflow/GeometryCollectionUtilityNodes.cpp
@@ -155,6 +155,11 @@ FGenerateClusterConvexHullsFromLeafHullsDataflowNode::FGenerateClusterConvexHull
 	RegisterInputConnection(&ConvexCount);
 	RegisterInputConnection(&ErrorTolerance);
 	RegisterInputConnection(&OptionalSelectionFilter);
+	RegisterInputConnection(&bProtectNegativeSpace);
+	RegisterInputConnection(&TargetNumSamples);
+	RegisterInputConnection(&MinSampleSpacing);
+	RegisterInputConnection(&NegativeSpaceTolerance);
+	RegisterInputConnection(&MinRadius);
 
 	RegisterOutputConnection(&Collection);
 }
@@ -167,26 +172,45 @@ void FGenerateClusterConvexHullsFromLeafHullsDataflowNode::Evaluate(Dataflow::FC
 
 		if (TUniquePtr<FGeometryCollection> GeomCollection = TUniquePtr<FGeometryCollection>(InCollection.NewCopy<FGeometryCollection>()))
 		{
+			TArray<int32> SelectionArray;
+			bool bHasSelectionFilter = IsConnected(&OptionalSelectionFilter);
+			if (bHasSelectionFilter)
+			{
+				const FDataflowTransformSelection& InOptionalSelectionFilter = GetValue<FDataflowTransformSelection>(Context, &OptionalSelectionFilter);
+				SelectionArray = InOptionalSelectionFilter.AsArray();
+			}
+
+			bool bHasNegativeSpace = false;
+			UE::Geometry::FSphereCovering NegativeSpace;
+			if (GetValue(Context, &bProtectNegativeSpace))
+			{
+				UE::Geometry::FNegativeSpaceSampleSettings NegativeSpaceSettings;
+				NegativeSpaceSettings.TargetNumSamples = GetValue(Context, &TargetNumSamples);
+				NegativeSpaceSettings.MinRadius = GetValue(Context, &MinRadius);
+				NegativeSpaceSettings.ReduceRadiusMargin = GetValue(Context, &NegativeSpaceTolerance);
+				NegativeSpaceSettings.MinSpacing = GetValue(Context, &MinSampleSpacing);
+				bHasNegativeSpace = UE::FractureEngine::Convex::ComputeConvexHullsNegativeSpace(*GeomCollection, NegativeSpace, NegativeSpaceSettings, bHasSelectionFilter, SelectionArray);
+			}
+
 			const int32 InConvexCount = GetValue(Context, &ConvexCount);
 			const double InErrorToleranceInCm = GetValue(Context, &ErrorTolerance);
-			if (IsConnected(&OptionalSelectionFilter))
+			FGeometryCollectionConvexUtility::FClusterConvexHullSettings HullMergeSettings(InConvexCount, InErrorToleranceInCm, bPreferExternalCollisionShapes);
+			HullMergeSettings.AllowMergesMethod = AllowMerges;
+			HullMergeSettings.EmptySpace = bHasNegativeSpace ? &NegativeSpace : nullptr;
+
+			if (bHasSelectionFilter)
 			{
-				const FDataflowTransformSelection& InOptionalSelectionFilter = GetValue<FDataflowTransformSelection>(Context, &OptionalSelectionFilter);
-				FGeometryCollectionConvexUtility::FClusterConvexHullSettings Settings(InConvexCount, InErrorToleranceInCm, bPreferExternalCollisionShapes);
-				Settings.AllowMergesMethod = AllowMerges;
 				FGeometryCollectionConvexUtility::GenerateClusterConvexHullsFromLeafHulls(
 					*GeomCollection,
-					Settings,
-					InOptionalSelectionFilter.AsArray()
+					HullMergeSettings,
+					SelectionArray
 				);
 			}
 			else
 			{
-				FGeometryCollectionConvexUtility::FClusterConvexHullSettings Settings(InConvexCount, InErrorToleranceInCm, bPreferExternalCollisionShapes);
-				Settings.AllowMergesMethod = AllowMerges;
 				FGeometryCollectionConvexUtility::GenerateClusterConvexHullsFromLeafHulls(
 					*GeomCollection,
-					Settings
+					HullMergeSettings
 				);
 			}
 
@@ -202,6 +226,11 @@ FGenerateClusterConvexHullsFromChildrenHullsDataflowNode::FGenerateClusterConvex
 	RegisterInputConnection(&ConvexCount);
 	RegisterInputConnection(&ErrorTolerance);
 	RegisterInputConnection(&OptionalSelectionFilter);
+	RegisterInputConnection(&bProtectNegativeSpace);
+	RegisterInputConnection(&TargetNumSamples);
+	RegisterInputConnection(&MinSampleSpacing);
+	RegisterInputConnection(&NegativeSpaceTolerance);
+	RegisterInputConnection(&MinRadius);
 
 	RegisterOutputConnection(&Collection);
 }
@@ -214,24 +243,45 @@ void FGenerateClusterConvexHullsFromChildrenHullsDataflowNode::Evaluate(Dataflow
 
 		if (TUniquePtr<FGeometryCollection> GeomCollection = TUniquePtr<FGeometryCollection>(InCollection.NewCopy<FGeometryCollection>()))
 		{
+			TArray<int32> SelectionArray;
+			bool bHasSelectionFilter = IsConnected(&OptionalSelectionFilter);
+			if (bHasSelectionFilter)
+			{
+				const FDataflowTransformSelection& InOptionalSelectionFilter = GetValue<FDataflowTransformSelection>(Context, &OptionalSelectionFilter);
+				SelectionArray = InOptionalSelectionFilter.AsArray();
+			}
+
+			bool bHasNegativeSpace = false;
+			UE::Geometry::FSphereCovering NegativeSpace;
+			if (GetValue(Context, &bProtectNegativeSpace))
+			{
+				UE::Geometry::FNegativeSpaceSampleSettings NegativeSpaceSettings;
+				NegativeSpaceSettings.TargetNumSamples = GetValue(Context, &TargetNumSamples);
+				NegativeSpaceSettings.MinRadius = GetValue(Context, &MinRadius);
+				NegativeSpaceSettings.ReduceRadiusMargin = GetValue(Context, &NegativeSpaceTolerance);
+				NegativeSpaceSettings.MinSpacing = GetValue(Context, &MinSampleSpacing);
+				bHasNegativeSpace = UE::FractureEngine::Convex::ComputeConvexHullsNegativeSpace(*GeomCollection, NegativeSpace, NegativeSpaceSettings, bHasSelectionFilter, SelectionArray);
+			}
+
 			const int32 InConvexCount = GetValue(Context, &ConvexCount);
 			const double InErrorToleranceInCm = GetValue(Context, &ErrorTolerance);
-			if (IsConnected(&OptionalSelectionFilter))
+			FGeometryCollectionConvexUtility::FClusterConvexHullSettings HullMergeSettings(InConvexCount, InErrorToleranceInCm, bPreferExternalCollisionShapes);
+			HullMergeSettings.AllowMergesMethod = EAllowConvexMergeMethod::Any; // Note: Only 'Any' is supported for this node currently
+			HullMergeSettings.EmptySpace = bHasNegativeSpace ? &NegativeSpace : nullptr;
+
+			if (bHasSelectionFilter)
 			{
-				const FDataflowTransformSelection& InOptionalSelectionFilter = GetValue<FDataflowTransformSelection>(Context, &OptionalSelectionFilter);
-				FGeometryCollectionConvexUtility::FClusterConvexHullSettings Settings(InConvexCount, InErrorToleranceInCm, bPreferExternalCollisionShapes);
 				FGeometryCollectionConvexUtility::GenerateClusterConvexHullsFromChildrenHulls(
 					*GeomCollection,
-					Settings,
-					InOptionalSelectionFilter.AsArray()
+					HullMergeSettings,
+					SelectionArray
 				);
 			}
 			else
 			{
-				FGeometryCollectionConvexUtility::FClusterConvexHullSettings Settings(InConvexCount, InErrorToleranceInCm, bPreferExternalCollisionShapes);
 				FGeometryCollectionConvexUtility::GenerateClusterConvexHullsFromChildrenHulls(
 					*GeomCollection,
-					Settings
+					HullMergeSettings
 				);
 			}
 
diff --git a/Engine/Plugins/Experimental/GeometryCollectionPlugin/Source/GeometryCollectionNodes/Public/Dataflow/GeometryCollectionUtilityNodes.h b/Engine/Plugins/Experimental/GeometryCollectionPlugin/Source/GeometryCollectionNodes/Public/Dataflow/GeometryCollectionUtilityNodes.h
--- a/Engine/Plugins/Experimental/GeometryCollectionPlugin/Source/GeometryCollectionNodes/Public/Dataflow/GeometryCollectionUtilityNodes.h
+++ b/Engine/Plugins/Experimental/GeometryCollectionPlugin/Source/GeometryCollectionNodes/Public/Dataflow/GeometryCollectionUtilityNodes.h
@@ -170,6 +170,26 @@ public:
 	UPROPERTY(meta = (DataflowInput, DataflowIntrinsic))
 	FDataflowTransformSelection OptionalSelectionFilter;
 
+	/** Whether to use a sphere cover to define negative space that should not be covered by convex hulls */
+	UPROPERTY(EditAnywhere, Category = NegativeSpace, meta = (DataflowInput))
+	bool bProtectNegativeSpace = false;
+
+	/** Approximate number of spheres to consider when covering negative space */
+	UPROPERTY(EditAnywhere, Category = NegativeSpace, meta = (DataflowInput, ClampMin = 1, EditCondition = "bProtectNegativeSpace", EditConditionHides))
+	int32 TargetNumSamples = 50;
+
+	/** Minimum desired spacing between spheres; if > 0, will attempt not to place sphere centers closer than this */
+	UPROPERTY(EditAnywhere, Category = NegativeSpace, meta = (DataflowInput, ClampMin = 0, Units = cm, EditCondition = "bProtectNegativeSpace", EditConditionHides))
+	double MinSampleSpacing = 1.0;
+
+	/** Amount of space to leave between convex hulls and protected negative space */
+	UPROPERTY(EditAnywhere, Category = NegativeSpace, meta = (DataflowInput, ClampMin = .01, UIMin = .1, Units = cm, EditCondition = "bProtectNegativeSpace", EditConditionHides))
+	double NegativeSpaceTolerance = 2.0;
+
+	/** Spheres smaller than this are not included in the negative space */
+	UPROPERTY(EditAnywhere, Category = NegativeSpace, meta = (DataflowInput, ClampMin = 0, Units = cm, EditCondition = "bProtectNegativeSpace", EditConditionHides))
+	double MinRadius = 10.0;
+
 	FGenerateClusterConvexHullsFromLeafHullsDataflowNode(const Dataflow::FNodeParameters& InParam, FGuid InGuid = FGuid::NewGuid());
 
 	virtual void Evaluate(Dataflow::FContext& Context, const FDataflowOutput* Out) const override;
@@ -211,6 +231,26 @@ public:
 	UPROPERTY(meta = (DataflowInput, DataflowIntrinsic))
 	FDataflowTransformSelection OptionalSelectionFilter;
 
+	/** Whether to use a sphere cover to define negative space that should not be covered by convex hulls */
+	UPROPERTY(EditAnywhere, Category = NegativeSpace, meta = (DataflowInput))
+	bool bProtectNegativeSpace = false;
+
+	/** Approximate number of spheres to consider when covering negative space */
+	UPROPERTY(EditAnywhere, Category = NegativeSpace, meta = (DataflowInput, ClampMin = 1, EditCondition = "bProtectNegativeSpace", EditConditionHides))
+	int32 TargetNumSamples = 50;
+
+	/** Minimum desired spacing between spheres; if > 0, will attempt not to place sphere centers closer than this */
+	UPROPERTY(EditAnywhere, Category = NegativeSpace, meta = (DataflowInput, ClampMin = 0, Units = cm, EditCondition = "bProtectNegativeSpace", EditConditionHides))
+	double MinSampleSpacing = 1.0;
+
+	/** Amount of space to leave between convex hulls and protected negative space */
+	UPROPERTY(EditAnywhere, Category = NegativeSpace, meta = (DataflowInput, ClampMin = .01, UIMin = .1, Units = cm, EditCondition = "bProtectNegativeSpace", EditConditionHides))
+	double NegativeSpaceTolerance = 2.0;
+
+	/** Spheres smaller than this are not included in the negative space */
+	UPROPERTY(EditAnywhere, Category = NegativeSpace, meta = (DataflowInput, ClampMin = 0, Units = cm, EditCondition = "bProtectNegativeSpace", EditConditionHides))
+	double MinRadius = 10.0;
+
 	FGenerateClusterConvexHullsFromChildrenHullsDataflowNode(const Dataflow::FNodeParameters& InParam, FGuid InGuid = FGuid::NewGuid());
 
 	virtual void Evaluate(Dataflow::FContext& Context, const FDataflowOutput* Out) const override;
diff --git a/Engine/Source/Runtime/Experimental/Chaos/Private/GeometryCollection/GeometryCollectionConvexUtility.cpp b/Engine/Source/Runtime/Experimental/Chaos/Private/GeometryCollection/GeometryCollectionConvexUtility.cpp
--- a/Engine/Source/Runtime/Experimental/Chaos/Private/GeometryCollection/GeometryCollectionConvexUtility.cpp
+++ b/Engine/Source/Runtime/Experimental/Chaos/Private/GeometryCollection/GeometryCollectionConvexUtility.cpp
@@ -1685,22 +1685,38 @@ void FGeometryCollectionConvexUtility::GenerateClusterConvexHullsFromLeafOrChild
 
 	const TArray<FTransform> GlobalTransforms = TransformFacade.ComputeCollectionSpaceTransforms();
 	const TManagedArrayAccessor<int32> SimulationTypeAttribute(Collection, FGeometryCollection::SimulationTypeAttribute, FTransformCollection::TransformGroup);
+	const TArray<int32> DepthFirstTransformIndices = HierarchyFacade.GetTransformArrayInDepthFirstOrder();
 	TArray<int32> TransformsToProcess;
 	if (bOnlySubset)
 	{
-		TransformsToProcess.Append(TransformSubset);
 		if (TransformsToProcess.IsEmpty())
 		{
 			return;
 		}
+
+		if (bUseDirectChildren)
+		{
+			// If using direct children, order matters -- use depth-first order
+			TSet<int32> ToProcessSet;
+			ToProcessSet.Append(TransformSubset);
+			TransformsToProcess.Reserve(TransformSubset.Num());
+			for (int32 TransformIdx : DepthFirstTransformIndices)
+			{
+				if (ToProcessSet.Contains(TransformIdx))
+				{
+					TransformsToProcess.Add(TransformIdx);
+				}
+			}
+		}
+		else
+		{
+			// Order doesn't matter
+			TransformsToProcess.Append(TransformSubset);
+		}
 	}
 	else // process all
 	{
-		TransformsToProcess.Reserve(GlobalTransforms.Num());
-		for (int32 Idx = 0; Idx < GlobalTransforms.Num(); ++Idx)
-		{
-			TransformsToProcess.Add(Idx);
-		}
+		TransformsToProcess = DepthFirstTransformIndices;
 	}
 
 	// Simulation types must be present
@@ -1715,12 +1731,11 @@ void FGeometryCollectionConvexUtility::GenerateClusterConvexHullsFromLeafOrChild
 		const bool bIsCluster = (SimulationTypeAttribute.Get()[TransformIndex] == FGeometryCollection::ESimulationTypes::FST_Clustered);
 		if (bIsCluster)
 		{
-			// compute using the leaf nodes 
-			// we cannot use the direct cluster children because they do not have proximity data that is stored at the geo level )
 			TArray<int32> SourceTransformIndices;
 			if (bUseDirectChildren)
 			{
 				SourceTransformIndices = HierarchyFacade.GetChildrenAsArray(TransformIndex);
+				// TODO: consider expanding this recursively to children for any nodes without convex hulls, until we hit a leaf/rigid node
 			}
 			else
 			{
@@ -1844,7 +1859,7 @@ void FGeometryCollectionConvexUtility::GenerateClusterConvexHullsFromLeafOrChild
 				};
 				UE::Geometry::FConvexDecomposition3 ConvexDecomposition;
 				ConvexDecomposition.InitializeFromHulls(Hulls.Num(), GetHullVolume, GetHullNumVertices, GetHullVertex, HullProximity);
-				ConvexDecomposition.MergeBest(Settings.ConvexCount, Settings.ErrorToleranceInCm, 0, true);
+				ConvexDecomposition.MergeBest(Settings.ConvexCount, Settings.ErrorToleranceInCm, 0, true /*bAllowCompact*/, false /*bRequireHullTriangles*/, Settings.EmptySpace, &ParentTransform);
 				
 				// reset existing hulls for this transform index
 				for (int32 ConvexIndex : TransformToConvexIndices[TransformIndex])
diff --git a/Engine/Source/Runtime/Experimental/Chaos/Public/GeometryCollection/GeometryCollectionConvexUtility.h b/Engine/Source/Runtime/Experimental/Chaos/Public/GeometryCollection/GeometryCollectionConvexUtility.h
--- a/Engine/Source/Runtime/Experimental/Chaos/Public/GeometryCollection/GeometryCollectionConvexUtility.h
+++ b/Engine/Source/Runtime/Experimental/Chaos/Public/GeometryCollection/GeometryCollectionConvexUtility.h
@@ -14,6 +14,11 @@ namespace Chaos
 	class FConvex;
 }
 
+namespace UE::Geometry
+{
+	class FSphereCovering;
+}
+
 UENUM()
 enum class EConvexOverlapRemoval : int32
 {
@@ -119,6 +124,7 @@ public:
 		double ErrorToleranceInCm = 0.0;
 		bool bUseExternalCollisionIfAvailable = true;
 		EAllowConvexMergeMethod AllowMergesMethod = EAllowConvexMergeMethod::ByProximity;
+		UE::Geometry::FSphereCovering* EmptySpace = nullptr;
 	};
 
 	static void GenerateClusterConvexHullsFromChildrenHulls(FGeometryCollection& Collection, const FClusterConvexHullSettings& Settings, const TArrayView<const int32> TransformSubset);
diff --git a/Engine/Source/Runtime/GeometryCore/Private/CompGeom/ConvexDecomposition3.cpp b/Engine/Source/Runtime/GeometryCore/Private/CompGeom/ConvexDecomposition3.cpp
--- a/Engine/Source/Runtime/GeometryCore/Private/CompGeom/ConvexDecomposition3.cpp
+++ b/Engine/Source/Runtime/GeometryCore/Private/CompGeom/ConvexDecomposition3.cpp
@@ -7,10 +7,13 @@
 #include "VectorUtil.h"
 #include "Spatial/PriorityOrderPoints.h"
 #include "VertexConnectedComponents.h"
+#include "Distance/DistPoint3Triangle3.h"
 
 #include "DynamicMesh/DynamicMesh3.h"
 #include "DynamicMesh/Operations/MergeCoincidentMeshEdges.h"
 #include "DynamicMesh/MeshTransforms.h"
+#include "DynamicMesh/DynamicMeshAABBTree3.h"
+#include "Spatial/FastWinding.h"
 
 #include "Async/ParallelFor.h"
 
@@ -19,7 +22,66 @@ namespace UE
 namespace Geometry
 {
 
+bool FSphereCovering::AddNegativeSpace(const TFastWindingTree<FDynamicMesh3>& Spatial, const FNegativeSpaceSampleSettings& SampleSettings)
+{
+	bool bAddedPoints = false;
+
+	check(Spatial.IsBuilt());
 
+	FAxisAlignedBox3d Bounds = Spatial.GetTree()->GetBoundingBox();
+
+	if (SampleSettings.SampleMethod == FNegativeSpaceSampleSettings::ESampleMethod::Uniform)
+	{
+		double MinCellSize = FMath::Max(SampleSettings.MinSpacing, Bounds.MaxDim() / SampleSettings.TargetNumSamples);
+		if (MinCellSize == 0)
+		{
+			return false; // no space to sample
+		}
+		FVector3d Ranges = Bounds.Diagonal();
+		FIndex3i Dims(FMath::Max(1, int(Ranges.X / MinCellSize)), FMath::Max(1, int(Ranges.Y / MinCellSize)), FMath::Max(1, int(Ranges.Z / MinCellSize)));
+		double ReduceFactor = FMath::Min(1.0, FMath::Pow(double(SampleSettings.TargetNumSamples) / double(Dims.A * Dims.B * Dims.C), 1.0/3.0));
+		Dims.A = FMath::Max(1, FMath::CeilToInt32(Dims.A * ReduceFactor));
+		Dims.B = FMath::Max(1, FMath::CeilToInt32(Dims.B * ReduceFactor));
+		Dims.C = FMath::Max(1, FMath::CeilToInt32(Dims.C * ReduceFactor));
+		int32 MaxSamples = Dims.A * Dims.B * Dims.C;
+		FVector3d Denoms(1.0 / (Dims.A + 1), 1.0 / (Dims.B + 1), 1.0 / (Dims.C + 1));
+		FVector3d StepSizes = Ranges * Denoms;
+
+		Position.Reserve(Position.Num() + MaxSamples);
+		Radius.Reserve(Radius.Num() + MaxSamples);
+		for (int32 X = 0; X < Dims.A; ++X)
+		{
+			double XPos = Bounds.Min.X + Ranges.X * (1.0 + X) * Denoms.X;
+			for (int32 Y = 0; Y < Dims.B; ++Y)
+			{
+				double YPos = Bounds.Min.Y + Ranges.Y * (1.0 + Y) * Denoms.Y;
+				for (int32 Z = 0; Z < Dims.C; ++Z)
+				{
+					double ZPos = Bounds.Min.Z + Ranges.Z * (1.0 + Z) * Denoms.Z;
+					FVector3d Pos(XPos, YPos, ZPos);
+					double Winding = Spatial.FastWindingNumber(Pos);
+					if (Winding < -.5)
+					{
+						continue;
+					}
+					double NearDistSq;
+					int NearTID = Spatial.GetTree()->FindNearestTriangle(Pos, NearDistSq);
+					double R = FMath::Sqrt(NearDistSq) - SampleSettings.ReduceRadiusMargin;
+					if (R < SampleSettings.MinRadius)
+					{
+						// point too close to the surfaces; skip it
+						continue;
+					}
+					bAddedPoints = true;
+					Position.Add(Pos);
+					Radius.Add(R);
+				}
+			}
+		}
+	}
+
+	return bAddedPoints;
+}
 
 // Internal struct to run templated TriangleMeshType functions on the hull mesh
 // Assumes the convex part is not compacted (i.e. is using the dynamic mesh vertices)
@@ -719,7 +781,7 @@ void FConvexDecomposition3::FConvexPart::Compact()
 		return;
 	}
 
-	// remove all triangles, and all verties that are not reference by the hull
+	// remove all triangles, and all vertices that are not reference by the hull
 	// and remap HullTriangles to the new compact indices
 	TMap<int32, int32> VertexRemapIDs;
 	FDynamicMesh3 VertexMesh;
@@ -1230,7 +1292,8 @@ void FConvexDecomposition3::InitializeFromHulls(int32 NumHulls, TFunctionRef<dou
 	}
 }
 
-int32 FConvexDecomposition3::MergeBest(int32 TargetNumParts, double MaxErrorTolerance, double MinThicknessToleranceWorldSpace, bool bAllowCompact, bool bRequireHullTriangles)
+int32 FConvexDecomposition3::MergeBest(int32 TargetNumParts, double MaxErrorTolerance, double MinThicknessToleranceWorldSpace, bool bAllowCompact, bool bRequireHullTriangles, 
+	const FSphereCovering* OptionalNegativeSpace, const FTransform* OptionalTransformIntoNegativeSpace)
 {
 	// Support having a max error tolerance
 	double VolumeTolerance = ConvertDistanceToleranceToLocalVolumeTolerance(MaxErrorTolerance);
@@ -1240,6 +1303,69 @@ int32 FConvexDecomposition3::MergeBest(int32 TargetNumParts, double MaxErrorTole
 
 	TMap<int32, TUniquePtr<FConvexPart>> ProximityComputedParts;
 
+	auto ConvexNegativeSpaceSphereOverlapTest = [OptionalNegativeSpace, OptionalTransformIntoNegativeSpace](FConvexPart* Part, int PtIdx)
+	{
+		check(OptionalNegativeSpace);
+		double Radius = OptionalNegativeSpace->GetRadius(PtIdx);
+		FVector3d Center = OptionalNegativeSpace->GetCenter(PtIdx);
+
+		bool bMustTransformConvex = false;
+		if (OptionalTransformIntoNegativeSpace)
+		{
+			FVector ScaleVec = OptionalTransformIntoNegativeSpace->GetScale3D();
+			// Uniform scale: Can tranform into negative space into part space
+			if (ScaleVec.AllComponentsEqual())
+			{
+				if (ScaleVec.X == 0)
+				{
+					return false;
+				}
+				Radius = FMath::Abs(Radius / ScaleVec.X);
+				Center = OptionalTransformIntoNegativeSpace->InverseTransformPosition(Center);
+			}
+			else // Non-uniform scale: Must transform convex into negative space
+			{
+				bMustTransformConvex = true;
+			}
+		}
+		
+		double RadiusSq = Radius * Radius;
+		double MaxPlaneDist = 0;
+		for (int32 PlaneIdx = 0; PlaneIdx < Part->HullPlanes.Num(); ++PlaneIdx)
+		{
+			FPlane3d Plane = Part->HullPlanes[PlaneIdx];
+			if (bMustTransformConvex)
+			{
+				Plane.Transform(*OptionalTransformIntoNegativeSpace);
+			}
+			double PlaneDist = Plane.DistanceTo(Center);
+			MaxPlaneDist = FMath::Max(MaxPlaneDist, PlaneDist);
+			if (PlaneDist > Radius)
+			{
+				return false;
+			}
+			else if (Radius > 0 && PlaneDist > FMath::Max(0, MaxPlaneDist - FMathd::ZeroTolerance))
+			{
+				FIndex3i TriInds = Part->HullTriangles[PlaneIdx];
+				FTriangle3d Tri(Part->InternalGeo.GetVertex(TriInds.A), Part->InternalGeo.GetVertex(TriInds.B), Part->InternalGeo.GetVertex(TriInds.C));
+				if (bMustTransformConvex)
+				{
+					Tri.V[0] = OptionalTransformIntoNegativeSpace->TransformPosition(Tri.V[0]);
+					Tri.V[1] = OptionalTransformIntoNegativeSpace->TransformPosition(Tri.V[1]);
+					Tri.V[2] = OptionalTransformIntoNegativeSpace->TransformPosition(Tri.V[2]);
+				}
+				// TODO: Can optimize this by writing custom Point-Tri distance logic to re-use what we already know and early-out if a vertex or edge distance is < Radius
+				FDistPoint3Triangle3d Dist(Center, Tri);
+				double TriDistSq = Dist.GetSquared();
+				if (TriDistSq < RadiusSq)
+				{
+					return true;
+				}
+			}
+		}
+		return (MaxPlaneDist <= 0);
+	};
+
 	auto IsPartBelowSizeTolerance = [MinThicknessTolerance](const FConvexPart& Part) -> bool
 	{ 
 		if (MinThicknessTolerance <= 0) // tolerance is disabled
@@ -1422,6 +1548,11 @@ int32 FConvexDecomposition3::MergeBest(int32 TargetNumParts, double MaxErrorTole
 		int32 BestKnownIdx = -1;
 		for (int32 ProxIdx = 0; ProxIdx < Proximities.Num(); ProxIdx++)
 		{
+			if (!Proximities[ProxIdx].bIsValidLink)
+			{
+				continue;
+			}
+
 			bool bIncludesMustMergePart =
 				Decomposition[Proximities[ProxIdx].Link.A].bMustMerge ||
 				Decomposition[Proximities[ProxIdx].Link.B].bMustMerge;
@@ -1457,6 +1588,31 @@ int32 FConvexDecomposition3::MergeBest(int32 TargetNumParts, double MaxErrorTole
 				MergeCost -= BiasToRemoveTooThinParts;
 			}
 
+			// if this would be the new best candidate, and we have negative space, check if the merge would cross protected negative space
+			if (OptionalNegativeSpace && OptionalNegativeSpace->Num() > 0 && MergeCost < BestKnownCost)
+			{
+				TUniquePtr<FConvexPart>* FoundMergeHull = ProximityComputedParts.Find(ProxIdx);
+				if (!FoundMergeHull)
+				{
+					FoundMergeHull = &ProximityComputedParts.Add(ProxIdx, CreateMergedPart(Proximities[ProxIdx]));
+				}
+				bool bCoversProtectedPt = false;
+				for (int32 PtIdx = 0; PtIdx < OptionalNegativeSpace->Num(); ++PtIdx)
+				{
+					if (ConvexNegativeSpaceSphereOverlapTest(FoundMergeHull->Get(), PtIdx))
+					{
+						bCoversProtectedPt = true;
+						break;
+					}
+				}
+				if (bCoversProtectedPt)
+				{
+					// Mark link as invalid so we don't repeat this overlap test in later iterations
+					Proximities[ProxIdx].bIsValidLink = false;
+					continue;
+				}
+			}
+
 			if (MergeCost < BestKnownCost)
 			{
 				BestKnownCost = MergeCost;
diff --git a/Engine/Source/Runtime/GeometryCore/Public/CompGeom/ConvexDecomposition3.h b/Engine/Source/Runtime/GeometryCore/Public/CompGeom/ConvexDecomposition3.h
--- a/Engine/Source/Runtime/GeometryCore/Public/CompGeom/ConvexDecomposition3.h
+++ b/Engine/Source/Runtime/GeometryCore/Public/CompGeom/ConvexDecomposition3.h
@@ -20,6 +20,10 @@
 namespace UE {
 namespace Geometry {
 
+// Forward declare fast winding tree for negative space sampling
+template<class TriangleMeshType>
+class TFastWindingTree;
+
 // TODO: To support meshes where volume is not well defined (e.g., open boundaries or many self-intersecting parts), we'll need alternative error metrics
 enum class EConvexErrorMethod
 {
@@ -29,6 +33,74 @@ enum class EConvexErrorMethod
 	RelativeVolume
 };
 
+struct FNegativeSpaceSampleSettings
+{
+	FNegativeSpaceSampleSettings() = default;
+	FNegativeSpaceSampleSettings(int32 TargetNumSamples, double MinSpacing, double ReduceRadiusMargin)
+		: TargetNumSamples(TargetNumSamples), MinSpacing(MinSpacing), ReduceRadiusMargin(ReduceRadiusMargin)
+	{}
+
+	enum class ESampleMethod
+	{
+		Uniform
+		// TODO: Consider e.g. sampling based on local search / floodfill
+	};
+
+	// Method used to place samples
+	ESampleMethod SampleMethod = ESampleMethod::Uniform;
+
+	// Target number of spheres to use to cover negative space. Intended as a rough guide; actual number used may be larger or smaller.
+	int32 TargetNumSamples = 30;
+
+	// Minimum desired spacing between sampling spheres; not a strictly enforced bound.
+	double MinSpacing = 3.0;
+
+	// Space to allow between spheres and actual surface
+	double ReduceRadiusMargin = 3.0;
+
+	// Ignore spheres with smaller radius than this
+	double MinRadius = 10.0;
+};
+
+// Define a volume with a set of spheres
+class GEOMETRYCORE_API FSphereCovering
+{
+public:
+
+	inline int32 Num() const
+	{
+		return Position.Num();
+	}
+
+	inline FVector3d GetCenter(int32 Idx) const
+	{
+		return Position[Idx];
+	}
+
+	inline double GetRadius(int32 Idx) const
+	{
+		return Radius[Idx];
+	}
+
+	// Add spheres covering the negative space of the given fast winding tree
+	// @return true if any spheres were added
+	bool AddNegativeSpace(const TFastWindingTree<FDynamicMesh3>& Spatial, const FNegativeSpaceSampleSettings& Settings);
+
+	void Reset()
+	{
+		Position.Reset();
+		Radius.Reset();
+	}
+
+private:
+	// Sphere centers
+	TArray<FVector3d> Position;
+	// Sphere radii
+	TArray<double> Radius;
+};
+
+
+
 class GEOMETRYCORE_API FConvexDecomposition3
 {
 
@@ -125,8 +197,11 @@ public:
 	//								Note: These parts may be further split so they can be merged into multiple hulls
 	// @param bAllowCompact			Allow the algorithm to discard underlying geometry once it will no longer be used, resulting in a smaller representation & faster merges
 	// @param bRequireHullTriangles	Require all hulls to have associated triangles after MergeBest is completed. (If using InitializeFromHulls, will need to triangulate any un-merged hulls.)
+	// @param OptionalNegativeSpace Optional representation of negative space that should not be covered by merges
+	// @param OptionalNegativeSpaceTransform Optional transform from space of the convex hulls into the space of the sphere covering; if not provided, assume spheres are in the same coordinate space as the hulls
 	// @return						The number of merges performed
-	int32 MergeBest(int32 TargetNumParts, double ErrorTolerance = 0, double MinThicknessTolerance = 0, bool bAllowCompact = true, bool bRequireHullTriangles = false);
+	int32 MergeBest(int32 TargetNumParts, double ErrorTolerance = 0, double MinThicknessTolerance = 0, bool bAllowCompact = true, bool bRequireHullTriangles = false, 
+		const FSphereCovering* OptionalNegativeSpace = nullptr, const FTransform* OptionalTransformIntoNegativeSpace = nullptr);
 
 	// simple helper to convert an error tolerance expressed in world space to a local-space volume tolerance
 	inline double ConvertDistanceToleranceToLocalVolumeTolerance(double DistTolerance) const
@@ -329,6 +404,9 @@ public:
 		
 		// Whether the geometry is separated by Plane
 		bool bPlaneSeparates = false;
+
+		// Whether this merge can be allowed (e.g. has not been ruled out by negative space)
+		bool bIsValidLink = true;
 		
 		// What the volume of the convex hull would be if we merged these two parts
 		double MergedVolume = VolumeNotComputed();
