Add Set/GetValue methods that take templated bone weight scalar data arrays.

#rnx

diff --git a/Engine/Source/Runtime/GeometryCore/Public/DynamicMesh/DynamicVertexSkinWeightsAttribute.h b/Engine/Source/Runtime/GeometryCore/Public/DynamicMesh/DynamicVertexSkinWeightsAttribute.h
--- a/Engine/Source/Runtime/GeometryCore/Public/DynamicMesh/DynamicVertexSkinWeightsAttribute.h
+++ b/Engine/Source/Runtime/GeometryCore/Public/DynamicMesh/DynamicVertexSkinWeightsAttribute.h
@@ -7,7 +7,7 @@
 
 #include "BoneWeights.h"
 #include "HAL/UnrealMemory.h"
-
+#include "Misc/MemStack.h"
 
 namespace UE
 {
@@ -308,6 +308,27 @@ public:
 		Data = VertexBoneWeights[VertexID];
 	}
 
+	/** 
+	 * Get the element at a given index via a pair of bone index and weight arrays. 
+	 * The number of influences is equal to the array size. 
+	 */
+	template<typename BoneIndexType, typename BoneFloatWeightType>
+	void GetValue(int VertexID, TArray<BoneIndexType>& OutBones, TArray<BoneFloatWeightType>& OutWeights) const
+	{
+		FBoneWeights BoneWeights;
+		GetValue(VertexID, BoneWeights);
+
+		const int32 NumEntries = BoneWeights.Num();
+
+		OutBones.SetNum(NumEntries);
+		OutWeights.SetNum(NumEntries);
+
+		for (int32 BoneIdx = 0; BoneIdx < NumEntries; ++BoneIdx)
+		{
+			OutBones[BoneIdx] = static_cast<BoneIndexType>(BoneWeights[BoneIdx].GetBoneIndex());
+			OutWeights[BoneIdx] = static_cast<BoneFloatWeightType>(BoneWeights[BoneIdx].GetWeight());
+		}
+	}
 
 	/** Set the element at a given index */
 	void SetValue(int VertexID, const FBoneWeights& Data)
@@ -322,6 +343,27 @@ public:
 		VertexBoneWeights[VertexID] = FBoneWeights::Create(Data);
 	}
 
+	/** Set the element at a given index */
+	template<typename BoneIndexType, typename BoneFloatWeightType>
+	void SetValue(int VertexID, const TArray<BoneIndexType>& InBones, const TArray<BoneFloatWeightType>& InWeights, int32 InNumEntries)
+	{
+		checkSlow(InBones.Num() == InNumEntries && InWeights.Num() == InNumEntries);
+		
+		TArray<FBoneIndexType, TMemStackAllocator<>> Bones;
+		TArray<float, TMemStackAllocator<>> Weights;
+ 
+		Bones.SetNumUninitialized(InNumEntries);
+		Weights.SetNumUninitialized(InNumEntries);
+
+		for (int32 Idx = 0; Idx < InNumEntries; ++Idx)
+		{
+			Bones[Idx] = static_cast<FBoneIndexType>(InBones[Idx]);
+			Weights[Idx] = static_cast<float>(InWeights[Idx]);
+		}
+
+		VertexBoneWeights[VertexID] = FBoneWeights::Create(Bones.GetData(), Weights.GetData(), InNumEntries);
+	}
+
 	/**
 	 * Copy the attribute value at FromVertexID to ToVertexID
 	 */
