GeometryCore: Fix potential buffer overrun in FSmallListSet serialization

diff --git a/Engine/Source/Runtime/GeometryCore/Private/Util/SmallListSet.cpp b/Engine/Source/Runtime/GeometryCore/Private/Util/SmallListSet.cpp
--- a/Engine/Source/Runtime/GeometryCore/Private/Util/SmallListSet.cpp
+++ b/Engine/Source/Runtime/GeometryCore/Private/Util/SmallListSet.cpp
@@ -464,13 +464,20 @@ void FSmallListSet::Serialize(FArchive& Ar, bool bCompactData, bool bUseCompress
 
 				SerializeVector(Ar, Buffer, bUseCompression);
 
+				// Read a value from the buffer while making sure there are no buffer overruns due to corrupted data.
+				const size_t BufferNum = Buffer.Num();
+				auto GetBufferValue = [&Buffer, &BufferNum](const size_t Index) -> int32
+				{
+					return Index < BufferNum ? Buffer[Index] : 0;
+				};
+				
 				size_t BufferIndex = 0;
-				const int32 ListCount = Buffer[BufferIndex++];
+				const int32 ListCount = GetBufferValue(BufferIndex++);
 				Resize(ListCount);
 
 				for (int32 ListIndex = 0; ListIndex < ListCount; ++ListIndex)
 				{
-					const int32 ListValueCount = Buffer[BufferIndex++];
+					const int32 ListValueCount = GetBufferValue(BufferIndex++);
 					if (ListValueCount > 0)
 					{
 						AllocateAt(ListIndex);
@@ -479,19 +486,27 @@ void FSmallListSet::Serialize(FArchive& Ar, bool bCompactData, bool bUseCompress
 						const int32 ListValueCountInBlock = FMath::Min(BLOCKSIZE, ListValueCount);
 						for (int32 ListValueIndex = 0; ListValueIndex < ListValueCountInBlock; ++ListValueIndex)
 						{
-							Insert(ListIndex, Buffer[BufferIndex + ListValueIndex]);
+							Insert(ListIndex, GetBufferValue(BufferIndex + ListValueIndex));
 						}
 
 						// Any values beyond BLOCKSIZE get inserted in reversed order. This is due to the fact that any values inserted into the linked list
 						// effectively get reversed in order, and we want to restore the original order before serializing out the data.
 						for (int32 ListValueIndex = ListValueCount - 1; ListValueIndex >= BLOCKSIZE; --ListValueIndex)
 						{
-							Insert(ListIndex, Buffer[BufferIndex + ListValueIndex]);
+							Insert(ListIndex, GetBufferValue(BufferIndex + ListValueIndex));
 						}
 
 						BufferIndex += ListValueCount;
 					}
 				}
+
+				if (BufferIndex > BufferNum)
+				{
+					UE_LOG(LogGeometry, Warning,
+					       TEXT("Encountered corrupted data when deserializing FSMallListSet; tried to read %llu values from buffer with %llu elements."),
+					       BufferIndex - 1, BufferNum);
+					Ar.SetError();
+				}
 			}
 			else
 			{
