GeometryCore: fix crash copy-paste-o in FDynamicMeshAttributeSet::EnableMatchingAttributes

diff --git a/Engine/Source/Runtime/GeometryCore/Private/DynamicMesh/DynamicMeshAttributeSet.cpp b/Engine/Source/Runtime/GeometryCore/Private/DynamicMesh/DynamicMeshAttributeSet.cpp
--- a/Engine/Source/Runtime/GeometryCore/Private/DynamicMesh/DynamicMeshAttributeSet.cpp
+++ b/Engine/Source/Runtime/GeometryCore/Private/DynamicMesh/DynamicMeshAttributeSet.cpp
@@ -291,7 +291,7 @@ void FDynamicMeshAttributeSet::EnableMatchingAttributes(const FDynamicMeshAttrib
 	}
 
 	int32 ExistingNormalLayers = NumNormalLayers();
-	int32 RequiredNormalLayers = (bClearExisting || bDiscardExtraAttributes) ? ToMatch.NumUVLayers() : FMath::Max(ExistingNormalLayers, ToMatch.NumNormalLayers());
+	int32 RequiredNormalLayers = (bClearExisting || bDiscardExtraAttributes) ? ToMatch.NumNormalLayers() : FMath::Max(ExistingNormalLayers, ToMatch.NumNormalLayers());
 	SetNumNormalLayers(RequiredNormalLayers);
 	for (int32 k = bClearExisting ? 0 : ExistingNormalLayers; k < NumNormalLayers(); k++)
 	{
