DynamicMesh: add TDynamicMeshVectorOverlay::EnumerateVertexElements, does triangle iteration around vertex and calls a function for each tri-vtx-element
#rnx

diff --git a/Engine/Source/Runtime/GeometryCore/Private/DynamicMesh/DynamicMeshOverlay.cpp b/Engine/Source/Runtime/GeometryCore/Private/DynamicMesh/DynamicMeshOverlay.cpp
--- a/Engine/Source/Runtime/GeometryCore/Private/DynamicMesh/DynamicMeshOverlay.cpp
+++ b/Engine/Source/Runtime/GeometryCore/Private/DynamicMesh/DynamicMeshOverlay.cpp
@@ -1476,6 +1476,52 @@ void TDynamicMeshOverlay<RealType, ElementSize>::SetElementFromBary(int SetEleme
 
 
 
+template<typename RealType, int ElementSize, typename VectorType>
+bool TDynamicMeshVectorOverlay<RealType, ElementSize, VectorType>::EnumerateVertexElements(
+	int VertexID,
+	TFunctionRef<bool(int TriangleID, int ElementID, const VectorType& Value)> ProcessFunc,
+	bool bFindUniqueElements)
+{
+	if (this->ParentMesh->IsVertex(VertexID) == false) return false;
+
+	TArray<int32, TInlineAllocator<16>> UniqueElements;
+
+	int32 Count = 0;
+	for (int tid : this->ParentMesh->VtxTrianglesItr(VertexID))
+	{
+		int32 BaseElemIdx = 3 * tid;
+		bool bIsSetTriangle = (this->ElementTriangles[BaseElemIdx] >= 0);
+		if (bIsSetTriangle)
+		{
+			Count++;
+
+			for (int j = 0; j < 3; ++j)
+			{
+				int32 ElementIdx = this->ElementTriangles[BaseElemIdx + j];
+				if (this->ParentVertices[ElementIdx] == VertexID)
+				{
+					int32 NumUnique = UniqueElements.Num();
+					if (bFindUniqueElements == false || (UniqueElements.AddUnique(ElementIdx) == NumUnique) )
+					{
+						bool bContinue = ProcessFunc(tid, ElementIdx, this->GetElement(ElementIdx));
+						if (!bContinue)
+						{
+							return true;
+						}
+					}
+				}
+			}
+
+		}
+	}
+	return (Count > 0);
+}
+
+
+
+
+
+
 template<typename RealType, int ElementSize>
 bool TDynamicMeshOverlay<RealType, ElementSize>::CheckValidity(bool bAllowNonManifoldVertices, EValidityCheckFailMode FailMode) const
 {
@@ -1847,14 +1893,14 @@ template class GEOMETRYCORE_API TDynamicMeshOverlay<int, 3>;
 template class GEOMETRYCORE_API TDynamicMeshOverlay<float, 4>;
 template class GEOMETRYCORE_API TDynamicMeshOverlay<double, 4>;
 
-template class TDynamicMeshVectorOverlay<float, 2, FVector2f>;
-template class TDynamicMeshVectorOverlay<double, 2, FVector2d>;
-template class TDynamicMeshVectorOverlay<int, 2, FVector2i>;
-template class TDynamicMeshVectorOverlay<float, 3, FVector3f>;
-template class TDynamicMeshVectorOverlay<double, 3, FVector3d>;
-template class TDynamicMeshVectorOverlay<int, 3, FVector3i>;
-template class TDynamicMeshVectorOverlay<float, 4, FVector4f>;
-template class TDynamicMeshVectorOverlay<double, 4, FVector4d>;
+template class GEOMETRYCORE_API TDynamicMeshVectorOverlay<float, 2, FVector2f>;
+template class GEOMETRYCORE_API TDynamicMeshVectorOverlay<double, 2, FVector2d>;
+template class GEOMETRYCORE_API TDynamicMeshVectorOverlay<int, 2, FVector2i>;
+template class GEOMETRYCORE_API TDynamicMeshVectorOverlay<float, 3, FVector3f>;
+template class GEOMETRYCORE_API TDynamicMeshVectorOverlay<double, 3, FVector3d>;
+template class GEOMETRYCORE_API TDynamicMeshVectorOverlay<int, 3, FVector3i>;
+template class GEOMETRYCORE_API TDynamicMeshVectorOverlay<float, 4, FVector4f>;
+template class GEOMETRYCORE_API TDynamicMeshVectorOverlay<double, 4, FVector4d>;
 
 } // end namespace UE::Geometry
 } // end namespace UE
\ No newline at end of file
diff --git a/Engine/Source/Runtime/GeometryCore/Public/DynamicMesh/DynamicMeshOverlay.h b/Engine/Source/Runtime/GeometryCore/Public/DynamicMesh/DynamicMeshOverlay.h
--- a/Engine/Source/Runtime/GeometryCore/Public/DynamicMesh/DynamicMeshOverlay.h
+++ b/Engine/Source/Runtime/GeometryCore/Public/DynamicMesh/DynamicMeshOverlay.h
@@ -752,6 +752,18 @@ public:
 	{
 		BaseType::SetElement(ElementID, Value);
 	}
+
+	/**
+	 * Iterate through triangles connected to VertexID and call ProcessFunc for each per-triangle-vertex Element with its Value.
+	 * ProcessFunc must return true to continue the enumeration, or false to early-terminate it
+	 * 
+	 * @param bFindUniqueElements if true, ProcessFunc is only called once for each ElementID, otherwise it is called once for each Triangle
+	 * @return true if at least one valid Element was found, ie if ProcessFunc was called at least one time
+	 */
+	bool EnumerateVertexElements(
+		int VertexID,
+		TFunctionRef<bool(int TriangleID, int ElementID, const VectorType& Value)> ProcessFunc,
+		bool bFindUniqueElements = true);
 };
 
 
