make the convex decomposition sphere covering do a better job of covering shape concavities near the border of the shape by expanding the sampling region by the sphere radius + moving the samples out to better cover the sampling region

diff --git a/Engine/Source/Runtime/GeometryCore/Private/CompGeom/ConvexDecomposition3.cpp b/Engine/Source/Runtime/GeometryCore/Private/CompGeom/ConvexDecomposition3.cpp
--- a/Engine/Source/Runtime/GeometryCore/Private/CompGeom/ConvexDecomposition3.cpp
+++ b/Engine/Source/Runtime/GeometryCore/Private/CompGeom/ConvexDecomposition3.cpp
@@ -32,6 +32,8 @@ bool FSphereCovering::AddNegativeSpace(const TFastWindingTree<FDynamicMesh3>& Sp
 
 	if (SampleSettings.SampleMethod == FNegativeSpaceSampleSettings::ESampleMethod::Uniform)
 	{
+		// Expand the sampling region by ~ the radius of interest, to make sure there is room to fit negative space spheres on concavities near the borders
+		Bounds.Expand(SampleSettings.MinRadius + SampleSettings.ReduceRadiusMargin);
 		double MinCellSize = FMath::Max(SampleSettings.MinSpacing, Bounds.MaxDim() / SampleSettings.TargetNumSamples);
 		if (MinCellSize == 0)
 		{
@@ -44,20 +46,20 @@ bool FSphereCovering::AddNegativeSpace(const TFastWindingTree<FDynamicMesh3>& Sp
 		Dims.B = FMath::Max(1, FMath::CeilToInt32(Dims.B * ReduceFactor));
 		Dims.C = FMath::Max(1, FMath::CeilToInt32(Dims.C * ReduceFactor));
 		int32 MaxSamples = Dims.A * Dims.B * Dims.C;
-		FVector3d Denoms(1.0 / (Dims.A + 1), 1.0 / (Dims.B + 1), 1.0 / (Dims.C + 1));
+		FVector3d Denoms(1.0 / double(Dims.A), 1.0 / double(Dims.B), 1.0 / double(Dims.C));
 		FVector3d StepSizes = Ranges * Denoms;
 
 		Position.Reserve(Position.Num() + MaxSamples);
 		Radius.Reserve(Radius.Num() + MaxSamples);
 		for (int32 X = 0; X < Dims.A; ++X)
 		{
-			double XPos = Bounds.Min.X + Ranges.X * (1.0 + X) * Denoms.X;
+			double XPos = Bounds.Min.X + StepSizes.X * (.5 + X);
 			for (int32 Y = 0; Y < Dims.B; ++Y)
 			{
-				double YPos = Bounds.Min.Y + Ranges.Y * (1.0 + Y) * Denoms.Y;
+				double YPos = Bounds.Min.Y + StepSizes.Y * (.5 + Y);
 				for (int32 Z = 0; Z < Dims.C; ++Z)
 				{
-					double ZPos = Bounds.Min.Z + Ranges.Z * (1.0 + Z) * Denoms.Z;
+					double ZPos = Bounds.Min.Z + StepSizes.Z * (.5 + Z);
 					FVector3d Pos(XPos, YPos, ZPos);
 					double Winding = Spatial.FastWindingNumber(Pos);
 					if (Winding < -.5)
