GeometryScript: Add DynamicMeshComponent NotifyMeshModified and NotifyMeshVertexAttributesModified functions (functions are called 'Updated' in BP, via DisplayName, because NotifyMeshUpdated already exists as C++ function). These allow a DynamicMeshComponent's rendering to be updated after the underlying DynamicMesh has been modified with deferred notifications (eg while setting many vertex positions, etc).

Add RecomputeNormalsForMeshSelection, allows subset of normals to be recomputed. Add bDeferChangeNotifications option to RecomputeNormals.

Add ComputeMeshLocalUVParam, computes expmap parameterization on bounded-geodesic-radius area of mesh

diff --git a/Engine/Plugins/Experimental/GeometryScripting/Source/GeometryScriptingCore/Private/MeshNormalsFunctions.cpp b/Engine/Plugins/Experimental/GeometryScripting/Source/GeometryScriptingCore/Private/MeshNormalsFunctions.cpp
--- a/Engine/Plugins/Experimental/GeometryScripting/Source/GeometryScriptingCore/Private/MeshNormalsFunctions.cpp
+++ b/Engine/Plugins/Experimental/GeometryScripting/Source/GeometryScriptingCore/Private/MeshNormalsFunctions.cpp
@@ -117,6 +117,7 @@ UDynamicMesh* UGeometryScriptLibrary_MeshNormalsFunctions::SetPerFaceNormals(UDy
 UDynamicMesh* UGeometryScriptLibrary_MeshNormalsFunctions::RecomputeNormals(
 	UDynamicMesh* TargetMesh,
 	FGeometryScriptCalculateNormalsOptions CalculateOptions,
+	bool bDeferChangeNotifications,
 	UGeometryScriptDebug* Debug)
 {
 	if (TargetMesh == nullptr)
@@ -136,13 +137,59 @@ UDynamicMesh* UGeometryScriptLibrary_MeshNormalsFunctions::RecomputeNormals(
 		MeshNormals.RecomputeOverlayNormals(EditMesh.Attributes()->PrimaryNormals(), CalculateOptions.bAreaWeighted, CalculateOptions.bAngleWeighted);
 		MeshNormals.CopyToOverlay(EditMesh.Attributes()->PrimaryNormals(), false);
 
-	}, EDynamicMeshChangeType::GeneralEdit, EDynamicMeshAttributeChangeFlags::Unknown, false);
+	}, EDynamicMeshChangeType::GeneralEdit, EDynamicMeshAttributeChangeFlags::Unknown, bDeferChangeNotifications);
+
+	return TargetMesh;
+}
+
+
+
+UDynamicMesh* UGeometryScriptLibrary_MeshNormalsFunctions::RecomputeNormalsForMeshSelection(
+	UDynamicMesh* TargetMesh,
+	FGeometryScriptMeshSelection Selection,
+	FGeometryScriptCalculateNormalsOptions CalculateOptions,
+	bool bDeferChangeNotifications,
+	UGeometryScriptDebug* Debug)
+{
+	if (TargetMesh == nullptr)
+	{
+		UE::Geometry::AppendError(Debug, EGeometryScriptErrorType::InvalidInputs, LOCTEXT("RecomputeNormalsForMeshSelection_InvalidInput", "RecomputeNormalsForMeshSelection: TargetMesh is Null"));
+		return TargetMesh;
+	}
+
+	// todo: publish correct change types
+	TargetMesh->EditMesh([&](FDynamicMesh3& EditMesh) 
+	{
+		if (EditMesh.HasAttributes() == false || EditMesh.Attributes()->PrimaryNormals() == nullptr)
+		{
+			UE::Geometry::AppendError(Debug, EGeometryScriptErrorType::InvalidInputs, LOCTEXT("RecomputeNormalsForMeshSelection_NoAttributes", "RecomputeNormalsForMeshSelection: TargetMesh has no Normals attribute enabled"));
+			return;
+		}
+		
+		if (Selection.GetSelectionType() == EGeometryScriptMeshSelectionType::Vertices)
+		{
+			FDynamicMeshNormalOverlay* Normals = EditMesh.Attributes()->PrimaryNormals();
+			TSet<int32> Elements;
+			Selection.ProcessByVertexID(EditMesh, [&](int32 VertexID) {
+				Normals->EnumerateVertexElements(VertexID, [&](int32 tid, int32 elemid, const FVector3f&) { Elements.Add(elemid); return true; }, /*bFindUniqueElements*/false);
+			}, /*bProcessAllVertsIfSelectionEmpty*/false);
+			FMeshNormals::RecomputeOverlayElementNormals(EditMesh, Elements.Array(), CalculateOptions.bAreaWeighted, CalculateOptions.bAngleWeighted);
+		}
+		else
+		{
+			TArray<int32> Triangles;
+			Selection.ConvertToMeshIndexArray(EditMesh, Triangles, EGeometryScriptIndexType::Triangle);
+			FMeshNormals::RecomputeOverlayTriNormals(EditMesh, Triangles, CalculateOptions.bAreaWeighted, CalculateOptions.bAngleWeighted);
+		}
+
+	}, EDynamicMeshChangeType::GeneralEdit, EDynamicMeshAttributeChangeFlags::Unknown, bDeferChangeNotifications);
 
 	return TargetMesh;
 }
 
 
 
+
 UDynamicMesh* UGeometryScriptLibrary_MeshNormalsFunctions::ComputeSplitNormals( 
 	UDynamicMesh* TargetMesh, 
 	FGeometryScriptSplitNormalsOptions SplitOptions, 
diff --git a/Engine/Plugins/Experimental/GeometryScripting/Source/GeometryScriptingCore/Private/MeshUVFunctions.cpp b/Engine/Plugins/Experimental/GeometryScripting/Source/GeometryScriptingCore/Private/MeshUVFunctions.cpp
--- a/Engine/Plugins/Experimental/GeometryScripting/Source/GeometryScriptingCore/Private/MeshUVFunctions.cpp
+++ b/Engine/Plugins/Experimental/GeometryScripting/Source/GeometryScriptingCore/Private/MeshUVFunctions.cpp
@@ -9,6 +9,8 @@
 #include "Selections/MeshConnectedComponents.h"
 #include "Parameterization/PatchBasedMeshUVGenerator.h"
 #include "DynamicMesh/MeshNormals.h"
+#include "Parameterization/MeshLocalParam.h"
+#include "TriangleTypes.h"
 #include "XAtlasWrapper.h"
 
 #include "Async/ParallelFor.h"
@@ -1084,4 +1086,78 @@ UDynamicMesh* UGeometryScriptLibrary_MeshUVFunctions::CopyMeshToMeshUVLayer(
 
 
 
+UDynamicMesh* UGeometryScriptLibrary_MeshUVFunctions::ComputeMeshLocalUVParam(
+	UDynamicMesh* TargetMesh,
+	FVector CenterPoint,
+	int32 CenterPointTriangleID,
+	TArray<int>& VertexIDs,
+	TArray<FVector2D>& VertexUVs,
+	double Radius,
+	bool bUseInterpolatedNormal,
+	FVector TangentYDirection,
+	double UVRotationDeg,
+	UGeometryScriptDebug* Debug)
+{
+	if (TargetMesh == nullptr)
+	{
+		UE::Geometry::AppendError(Debug, EGeometryScriptErrorType::InvalidInputs, LOCTEXT("ComputeMeshLocalUVParam_InvalidInput", "ComputeMeshLocalUVParam: TargetMesh is Null"));
+		return TargetMesh;
+	}
+
+	TargetMesh->ProcessMesh([&](const FDynamicMesh3& Mesh)
+	{
+		if (Mesh.IsTriangle(CenterPointTriangleID) == false)
+		{
+			UE::Geometry::AppendError(Debug, EGeometryScriptErrorType::InvalidInputs, LOCTEXT("ComputeMeshLocalUVParam_InvalidTriangle", "ComputeMeshLocalUVParam: CenterPointTriangleID is not a valid triangle"));
+			return;
+		}
+		FIndex3i TriVertices = Mesh.GetTriangle(CenterPointTriangleID);
+		FFrame3d SeedFrame(Mesh.GetTriCentroid(CenterPointTriangleID), Mesh.GetTriNormal(CenterPointTriangleID));
+		SeedFrame.Origin = CenterPoint;
+
+		if (bUseInterpolatedNormal && Mesh.HasAttributes())
+		{
+			if (const FDynamicMeshNormalOverlay* Normals = Mesh.Attributes()->PrimaryNormals())
+			{
+				if (Normals->IsSetTriangle(CenterPointTriangleID))
+				{
+					FTriangle3d Triangle(Mesh.GetVertex(TriVertices.A), Mesh.GetVertex(TriVertices.B), Mesh.GetVertex(TriVertices.C));
+					FVector3d BaryCoords = Triangle.GetBarycentricCoords(SeedFrame.Origin);
+					FVector3d InterpNormal;
+					Mesh.Attributes()->PrimaryNormals()->GetTriBaryInterpolate<double>(CenterPointTriangleID, &BaryCoords.X, &InterpNormal.X);
+					SeedFrame.AlignAxis(2, InterpNormal);
+				}
+			}
+		}
+
+		if (TangentYDirection.SquaredLength() > 0.1)
+		{
+			SeedFrame.ConstrainedAlignAxis(1, TangentYDirection, SeedFrame.Z());
+		}
+		else
+		{
+			SeedFrame.ConstrainedAlignAxis(1, FVector3d::UnitZ(), SeedFrame.Z());
+		}
+
+		if (UVRotationDeg != 0)
+		{
+			SeedFrame.Rotate(FQuaterniond(SeedFrame.Z(), UVRotationDeg, true));
+		}
+
+		TMeshLocalParam<FDynamicMesh3> Param(&Mesh);
+		Param.ParamMode = ELocalParamTypes::ExponentialMapUpwindAvg;
+
+		Param.bEnableExternalNormals = true;
+		Param.ExternalNormalFunc = [&](int32 VertexID) { return FMeshNormals::ComputeVertexNormal(Mesh, VertexID); };
+
+		Param.ComputeToMaxDistance(SeedFrame, TriVertices, Radius * FMathd::Sqrt2);
+
+		Param.GetAllComputedUVs(VertexIDs, VertexUVs, Radius, Radius * 2.0);
+	});
+
+	return TargetMesh;
+}
+
+
+
 #undef LOCTEXT_NAMESPACE
\ No newline at end of file
diff --git a/Engine/Plugins/Experimental/GeometryScripting/Source/GeometryScriptingCore/Public/GeometryScript/MeshNormalsFunctions.h b/Engine/Plugins/Experimental/GeometryScripting/Source/GeometryScriptingCore/Public/GeometryScript/MeshNormalsFunctions.h
--- a/Engine/Plugins/Experimental/GeometryScripting/Source/GeometryScriptingCore/Public/GeometryScript/MeshNormalsFunctions.h
+++ b/Engine/Plugins/Experimental/GeometryScripting/Source/GeometryScriptingCore/Public/GeometryScript/MeshNormalsFunctions.h
@@ -5,6 +5,7 @@
 #include "CoreMinimal.h"
 #include "Kismet/BlueprintFunctionLibrary.h"
 #include "GeometryScript/GeometryScriptTypes.h"
+#include "GeometryScript/GeometryScriptSelectionTypes.h"
 #include "MeshNormalsFunctions.generated.h"
 
 class UDynamicMesh;
@@ -122,6 +123,21 @@ public:
 	RecomputeNormals(  
 		UDynamicMesh* TargetMesh, 
 		FGeometryScriptCalculateNormalsOptions CalculateOptions,
+		bool bDeferChangeNotifications = false,
+		UGeometryScriptDebug* Debug = nullptr);
+
+	/**
+	 * Recompute the normals of TargetMesh on all the triangles/vertices of the given Selection using the given CalculateOptions. 
+	 * This method will preserve any existing hard edges, ie each shared triangle-vertex normal is recomputed by averaging 
+	 * the face normals of triangles that reference that shared triangle-vertex normal
+	 */
+	UFUNCTION(BlueprintCallable, Category = "GeometryScript|Normals", meta=(ScriptMethod))
+	static UPARAM(DisplayName = "Target Mesh") UDynamicMesh* 
+	RecomputeNormalsForMeshSelection(  
+		UDynamicMesh* TargetMesh, 
+		FGeometryScriptMeshSelection Selection,
+		FGeometryScriptCalculateNormalsOptions CalculateOptions,
+		bool bDeferChangeNotifications = false,
 		UGeometryScriptDebug* Debug = nullptr);
 
 	/**
diff --git a/Engine/Plugins/Experimental/GeometryScripting/Source/GeometryScriptingCore/Public/GeometryScript/MeshUVFunctions.h b/Engine/Plugins/Experimental/GeometryScripting/Source/GeometryScriptingCore/Public/GeometryScript/MeshUVFunctions.h
--- a/Engine/Plugins/Experimental/GeometryScripting/Source/GeometryScriptingCore/Public/GeometryScript/MeshUVFunctions.h
+++ b/Engine/Plugins/Experimental/GeometryScripting/Source/GeometryScriptingCore/Public/GeometryScript/MeshUVFunctions.h
@@ -407,4 +407,31 @@ public:
 		UGeometryScriptDebug* Debug = nullptr);
 
 
+	/**
+	 * Compute local UV parameterization on TargetMesh vertices around the given CenterPoint / Triangle. This method
+	 * uses a Discrete Exponential Map parameterization, which unwraps the mesh locally based on geodesic distances and angles.
+	 * The CenterPoint will have UV value (0,0), and the computed vertex UVs will be such that Length(UV) == geodesic distance.
+	 * 
+	 * @param CenterPoint the center point of the parameterization. This point must lie on the triangle specified by CenterPointTriangleID
+	 * @param CenterPointTriangleID the ID of the Triangle that contains CenterPoint
+	 * @param Radius the parameterization will be computed out to this geodesic radius
+	 * @param bUseInterpolatedNormal if true (default false), the normal frame used for the parameterization will be taken from the normal overlay, otherwise the CenterPointTriangleID normal will be used
+	 * @param VertexIDs output list of VertexIDs that UVs have been computed for, ie are within geodesic distance Radius from the CenterPoint
+	 * @param VertexUVs output list of Vertex UVs that corresponds to VertexIDs
+	 */
+	UFUNCTION(BlueprintCallable, Category = "GeometryScript|UVs", meta=(ScriptMethod, AdvancedDisplay = "bUseInterpolatedNormal, TangentYDirection"))
+	static UPARAM(DisplayName = "Target Mesh") UDynamicMesh* 
+	ComputeMeshLocalUVParam( 
+		UDynamicMesh* TargetMesh, 
+		FVector CenterPoint,
+		int32 CenterPointTriangleID,
+		TArray<int>& VertexIDs,
+		TArray<FVector2D>& VertexUVs,
+		double Radius = 1,
+		bool bUseInterpolatedNormal = false,
+		FVector TangentYDirection = FVector(0,0,0),
+		double UVRotationDeg = 0.0,
+		UGeometryScriptDebug* Debug = nullptr );
+
+
 };
\ No newline at end of file
diff --git a/Engine/Source/Runtime/GeometryCore/Public/Parameterization/MeshLocalParam.h b/Engine/Source/Runtime/GeometryCore/Public/Parameterization/MeshLocalParam.h
--- a/Engine/Source/Runtime/GeometryCore/Public/Parameterization/MeshLocalParam.h
+++ b/Engine/Source/Runtime/GeometryCore/Public/Parameterization/MeshLocalParam.h
@@ -46,6 +46,17 @@ public:
 	/** Type of local parameterization to compute */
 	ELocalParamTypes ParamMode = ELocalParamTypes::ExponentialMapUpwindAvg;
 
+	/** 
+	 * If true, ExternalNormalFunc will be used to fetch normals instead of requesting from PointSet->GetVertexNormal.
+	 * This allows for (1) vertex normals to be computed on the fly, if only a subset of the point set is being parameterizedx
+	 * and (2) alternate normals (eg smoothed, etc) to be provided w/o having to define a separate class.
+	 */
+	bool bEnableExternalNormals = false;
+
+	/** Normals will be requested from this function if bEnableExternalNormals == true */
+	TFunction<FVector3d(int32)> ExternalNormalFunc;
+
+
 	TMeshLocalParam(const PointSetType* PointSetIn)
 	{
 		PointSet = PointSetIn;
@@ -188,6 +199,31 @@ public:
 		return (Node != nullptr && Node->bFrozen) ? Node->UV : InvalidUV();
 	}
 
+	/**
+	 * Find all computed UVs within the specified distances
+	 * @param PointIDsOut PointID corresponding to each computed UV is returned here
+	 * @param PointUVsOut UV value corresponding to each PointID is returned here
+	 * @param MaxUVMagnitude computed UVs will only be included if their magnitude is smaller than this value (eg, this is the UV-space radius)
+	 * @param MaxGraphDistance computed UVs will only be included if their graph distance is smaller than this value. Graph distance is generally always be larger than UV distance. This can be used to avoid problematic UVs that might result from algorithm failures.
+	 */
+	void GetAllComputedUVs(
+		TArray<int32>& PointIDsOut, 
+		TArray<FVector2d>& PointUVsOut,
+		double MaxUVMagnitude = TNumericLimits<float>::Max(),
+		double MaxGraphDist = TNumericLimits<float>::Max() ) const
+	{
+		double MaxUVMagSqr = MaxUVMagnitude * MaxUVMagnitude;
+		for (const FGraphNode& Node : AllocatedNodes)
+		{
+			if (Node.bFrozen && Node.GraphDistance < MaxGraphDist && Node.UV.SquaredLength() < MaxUVMagSqr)
+			{
+				PointIDsOut.Add(Node.PointID);
+				PointUVsOut.Add(Node.UV);
+			}
+		}
+	}
+
+
 
 	/**
 	 * Apply a function to each calculated UV
@@ -235,13 +271,6 @@ protected:
 		return FVector3d(PointSet->GetVertexNormal(PointID));
 	}
 
-	FFrame3d GetFrame(const int32 PointID) const
-	{
-		return FFrame3d(GetPosition(PointID), GetNormal(PointID));
-	}
-
-
-
 protected:
 
 	// information about each active/computed point
@@ -252,8 +281,17 @@ protected:
 		double GraphDistance;
 		FVector2d UV;
 		bool bFrozen;
+
+		FVector3d CachedNormal;
 	};
 
+
+	FFrame3d GetFrame(const FGraphNode& Node) const
+	{
+		return FFrame3d(GetPosition(Node.PointID), Node.CachedNormal);
+	}
+
+
 	// To avoid constructing FGraphNode for all input points (because we are computing a "local" param),
 	// we only allocate on demand, and then store a sparse mapping in IDToNodeIndexMap
 	TMap<int32, int32> IDToNodeIndexMap;
@@ -272,9 +310,6 @@ protected:
 	double MaxUVDistance;
 
 
-
-
-
 	void ProcessQueueUntilTermination(double MaxDistance)
 	{
 		while (Queue.GetCount() > 0)
@@ -363,8 +398,7 @@ protected:
 		FGraphNode* ParentNode = GetNodeForPointSetID(Node.ParentPointID, false);
 		check(ParentNode != nullptr);
 
-		FVector3d ParentPos = GetPosition(Node.ParentPointID);
-		FFrame3d ParentFrame = GetFrame(Node.ParentPointID);
+		FFrame3d ParentFrame = GetFrame(*ParentNode);
 
 		Node.UV = PropagateUV(GetPosition(Node.PointID), ParentNode->UV, ParentFrame, SeedFrame);
 	}
@@ -386,7 +420,7 @@ protected:
 			FGraphNode* NbrNode = GetNodeForPointSetID(NbrPointID, false);
 			if (NbrNode != nullptr && NbrNode->bFrozen)
 			{
-				FFrame3d NbrFrame(GetFrame(NbrNode->PointID));
+				FFrame3d NbrFrame(GetFrame(*NbrNode));
 				FVector2d NbrUV = PropagateUV(NodePos, NbrNode->UV, NbrFrame, SeedFrame);
 				double Weight = 1.0 / (DistanceSquared(NodePos,NbrFrame.Origin) + TMathUtil<double>::ZeroTolerance);
 				AverageUV += Weight * NbrUV;
@@ -416,7 +450,17 @@ protected:
 		{
 			if (bCreateIfMissing)
 			{
-				FGraphNode NewNode{ PointSetID, -1, 0, FVector2d::Zero(), false };
+				FGraphNode NewNode{ PointSetID, -1, 0, FVector2d::Zero(), false, FVector3d::UnitZ() };
+
+				if (bEnableExternalNormals)
+				{
+					NewNode.CachedNormal = ExternalNormalFunc(PointSetID);
+				}
+				else
+				{
+					NewNode.CachedNormal = GetNormal(PointSetID);
+				}
+
 				int32 NewIndex = AllocatedNodes.Num();
 				AllocatedNodes.Add(NewNode);
 				IDToNodeIndexMap.Add(PointSetID, NewIndex);
diff --git a/Engine/Source/Runtime/GeometryFramework/Private/Components/DynamicMeshComponent.cpp b/Engine/Source/Runtime/GeometryFramework/Private/Components/DynamicMeshComponent.cpp
--- a/Engine/Source/Runtime/GeometryFramework/Private/Components/DynamicMeshComponent.cpp
+++ b/Engine/Source/Runtime/GeometryFramework/Private/Components/DynamicMeshComponent.cpp
@@ -359,6 +359,11 @@ void UDynamicMeshComponent::NotifyMeshUpdated()
 	ResetProxy();
 }
 
+void UDynamicMeshComponent::NotifyMeshModified()
+{
+	NotifyMeshUpdated();
+}
+
 
 void UDynamicMeshComponent::FastNotifyColorsUpdated()
 {
@@ -515,13 +520,40 @@ void UDynamicMeshComponent::FastNotifyVertexAttributesUpdated(EMeshRenderAttribu
 	}
 }
 
-
 void UDynamicMeshComponent::FastNotifyUVsUpdated()
 {
 	FastNotifyVertexAttributesUpdated(EMeshRenderAttributeFlags::VertexUVs);
 }
 
 
+void UDynamicMeshComponent::NotifyMeshVertexAttributesModified( bool bPositions, bool bNormals, bool bUVs, bool bColors )
+{
+	EMeshRenderAttributeFlags Flags = EMeshRenderAttributeFlags::None;
+	if (bPositions)
+	{
+		Flags |= EMeshRenderAttributeFlags::Positions;
+	}
+	if (bNormals)
+	{
+		Flags |= EMeshRenderAttributeFlags::VertexNormals;
+	}
+	if (bUVs)
+	{
+		Flags |= EMeshRenderAttributeFlags::VertexUVs;
+	}
+	if (bColors)
+	{
+		Flags |= EMeshRenderAttributeFlags::VertexColors;
+	}
+
+	if (Flags == EMeshRenderAttributeFlags::None)
+	{
+		return;
+	}
+	FastNotifyVertexAttributesUpdated(Flags);
+}
+
+
 
 void UDynamicMeshComponent::FastNotifySecondaryTrianglesChanged()
 {
@@ -1148,7 +1180,9 @@ void UDynamicMeshComponent::SetDynamicMesh(UDynamicMesh* NewMesh)
 		MeshObject->OnMeshChanged().Remove(MeshObjectChangedHandle);
 	}
 
-	NewMesh->Rename( nullptr, this );		// set Outer of NewMesh to be this Component
+	// set Outer of NewMesh to be this Component, ie transfer ownership. This is done via "renaming", which is
+	// a bit odd, so the flags prevent some standard "renaming" behaviors from happening
+	NewMesh->Rename( nullptr, this, REN_DontCreateRedirectors | REN_ForceNoResetLoaders);
 	MeshObject = NewMesh;
 	MeshObjectChangedHandle = MeshObject->OnMeshChanged().AddUObject(this, &UDynamicMeshComponent::OnMeshObjectChanged);
 
diff --git a/Engine/Source/Runtime/GeometryFramework/Public/Components/DynamicMeshComponent.h b/Engine/Source/Runtime/GeometryFramework/Public/Components/DynamicMeshComponent.h
--- a/Engine/Source/Runtime/GeometryFramework/Public/Components/DynamicMeshComponent.h
+++ b/Engine/Source/Runtime/GeometryFramework/Public/Components/DynamicMeshComponent.h
@@ -89,8 +89,9 @@ public:
 	virtual UDynamicMesh* GetDynamicMesh() override { return MeshObject; }
 
 	/**
-	 * Set the child UDynamicMesh. This can be used to 'share' a UDynamicMesh between Component instances.
-	 * @warning Currently this is somewhat risky, it is on the caller/clients to make sure that the actual mesh is not being simultaneously modified on multiple threads
+	 * Replace the current UDynamicMesh with a new one, and transfer ownership of NewMesh to this Component.
+	 * This can be used to (eg) assign a UDynamicMesh created with NewObject in the Transient Package to this Component.
+	 * @warning If NewMesh is owned/Outer'd to another DynamicMeshComponent, a GLEO error may occur if that Component is serialized.
 	 */
 	UFUNCTION(BlueprintCallable, Category = "Dynamic Mesh Component")
 	void SetDynamicMesh(UDynamicMesh* NewMesh);
@@ -210,8 +211,36 @@ public:
 		TFuture<bool>& Precompute, const TArray<int32>& UpdateSets, const UE::Geometry::FAxisAlignedBox3d& UpdateSetBounds);
 
 
+	//===============================================================================================================
+	// RenderBuffer Update Blueprint API. 
+public:
 
+	/**
+	 * Notify the Component that it's DynamicMesh has been modified externally. This will result in all Rendering Data
+	 * for the Component being rebuilt on the next frame (internally the Scene Proxy is fully destroyed and rebuilt).
+	 *
+	 * You must use this function if the mesh triangulation has been modified, or if polygroups or material assignments
+	 * have been changed, or if Normal/UV/Color topology has changed (ie new split-vertices have been introduced).
+	 * If only vertex attribute values (position, normals, UVs, colors) have been modified, then 
+	 * Notify Vertex Attributes Updated can be used to do a faster update.
+	 */
+	UFUNCTION(BlueprintCallable, Category = "Dynamic Mesh Component|Rendering", DisplayName = "Notify Mesh Updated")
+	virtual void NotifyMeshModified();
 
+	/**
+	 * Notify the Component that vertex attribute values of it's DynamicMesh have been modified externally. This will result in
+	 * Rendering vertex buffers being updated. This update path is more efficient than doing a full Notify Mesh Updated.
+	 * 
+	 * @warning it is invalid to call this function if (1) the mesh triangulation has also been changed, (2) triangle MaterialIDs have been changed,
+	 * or (3) any attribute overlay (normal, color, UV) topology has been modified, ie split-vertices have been added/removed.
+	 * Behavior of this function is undefined in these cases and may crash. If you are unsure, use Notify Mesh Updated.
+	 */
+	UFUNCTION(BlueprintCallable, Category = "Dynamic Mesh Component|Rendering", DisplayName = "Notify Vertex Attributes Updated")
+	virtual void NotifyMeshVertexAttributesModified(
+		bool bPositions = true, 
+		bool bNormals = true,
+		bool bUVs = true,
+		bool bColors = true);
 
 	//===============================================================================================================
 	// Change Support. These changes are primarily used for Undo/Redo, however there is no strict assumption
