[Engine]

diff --git a/Engine/Source/Runtime/GeometryCore/Public/DynamicMesh/MeshNormals.h b/Engine/Source/Runtime/GeometryCore/Public/DynamicMesh/MeshNormals.h
--- a/Engine/Source/Runtime/GeometryCore/Public/DynamicMesh/MeshNormals.h
+++ b/Engine/Source/Runtime/GeometryCore/Public/DynamicMesh/MeshNormals.h
@@ -22,7 +22,7 @@ namespace Geometry
  * FMeshNormals is a utility class that can calculate and store various types of
  * normal vectors for a FDynamicMesh. 
  */
-class GEOMETRYCORE_API FMeshNormals
+class FMeshNormals
 {
 protected:
 	/** Target Mesh */
@@ -56,7 +56,7 @@ public:
 	/**
 	 * Set the size of the Normals array to Count, and optionally clear all values to (0,0,0)
 	 */
-	void SetCount(int Count, bool bClearToZero);
+	GEOMETRYCORE_API void SetCount(int Count, bool bClearToZero);
 
 	/**
 	 * Compute standard per-vertex normals by averaging one-ring face normals
@@ -84,7 +84,7 @@ public:
 	 * This won't help if a whole mesh componenent is entirely made of such degenerate triangles
 	 * (which is usually ok).
 	 */
-	void SetDegenerateTriangleNormalsToNeighborNormal();
+	GEOMETRYCORE_API void SetDegenerateTriangleNormalsToNeighborNormal();
 
 	/**
 	 * Recompute the per-element normals of the given overlay by averaging one-ring face normals
@@ -101,27 +101,27 @@ public:
 	 * @warning assumes that the computed normals are vertex normals
 	 * @param bInvert if true, normals are flipped
 	 */
-	void CopyToVertexNormals(FDynamicMesh3* SetMesh, bool bInvert = false) const;
+	GEOMETRYCORE_API void CopyToVertexNormals(FDynamicMesh3* SetMesh, bool bInvert = false) const;
 
 	/**
 	 * Copy the current set of normals to the NormalOverlay attribute layer
 	 * @warning assumes that the computed normals are attribute normals
 	 * @param bInvert if true, normals are flipped
 	 */
-	void CopyToOverlay(FDynamicMeshNormalOverlay* NormalOverlay, bool bInvert = false) const;
+	GEOMETRYCORE_API void CopyToOverlay(FDynamicMeshNormalOverlay* NormalOverlay, bool bInvert = false) const;
 
 
 	/**
 	 * Compute per-vertex normals for the given Mesh
 	 * @param bInvert if true, normals are flipped
 	 */
-	static void QuickComputeVertexNormals(FDynamicMesh3& Mesh, bool bInvert = false);
+	static GEOMETRYCORE_API void QuickComputeVertexNormals(FDynamicMesh3& Mesh, bool bInvert = false);
 
 
 	/**
 	 * Apply rounds of explicit uniform-weighted normal smoothing to the VertexNormals attribute of the given Mesh.
 	 */
-	static void SmoothVertexNormals(FDynamicMesh3& Mesh, int32 SmoothingRounds, double SmoothingAlpha);
+	static GEOMETRYCORE_API void SmoothVertexNormals(FDynamicMesh3& Mesh, int32 SmoothingRounds, double SmoothingAlpha);
 
 
 	/**
@@ -130,7 +130,7 @@ public:
 	 * @param bWeightByAngle weight neighbor triangles by angle
 	 * @param bInvert if true, normals are flipped
 	 */
-	static void QuickComputeVertexNormalsForTriangles(FDynamicMesh3& Mesh, const TArray<int32>& Triangles, bool bWeightByArea = true, bool bWeightByAngle = true, bool bInvert = false);
+	static GEOMETRYCORE_API void QuickComputeVertexNormalsForTriangles(FDynamicMesh3& Mesh, const TArray<int32>& Triangles, bool bWeightByArea = true, bool bWeightByAngle = true, bool bInvert = false);
 
 
 	/**
@@ -139,7 +139,7 @@ public:
 	 * @param bWeightByAngle weight neighbor triangles by angle
 	 * @return the vertex normal at vertex VertIdx of Mesh.
 	 */
-	static FVector3d ComputeVertexNormal(const FDynamicMesh3& Mesh, int VertIdx, bool bWeightByArea = true, bool bWeightByAngle = true);
+	static GEOMETRYCORE_API FVector3d ComputeVertexNormal(const FDynamicMesh3& Mesh, int VertIdx, bool bWeightByArea = true, bool bWeightByAngle = true);
 
 	/**
 	 * Compute normal at mesh vertex by weighted sum of subset of one-ring triangle normals. Can optionally weight by area, angle, or both (averaged)
@@ -148,7 +148,7 @@ public:
 	 * @param bWeightByAngle weight neighbor triangles by angle
 	 * @return the vertex normal at vertex VertIdx of Mesh.
 	 */
-	static FVector3d ComputeVertexNormal(const FDynamicMesh3& Mesh, int32 VertIdx, TFunctionRef<bool(int32)> TriangleFilterFunc, bool bWeightByArea = true, bool bWeightByAngle = true);
+	static GEOMETRYCORE_API FVector3d ComputeVertexNormal(const FDynamicMesh3& Mesh, int32 VertIdx, TFunctionRef<bool(int32)> TriangleFilterFunc, bool bWeightByArea = true, bool bWeightByAngle = true);
 
 
 	/**
@@ -156,54 +156,54 @@ public:
 	 * @param bWeightByArea weight neighbor triangles by area
 	 * @param bWeightByAngle weight neighbor triangles by angle
 	 */
-	static FVector3d ComputeOverlayNormal(const FDynamicMesh3& Mesh, const FDynamicMeshNormalOverlay* NormalOverlay, int ElemIdx, bool bWeightByArea = true, bool bWeightByAngle = true);
+	static GEOMETRYCORE_API FVector3d ComputeOverlayNormal(const FDynamicMesh3& Mesh, const FDynamicMeshNormalOverlay* NormalOverlay, int ElemIdx, bool bWeightByArea = true, bool bWeightByAngle = true);
 
 	/**
 	 * Initialize the given NormalOverlay with per-vertex normals, ie single overlay element for each mesh vertex.
 	 * @param bUseMeshVertexNormalsIfAvailable if true and the parent mesh has per-vertex normals, use them instead of calculating new ones
 	 */
-	static void InitializeOverlayToPerVertexNormals(FDynamicMeshNormalOverlay* NormalOverlay, bool bUseMeshVertexNormalsIfAvailable = true);
+	static GEOMETRYCORE_API void InitializeOverlayToPerVertexNormals(FDynamicMeshNormalOverlay* NormalOverlay, bool bUseMeshVertexNormalsIfAvailable = true);
 
 
 	/**
 	 * Initialize the given NormalOverlay with per-face normals, ie separate overlay element for each vertex of each triangle.
 	 */
-	static void InitializeOverlayToPerTriangleNormals(FDynamicMeshNormalOverlay* NormalOverlay);
+	static GEOMETRYCORE_API void InitializeOverlayToPerTriangleNormals(FDynamicMeshNormalOverlay* NormalOverlay);
 
 
-	static void InitializeOverlayTopologyFromOpeningAngle(const FDynamicMesh3* Mesh, FDynamicMeshNormalOverlay* NormalOverlay, 
+	static GEOMETRYCORE_API void InitializeOverlayTopologyFromOpeningAngle(const FDynamicMesh3* Mesh, FDynamicMeshNormalOverlay* NormalOverlay, 
 		double AngleThresholdDeg);
 
-	static void InitializeOverlayTopologyFromFaceGroups(const FDynamicMesh3* Mesh, FDynamicMeshNormalOverlay* NormalOverlay);
+	static GEOMETRYCORE_API void InitializeOverlayTopologyFromFaceGroups(const FDynamicMesh3* Mesh, FDynamicMeshNormalOverlay* NormalOverlay);
 
 
 	/**
 	 * Initialize the given Mesh with per-face normals, ie separate overlay element for each vertex of each triangle.
 	 */
-	static void InitializeMeshToPerTriangleNormals(FDynamicMesh3* Mesh);
+	static GEOMETRYCORE_API void InitializeMeshToPerTriangleNormals(FDynamicMesh3* Mesh);
 
 
 	/**
 	 * Initialize the given triangles of NormalOverlay with per-vertex normals, ie single overlay element for each mesh vertex.
 	 * Only the triangles included in the region are considered when calculating per-vertex normals.
 	 */
-	static void InitializeOverlayRegionToPerVertexNormals(FDynamicMeshNormalOverlay* NormalOverlay, const TArray<int32>& Triangles);
+	static GEOMETRYCORE_API void InitializeOverlayRegionToPerVertexNormals(FDynamicMeshNormalOverlay* NormalOverlay, const TArray<int32>& Triangles);
 
 	/**
 	 * Compute overlay normals for the given mesh
 	 * @param bInvert if true, normals are flipped
 	 */
-	static bool QuickRecomputeOverlayNormals(FDynamicMesh3& Mesh, bool bInvert = false, bool bWeightByArea = true, bool bWeightByAngle = true);
+	static GEOMETRYCORE_API bool QuickRecomputeOverlayNormals(FDynamicMesh3& Mesh, bool bInvert = false, bool bWeightByArea = true, bool bWeightByAngle = true);
 
 	/**
 	 * Compute overlay normals for the given mesh, for the given set of triangles
 	 */
-	static bool RecomputeOverlayTriNormals(FDynamicMesh3& Mesh, const TArray<int32>& Triangles, bool bWeightByArea = true, bool bWeightByAngle = true);
+	static GEOMETRYCORE_API bool RecomputeOverlayTriNormals(FDynamicMesh3& Mesh, const TArray<int32>& Triangles, bool bWeightByArea = true, bool bWeightByAngle = true);
 
 	/**
 	 * Compute overlay normals for the given mesh, for the given set of element IDs
 	 */
-	static bool RecomputeOverlayElementNormals(FDynamicMesh3& Mesh, const TArray<int32>& ElementIDs, bool bWeightByArea = true, bool bWeightByAngle = true);
+	static GEOMETRYCORE_API bool RecomputeOverlayElementNormals(FDynamicMesh3& Mesh, const TArray<int32>& ElementIDs, bool bWeightByArea = true, bool bWeightByAngle = true);
 
 	/**
 	 * Retrieve the area and/or angle weights for each vertex of a triangle
@@ -213,24 +213,24 @@ public:
 	 * @param bWeightByArea if true, include weighting by the area of the triangle
 	 * @param bWeightByAngle if true, include weighting by the interior angles of the triangle
 	 */
-	static FVector3d GetVertexWeightsOnTriangle(const FDynamicMesh3* Mesh, int TriID, double TriArea, bool bWeightByArea, bool bWeightByAngle);
+	static GEOMETRYCORE_API FVector3d GetVertexWeightsOnTriangle(const FDynamicMesh3* Mesh, int TriID, double TriArea, bool bWeightByArea, bool bWeightByAngle);
 
 protected:
 	/** Compute per-vertex normals using area-weighted averaging of one-ring triangle normals */
-	void Compute_FaceAvg_AreaWeighted();
+	GEOMETRYCORE_API void Compute_FaceAvg_AreaWeighted();
 	/** Compute per-vertex normals using a custom combination of area-weighted and angle-weighted averaging of one-ring triangle normals */
-	void Compute_FaceAvg(bool bWeightByArea, bool bWeightByAngle);
+	GEOMETRYCORE_API void Compute_FaceAvg(bool bWeightByArea, bool bWeightByAngle);
 
 	/** Compute per-triangle normals */
-	void Compute_Triangle();
+	GEOMETRYCORE_API void Compute_Triangle();
 
 	/** Recompute the element Normals of the given attribute overlay using area-weighted averaging of one-ring triangle normals */
-	void Compute_Overlay_FaceAvg_AreaWeighted(const FDynamicMeshNormalOverlay* NormalOverlay);
+	GEOMETRYCORE_API void Compute_Overlay_FaceAvg_AreaWeighted(const FDynamicMeshNormalOverlay* NormalOverlay);
 	/** Recompute the element Normals of the given attribute overlay using a custom combination of area-weighted and angle-weighted averaging of one-ring triangle normals */
-	void Compute_Overlay_FaceAvg(const FDynamicMeshNormalOverlay* NormalOverlay, bool bWeightByArea, bool bWeightByAngle);
+	GEOMETRYCORE_API void Compute_Overlay_FaceAvg(const FDynamicMeshNormalOverlay* NormalOverlay, bool bWeightByArea, bool bWeightByAngle);
 
 };
 
 
 } // end namespace UE::Geometry
-} // end namespace UE
\ No newline at end of file
+} // end namespace UE
diff --git a/Engine/Source/Runtime/GeometryCore/Public/DynamicMesh/MeshTangents.h b/Engine/Source/Runtime/GeometryCore/Public/DynamicMesh/MeshTangents.h
--- a/Engine/Source/Runtime/GeometryCore/Public/DynamicMesh/MeshTangents.h
+++ b/Engine/Source/Runtime/GeometryCore/Public/DynamicMesh/MeshTangents.h
@@ -21,7 +21,7 @@ using namespace UE::Math;
 /**
  * FDynamicMeshTangents is a helper object for accessing tangents stored in the AttributeSet of a FDynamicMesh3.
  */
-class GEOMETRYCORE_API FDynamicMeshTangents
+class FDynamicMeshTangents
 {
 public:
 	const FDynamicMesh3* Mesh = nullptr;
@@ -29,7 +29,7 @@ public:
 	const FDynamicMeshNormalOverlay* Tangents = nullptr;
 	const FDynamicMeshNormalOverlay* Bitangents = nullptr;
 
-	FDynamicMeshTangents(const FDynamicMesh3* MeshIn);
+	GEOMETRYCORE_API FDynamicMeshTangents(const FDynamicMesh3* MeshIn);
 
 	/**
 	 * Checks the mesh for valid tangents. When bCheckValues == true,
@@ -38,27 +38,27 @@ public:
 	 * @param bCheckValues inspect tangent values for zero/NaN
 	 * @return true if the mesh has valid tangents
 	 */
-	bool HasValidTangents(bool bCheckValues=false) const;
+	GEOMETRYCORE_API bool HasValidTangents(bool bCheckValues=false) const;
 
 	/**
 	 * If tangents are available in the overlays, returns them. If only Normal is available, computes orthogonal basis. Falls back to unit axes if no overlays are available.
 	 */
-	void GetTangentFrame(int32 TriangleID, int32 TriVertexIndex, FVector3f& NormalOut, FVector3f& TangentOut, FVector3f& BitangentOut) const;
+	GEOMETRYCORE_API void GetTangentFrame(int32 TriangleID, int32 TriVertexIndex, FVector3f& NormalOut, FVector3f& TangentOut, FVector3f& BitangentOut) const;
 	/**
 	 * If tangents are available in the overlays, returns them. Otherwise computes orthogonal basis to Normal argument.
 	 */
-	void GetTangentVectors(int32 TriangleID, int32 TriVertexIndex, const FVector3f& Normal, FVector3f& TangentOut, FVector3f& BitangentOut) const;
+	GEOMETRYCORE_API void GetTangentVectors(int32 TriangleID, int32 TriVertexIndex, const FVector3f& Normal, FVector3f& TangentOut, FVector3f& BitangentOut) const;
 	/**
 	 * If tangents are available in the overlays, returns them, otherwise falls back to unit X/Y axes.
 	 */
-	void GetTangentVectors(int32 TriangleID, int32 TriVertexIndex, FVector3f& TangentOut, FVector3f& BitangentOut) const;
+	GEOMETRYCORE_API void GetTangentVectors(int32 TriangleID, int32 TriVertexIndex, FVector3f& TangentOut, FVector3f& BitangentOut) const;
 };
 
 
 /**
  * Options used by TMeshTangents for tangents computation
  */
-struct GEOMETRYCORE_API FComputeTangentsOptions
+struct FComputeTangentsOptions
 {
 	/** if true, per-face tangents are blended to create triange-vertex tangents */
 	bool bAveraged = true;		
@@ -315,4 +315,4 @@ void TMeshTangents<RealType>::CopyTriVertexTangents(const TMeshTangents<OtherRea
 
 
 } // end namespace UE::Geometry
-} // end namespace UE
\ No newline at end of file
+} // end namespace UE
diff --git a/Engine/Source/Runtime/GeometryCore/Public/DynamicMesh/NonManifoldMappingSupport.h b/Engine/Source/Runtime/GeometryCore/Public/DynamicMesh/NonManifoldMappingSupport.h
--- a/Engine/Source/Runtime/GeometryCore/Public/DynamicMesh/NonManifoldMappingSupport.h
+++ b/Engine/Source/Runtime/GeometryCore/Public/DynamicMesh/NonManifoldMappingSupport.h
@@ -35,22 +35,22 @@ typedef  TDynamicVertexAttribute<int32, 1, FDynamicMesh3> FDynamicMeshVertexInt3
 * to the DynamicMesh (e.g. edge splits, collapses etc).
 *  
 */
-class GEOMETRYCORE_API FNonManifoldMappingSupport
+class FNonManifoldMappingSupport
 {
 public:
 
-	FNonManifoldMappingSupport(const FDynamicMesh3& Mesh);
+	GEOMETRYCORE_API FNonManifoldMappingSupport(const FDynamicMesh3& Mesh);
 
 
 	/**
 	* Update the support for a new DynamicMesh.
 	*/ 
-	void Reset(const FDynamicMesh3& Mesh);
+	GEOMETRYCORE_API void Reset(const FDynamicMesh3& Mesh);
 
 	/*
 	* Return true if attribute data indicates that the source data that was converted to this DynamicMesh contained non-manifold vertices.
 	*/
-	bool IsNonManifoldVertexInSource() const;
+	GEOMETRYCORE_API bool IsNonManifoldVertexInSource() const;
 
 	/*
 	* Return true if the provided DynamicMesh vertex id resulted from a non-manifold vertex in the source data.
@@ -70,7 +70,7 @@ public:
 	* 
 	* Note: the code assumes but does not check that vid is a valid vertex vid.
 	*/
-	int32 GetOriginalNonManifoldVertexID(const int32 vid) const;
+	GEOMETRYCORE_API int32 GetOriginalNonManifoldVertexID(const int32 vid) const;
 
 
 	// --- helper functions.
@@ -82,7 +82,7 @@ public:
 	* 
 	* Note: Failure occurs if the DynamicMesh does not have attributes enabled or if the provided array is not long enough to provide a mapping value for each DynamicMesh vertex id.
 	*/
-	static bool AttachNonManifoldVertexMappingData(const TArray<int32>& VertexToNonManifoldVertexIDMap, FDynamicMesh3& MeshInOut);
+	static GEOMETRYCORE_API bool AttachNonManifoldVertexMappingData(const TArray<int32>& VertexToNonManifoldVertexIDMap, FDynamicMesh3& MeshInOut);
 
 	/*
 	*  Removes vertex mapping data related to the non-manifold nature of the source data the produced this DynamicMesh.
@@ -90,7 +90,7 @@ public:
 	*  Note: this will invalidate any NonManifoldMappingSupport object associated with this DynamicMesh, 
 	*  and subsequent use of such object will produce unexpected results.
 	*/
-	static void RemoveNonManifoldVertexMappingData(FDynamicMesh3& MeshInOut);
+	static GEOMETRYCORE_API void RemoveNonManifoldVertexMappingData(FDynamicMesh3& MeshInOut);
 
 
 
@@ -108,7 +108,7 @@ public:
 	/*
 	* Name used to identify vertex attribute data generated during conversion to a DynamicMesh in the case that the source was non-manifold.
 	*/
-	static FName NonManifoldMeshVIDsAttrName;
+	static GEOMETRYCORE_API FName NonManifoldMeshVIDsAttrName;
 
 
 protected:
@@ -119,4 +119,4 @@ protected:
 
 
 } // namespace Geometry
-} // namespace UE
\ No newline at end of file
+} // namespace UE
diff --git a/Engine/Source/Runtime/GeometryCore/Public/DynamicMesh/Operations/MergeCoincidentMeshEdges.h b/Engine/Source/Runtime/GeometryCore/Public/DynamicMesh/Operations/MergeCoincidentMeshEdges.h
--- a/Engine/Source/Runtime/GeometryCore/Public/DynamicMesh/Operations/MergeCoincidentMeshEdges.h
+++ b/Engine/Source/Runtime/GeometryCore/Public/DynamicMesh/Operations/MergeCoincidentMeshEdges.h
@@ -25,11 +25,11 @@ class FDynamicMesh3;
  * they cannot be merged. 
  *
  */
-class GEOMETRYCORE_API FMergeCoincidentMeshEdges
+class FMergeCoincidentMeshEdges
 {
 public:
 	/** default tolerance is float ZeroTolerance */
-	static const double DEFAULT_TOLERANCE;  // = FMathf::ZeroTolerance;
+	static GEOMETRYCORE_API const double DEFAULT_TOLERANCE;  // = FMathf::ZeroTolerance;
 
 	/** The mesh that we are modifying */
 	FDynamicMesh3* Mesh;
@@ -69,7 +69,7 @@ public:
 	 * Run the merge operation and modify .Mesh
 	 * @return true if the algorithm succeeds
 	 */
-	virtual bool Apply();
+	GEOMETRYCORE_API virtual bool Apply();
 
 
 protected:
@@ -86,4 +86,4 @@ protected:
 
 
 } // end namespace UE::Geometry
-} // end namespace UE
\ No newline at end of file
+} // end namespace UE
diff --git a/Engine/Source/Runtime/GeometryCore/Public/DynamicMesh/Operations/SplitAttributeWelder.h b/Engine/Source/Runtime/GeometryCore/Public/DynamicMesh/Operations/SplitAttributeWelder.h
--- a/Engine/Source/Runtime/GeometryCore/Public/DynamicMesh/Operations/SplitAttributeWelder.h
+++ b/Engine/Source/Runtime/GeometryCore/Public/DynamicMesh/Operations/SplitAttributeWelder.h
@@ -20,7 +20,7 @@ class TDynamicMeshVectorOverlay;
 * FSplitAttributeWelder can be used to weld split-attributes that share the same vertex in the parent mesh (e.g. split normals)
 *
 */
-class GEOMETRYCORE_API FSplitAttributeWelder
+class FSplitAttributeWelder
 {
 
 	typedef TDynamicMeshVectorOverlay<float, 2, FVector2f> FDynamicMeshUVOverlay;
@@ -60,13 +60,13 @@ public:
 	/**
 	* Weld split-elements at the ParentVID in each overlay that are "close enough" as defined by the appropriate threshold. 
 	*/
-	void WeldSplitElements(FDynamicMesh3& ParentMesh, const int32 ParentVID);
+	GEOMETRYCORE_API void WeldSplitElements(FDynamicMesh3& ParentMesh, const int32 ParentVID);
 	
 
 	/**
 	*  Weld split-elements across the entire mesh. 
 	*/
-	void WeldSplitElements(FDynamicMesh3& ParentMesh);
+	GEOMETRYCORE_API void WeldSplitElements(FDynamicMesh3& ParentMesh);
 	
 
 	/**
@@ -76,7 +76,7 @@ public:
 	* @param UVOverlay contains potentially split UVs associated with the ParentVID. Assumed to be an overlay of the ParentMesh.
 	* @param UVDistSqrdThreshold  provides a threshold in distance squared UV space between split UVs that should be welded. 
 	*/
-	static void WeldSplitUVs(const int32 ParentVID, FDynamicMeshUVOverlay& UVOverlay, float UVDistSqrdThreshold);
+	static GEOMETRYCORE_API void WeldSplitUVs(const int32 ParentVID, FDynamicMeshUVOverlay& UVOverlay, float UVDistSqrdThreshold);
 	
 
 	/**
@@ -87,7 +87,7 @@ public:
 	* @param DotThreshold  provides a threshold for the dot product between two vectors that should be welded.
 	* @param bMergeZeroVectors will weld vectors that are too small to normalize.
 	*/
-	static void WeldSplitUnitVectors(const int32 ParentVID, FDynamicMeshNormalOverlay& NormalOverlay, float DotThreshold, bool bMergeZeroVectors = true);
+	static GEOMETRYCORE_API void WeldSplitUnitVectors(const int32 ParentVID, FDynamicMeshNormalOverlay& NormalOverlay, float DotThreshold, bool bMergeZeroVectors = true);
 	
 	/**
 	* Welds split colors shared by the ParentVID vertex in the ParentMesh.  
@@ -96,7 +96,7 @@ public:
 	* @param ColorOverlay contains potentially split colors associated with the ParentVID. Assumed to be an overlay of the ParentMesh.
 	* @param ColorDistSqrdThreshold  provides a threshold applied to the sum of the square differences between to colors.  Used to determine if a weld should be preformed.
 	*/
-	static void WeldSplitColors(const int32 ParentVID, FDynamicMeshColorOverlay& ColorOverlay, float ColorDistSqrdThreshold);
+	static GEOMETRYCORE_API void WeldSplitColors(const int32 ParentVID, FDynamicMeshColorOverlay& ColorOverlay, float ColorDistSqrdThreshold);
 
 };
 
diff --git a/Engine/Source/Runtime/GeometryCore/Public/DynamicMeshEditor.h b/Engine/Source/Runtime/GeometryCore/Public/DynamicMeshEditor.h
--- a/Engine/Source/Runtime/GeometryCore/Public/DynamicMeshEditor.h
+++ b/Engine/Source/Runtime/GeometryCore/Public/DynamicMeshEditor.h
@@ -22,7 +22,7 @@ struct FDynamicSubmesh3;
  * FDynamicMeshEditResult is used to return information about new mesh elements
  * created by mesh changes, primarily in FDynamicMeshEditor
  */
-struct GEOMETRYCORE_API FDynamicMeshEditResult
+struct FDynamicMeshEditResult
 {
 	/** New vertices created by an edit */
 	TArray<int> NewVertices;
@@ -56,7 +56,7 @@ struct GEOMETRYCORE_API FDynamicMeshEditResult
 	}
 
 	/** Flatten the triangle/quad/polygon lists into a single list of all triangles */
-	void GetAllTriangles(TArray<int>& TrianglesOut) const;
+	GEOMETRYCORE_API void GetAllTriangles(TArray<int>& TrianglesOut) const;
 };
 
 
@@ -67,7 +67,7 @@ struct GEOMETRYCORE_API FDynamicMeshEditResult
  * can be used to construct higher-level operations. For example an Extrude operation
  * could be implemented via DuplicateTriangles() and StitchLoopMinimal().
  */
-class GEOMETRYCORE_API FDynamicMeshEditor
+class FDynamicMeshEditor
 {
 public:
 	/** The mesh we will be editing */
@@ -81,7 +81,7 @@ public:
 	/**
 	 * Remove any vertices that are not used by any triangles
 	 */
-	bool RemoveIsolatedVertices();
+	GEOMETRYCORE_API bool RemoveIsolatedVertices();
 
 	//////////////////////////////////////////////////////////////////////////
 	// Create and Remove Triangle Functions
@@ -96,7 +96,7 @@ public:
 	 * @param ResultOut lists of newly created triangles/vertices/etc
 	 * @return true if operation succeeded. If a failure occurs, any added triangles are removed via RemoveTriangles
 	 */
-	bool StitchVertexLoopsMinimal(const TArray<int>& VertexLoop1, const TArray<int>& VertexLoop2, FDynamicMeshEditResult& ResultOut);
+	GEOMETRYCORE_API bool StitchVertexLoopsMinimal(const TArray<int>& VertexLoop1, const TArray<int>& VertexLoop2, FDynamicMeshEditResult& ResultOut);
 
 
 	/**
@@ -118,7 +118,7 @@ public:
 	 * @param ResultOut lists of newly created triangles/vertices/etc
 	 * @return true if operation succeeded. If a failure occurs, any added triangles are removed via RemoveTriangles
 	 */
-	bool StitchVertexLoopToTriVidPairSequence(
+	GEOMETRYCORE_API bool StitchVertexLoopToTriVidPairSequence(
 		const TArray<TPair<int32, TPair<int8, int8>>>& TriVidPairs1,
 		const TArray<int>& VertexLoop2, FDynamicMeshEditResult& ResultOut);
 
@@ -134,7 +134,7 @@ public:
 	 * @param TriVertPairsOut Output
 	 * @return false if not successful (usually due to a mismatch between VidLoop and EdgeLoop).
 	 */
-	static bool ConvertLoopToTriVidPairSequence(const FDynamicMesh3& Mesh,
+	static GEOMETRYCORE_API bool ConvertLoopToTriVidPairSequence(const FDynamicMesh3& Mesh,
 		const TArray<int32>& VidLoop, const TArray<int32>& EdgeLoop,
 		TArray<TPair<int32, TPair<int8, int8>>>& TriVertPairsOut);
 
@@ -147,7 +147,7 @@ public:
 	 * @param ResultOut lists of newly created triangles/vertices/etc
 	 * @return true if operation succeeded.  If a failure occurs, any added triangles are removed via RemoveTriangles
 	 */
-	bool StitchSparselyCorrespondedVertexLoops(const TArray<int>& VertexIDs1, const TArray<int>& MatchedIndices1, const TArray<int>& VertexIDs2, const TArray<int>& MatchedIndices2, FDynamicMeshEditResult& ResultOut, bool bReverseOrientation = false);
+	GEOMETRYCORE_API bool StitchSparselyCorrespondedVertexLoops(const TArray<int>& VertexIDs1, const TArray<int>& MatchedIndices1, const TArray<int>& VertexIDs2, const TArray<int>& MatchedIndices2, FDynamicMeshEditResult& ResultOut, bool bReverseOrientation = false);
 
 
 	/**
@@ -156,7 +156,7 @@ public:
 	 * @param Loop2 second loop of sequential vertices. These vertices and their edges will not exist after the operation.
 	 * @return true if operation succeeded, false if any errors ocurred
 	 */
-	bool WeldVertexLoops(const TArray<int32>& VertexLoop1, const TArray<int32>& VertexLoop2);
+	GEOMETRYCORE_API bool WeldVertexLoops(const TArray<int32>& VertexLoop1, const TArray<int32>& VertexLoop2);
 
 
 	/**
@@ -167,7 +167,7 @@ public:
 	 * @param ResultOut lists of newly created triangles
 	 * @return true if operation succeeded.  If a failure occurs, any added triangles are removed via RemoveTriangles.
 	 */
-	bool AddTriangleFan_OrderedVertexLoop(int CenterVertex, const TArray<int>& VertexLoop, int GroupID, FDynamicMeshEditResult& ResultOut);
+	GEOMETRYCORE_API bool AddTriangleFan_OrderedVertexLoop(int CenterVertex, const TArray<int>& VertexLoop, int GroupID, FDynamicMeshEditResult& ResultOut);
 
 
 	/**
@@ -176,7 +176,7 @@ public:
 	 * @param IndexMaps returned mappings from old to new triangles/vertices/etc (you may initialize to optimize memory usage, etc)
 	 * @param ResultOut lists of newly created triangles/vertices/etc
 	 */
-	void DuplicateTriangles(const TArray<int>& Triangles, FMeshIndexMappings& IndexMaps, FDynamicMeshEditResult& ResultOut);
+	GEOMETRYCORE_API void DuplicateTriangles(const TArray<int>& Triangles, FMeshIndexMappings& IndexMaps, FDynamicMeshEditResult& ResultOut);
 
 
 	/**
@@ -210,7 +210,7 @@ public:
 	 * @param bHandleBoundaryVertices if true, boundary vertices are handled as described above, otherwise function aborts and returns false if a boundary vertex is encountered
 	 * @return true on success
 	 */
-	bool DisconnectTriangles(const TArray<int>& Triangles, TArray<FLoopPairSet>& LoopSetOut, bool bHandleBoundaryVertices);
+	GEOMETRYCORE_API bool DisconnectTriangles(const TArray<int>& Triangles, TArray<FLoopPairSet>& LoopSetOut, bool bHandleBoundaryVertices);
 
 	/**
 	 * Variant of DisconnectTriangles that takes a precomputed set of BoundaryLoops of the triangles, and a TSet of Triangles.
@@ -221,7 +221,7 @@ public:
 	 * @param LoopSetOut returned set of boundary loops. LoopA is original loop which remains with "outer" triangles previously connected to TriangleSet, and LoopB is new boundary loop of TriangleSet
 	 * @param bAllowBoundaryVertices if true, mesh boundary vertices are duplicated and left on the 'outer' loop, otherwise function aborts and returns false if a boundary vertex is encountered
 	 */
-	bool DisconnectTriangles(const TSet<int>& TriangleSet, const TArray<FEdgeLoop>& BoundaryLoops, TArray<FLoopPairSet>& LoopSetOut, bool bAllowBoundaryVertices);
+	GEOMETRYCORE_API bool DisconnectTriangles(const TSet<int>& TriangleSet, const TArray<FEdgeLoop>& BoundaryLoops, TArray<FLoopPairSet>& LoopSetOut, bool bAllowBoundaryVertices);
 
 
 
@@ -230,23 +230,23 @@ public:
 	* @param Triangles set of triangles
 	* @param bPreventBowties do some additional processing and vertex splitting as needed to prevent the creation of any new bowties
 	*/
-	void DisconnectTriangles(const TArray<int>& Triangles, bool bPreventBowties = true);
+	GEOMETRYCORE_API void DisconnectTriangles(const TArray<int>& Triangles, bool bPreventBowties = true);
 
 
 	/**
 	 * Splits all bowties across the whole mesh
 	 */
-	void SplitBowties(FDynamicMeshEditResult& ResultOut);
+	GEOMETRYCORE_API void SplitBowties(FDynamicMeshEditResult& ResultOut);
 
 	/**
 	 * Splits any bowties specifically on the given vertex, and updates (does not reset!) ResultOut with any added vertices
 	 */
-	void SplitBowties(int VertexID, FDynamicMeshEditResult& ResultOut);
+	GEOMETRYCORE_API void SplitBowties(int VertexID, FDynamicMeshEditResult& ResultOut);
 
 	/**
 	 * Splits bowties attached to any of the given triangles, and updates (does not reset!) ResultOut with any added vertices
 	 */
-	void SplitBowtiesAtTriangles(const TArray<int32>& TriangleIDs, FDynamicMeshEditResult& ResultOut);
+	GEOMETRYCORE_API void SplitBowtiesAtTriangles(const TArray<int32>& TriangleIDs, FDynamicMeshEditResult& ResultOut);
 
 	/**
 	 * In ReinsertSubmesh, a problem can arise where the mesh we are 
@@ -278,7 +278,7 @@ public:
 	 * @param DuplicateBehavior Choice of what to do if inserting a triangle from the submesh would duplicate a triangle in the base mesh
 	 * @return true if submesh successfully inserted, false if any triangles failed (which happens if triangle would result in non-manifold mesh)
 	 */
-	bool ReinsertSubmesh(const FDynamicSubmesh3& Submesh, FOptionallySparseIndexMap& SubToNewV, TArray<int>* NewTris = nullptr,
+	GEOMETRYCORE_API bool ReinsertSubmesh(const FDynamicSubmesh3& Submesh, FOptionallySparseIndexMap& SubToNewV, TArray<int>* NewTris = nullptr,
 						 EDuplicateTriBehavior DuplicateBehavior = EDuplicateTriBehavior::EnsureAbort);
 
 	/**
@@ -287,7 +287,7 @@ public:
 	 * @param bRemoveIsolatedVerts if true, remove vertices that end up with no triangles
 	 * @return true if all removes succeeded
 	 */
-	bool RemoveTriangles(const TArray<int>& Triangles, bool bRemoveIsolatedVerts);
+	GEOMETRYCORE_API bool RemoveTriangles(const TArray<int>& Triangles, bool bRemoveIsolatedVerts);
 
 
 	/**
@@ -296,7 +296,7 @@ public:
 	 * @param MinArea Remove components with less area than this
 	 * @return number of components removed
 	 */
-	int RemoveSmallComponents(double MinVolume, double MinArea = 0.0, int MinTriangleCount = 0);
+	GEOMETRYCORE_API int RemoveSmallComponents(double MinVolume, double MinArea = 0.0, int MinTriangleCount = 0);
 
 
 	/**
@@ -306,7 +306,7 @@ public:
 	 * @param OnRemoveTriFunc called for each triangle to be removed
 	 * @return true if all removes succeeded
 	 */
-	bool RemoveTriangles(const TArray<int>& Triangles, bool bRemoveIsolatedVerts, TFunctionRef<void(int)> OnRemoveTriFunc);
+	GEOMETRYCORE_API bool RemoveTriangles(const TArray<int>& Triangles, bool bRemoveIsolatedVerts, TFunctionRef<void(int)> OnRemoveTriFunc);
 
 
 	//////////////////////////////////////////////////////////////////////////
@@ -318,14 +318,14 @@ public:
 	 * @param Triangles the triangles to modify
 	 * @param bInvertNormals if ture we call InvertTriangleNormals()
 	 */
-	void ReverseTriangleOrientations(const TArray<int>& Triangles, bool bInvertNormals);
+	GEOMETRYCORE_API void ReverseTriangleOrientations(const TArray<int>& Triangles, bool bInvertNormals);
 
 	/**
 	 * Flip the normals of the given triangles. This includes their vertex normals, if they
 	 * exist, as well as any per-triangle attribute normals. @todo currently creates full-mesh bit arrays, could be more efficient on subsets
 	 * @param Triangles the triangles to modify
 	 */
-	void InvertTriangleNormals(const TArray<int>& Triangles);
+	GEOMETRYCORE_API void InvertTriangleNormals(const TArray<int>& Triangles);
 
 
 	/**
@@ -335,7 +335,7 @@ public:
 	 * @param bIsPlanar if the quad is known to be planar, operation is more efficient
 	 * @return the normal vector that was set
 	 */
-	FVector3f ComputeAndSetQuadNormal(const FIndex2i& QuadTris, bool bIsPlanar = false);
+	GEOMETRYCORE_API FVector3f ComputeAndSetQuadNormal(const FIndex2i& QuadTris, bool bIsPlanar = false);
 
 
 	/**
@@ -343,26 +343,26 @@ public:
 	 * @param QuadTris pair of triangle IDs. If second ID is invalid, it is ignored
 	 * @param Normal normal vector to set
 	 */
-	void SetQuadNormals(const FIndex2i& QuadTris, const FVector3f& Normal);
+	GEOMETRYCORE_API void SetQuadNormals(const FIndex2i& QuadTris, const FVector3f& Normal);
 
 	/**
 	 * Create and set new shared per-triangle normals for a list of triangles
 	 * @param Triangles list of triangle IDs
 	 * @param Normal normal vector to set
 	 */
-	void SetTriangleNormals(const TArray<int>& Triangles, const FVector3f& Normal);
+	GEOMETRYCORE_API void SetTriangleNormals(const TArray<int>& Triangles, const FVector3f& Normal);
 
 	/**
 	 * Create and set new shared per-triangle normals for a list of triangles. 
 	 * Normal at each vertex is calculated based only on average of triangles in set
 	 * @param Triangles list of triangle IDs
 	 */
-	void SetTriangleNormals(const TArray<int>& Triangles);
+	GEOMETRYCORE_API void SetTriangleNormals(const TArray<int>& Triangles);
 
 	/**
 	 * For a 'tube' of triangles connecting loops of corresponded vertices, set smooth normals such that corresponding vertices have corresponding normals
 	 */
-	void SetTubeNormals(const TArray<int>& Triangles, const TArray<int>& VertexIDs1, const TArray<int>& MatchedIndices1, const TArray<int>& VertexIDs2, const TArray<int>& MatchedIndices2);
+	GEOMETRYCORE_API void SetTubeNormals(const TArray<int>& Triangles, const TArray<int>& VertexIDs1, const TArray<int>& MatchedIndices1, const TArray<int>& VertexIDs2, const TArray<int>& MatchedIndices2);
 
 
 	//////////////////////////////////////////////////////////////////////////
@@ -379,7 +379,7 @@ public:
 	 * @param UVTranslation UVs are translated after scaling
 	 * @param UVLayerIndex which UV layer to operate on (must exist)
 	 */
-	void SetQuadUVsFromProjection(const FIndex2i& QuadTris, const FFrame3d& ProjectionFrame, float UVScaleFactor = 1.0f, const FVector2f& UVTranslation = FVector2f::Zero(), int UVLayerIndex = 0);
+	GEOMETRYCORE_API void SetQuadUVsFromProjection(const FIndex2i& QuadTris, const FFrame3d& ProjectionFrame, float UVScaleFactor = 1.0f, const FVector2f& UVTranslation = FVector2f::Zero(), int UVLayerIndex = 0);
 
 	/**
 	* Project triangles onto a plane defined by the ProjectionFrame and use that to create/set new shared per-triangle UVs.
@@ -393,7 +393,7 @@ public:
 	* @param bShiftToOrigin Whether to translate the UV coordinates to make their bounding box min corner be (0,0) before applying UVTranslation
 	* @param UVLayerIndex which UV layer to operate on (must exist)
 	*/
-	void SetTriangleUVsFromProjection(const TArray<int32>& Triangles, const FFrame3d& ProjectionFrame, 
+	GEOMETRYCORE_API void SetTriangleUVsFromProjection(const TArray<int32>& Triangles, const FFrame3d& ProjectionFrame, 
 		float UVScaleFactor = 1.0f, const FVector2f& UVTranslation = FVector2f::Zero(), bool bShiftToOrigin = true, int32 UVLayerIndex = 0);
 
 	/**
@@ -407,7 +407,7 @@ public:
 	* @param bShiftToOrigin Whether to translate the UV coordinates to make their bounding box min corner be (0,0) before applying UVTranslation
 	* @param bNormalizeBeforeScaling Whether to place the UV coordinates into the range [0,1] before applying UVScaleFactor and UVTranslation
 	*/
-	void SetTriangleUVsFromProjection(const TArray<int>& Triangles, const FFrame3d& ProjectionFrame, 
+	GEOMETRYCORE_API void SetTriangleUVsFromProjection(const TArray<int>& Triangles, const FFrame3d& ProjectionFrame, 
 		const FVector2f& UVScale = FVector2f::One(), const FVector2f& UVTranslation = FVector2f::Zero(), int UVLayerIndex = 0,
 		bool bShiftToOrigin = true, bool bNormalizeBeforeScaling = false);
 
@@ -415,7 +415,7 @@ public:
 	 * For triangles connecting loops of corresponded vertices, set UVs in a cylindrical pattern so that the U coordinate starts at 0 for the first corresponded pair of vertices, and cycles around to 1
 	 * Assumes Triangles array stores indices of triangles in progressively filling the tube, starting with VertexIDs*[0].  (This is used to set the UVs correctly at the seam joining the start & end of the loop)
 	 */
-	void SetGeneralTubeUVs(const TArray<int>& Triangles, const TArray<int>& VertexIDs1, const TArray<int>& MatchedIndices1, const TArray<int>& VertexIDs2, const TArray<int>& MatchedIndices2, const TArray<float>& UValues, const FVector3f& VDir, float UVScaleFactor = 1.0f, const FVector2f& UVTranslation = FVector2f::Zero(), int UVLayerIndex = 0);
+	GEOMETRYCORE_API void SetGeneralTubeUVs(const TArray<int>& Triangles, const TArray<int>& VertexIDs1, const TArray<int>& MatchedIndices1, const TArray<int>& VertexIDs2, const TArray<int>& MatchedIndices2, const TArray<float>& UValues, const FVector3f& VDir, float UVScaleFactor = 1.0f, const FVector2f& UVTranslation = FVector2f::Zero(), int UVLayerIndex = 0);
 
 	/**
 	 * Rescale UVs for the whole mesh, for the given UV attribute layer
@@ -424,7 +424,7 @@ public:
 	 * @param UVLayerIndex which UV layer to operate on (must exist)
 	 * @param ToWorld Optionally transform vertices for world space scaling
 	 */
-	void RescaleAttributeUVs(float UVScale = 1.0f, bool bWorldSpace = false, int UVLayerIndex = 0, TOptional<FTransformSRT3d> ToWorld = TOptional<FTransformSRT3d>());
+	GEOMETRYCORE_API void RescaleAttributeUVs(float UVScale = 1.0f, bool bWorldSpace = false, int UVLayerIndex = 0, TOptional<FTransformSRT3d> ToWorld = TOptional<FTransformSRT3d>());
 
 
 	//////////////////////////////////////////////////////////////////////////
@@ -439,7 +439,7 @@ public:
 	 * @param ResultOut newly-created vertices are stored here
 	 * @return index of duplicate vertex
 	 */	
-	int FindOrCreateDuplicateVertex(int VertexID, FMeshIndexMappings& IndexMaps, FDynamicMeshEditResult& ResultOut);
+	GEOMETRYCORE_API int FindOrCreateDuplicateVertex(int VertexID, FMeshIndexMappings& IndexMaps, FDynamicMeshEditResult& ResultOut);
 
 	/**
 	 * Find "new" group for input group under Index mapping, or create new if missing
@@ -448,7 +448,7 @@ public:
 	 * @param ResultOut newly-created groups are stored here
 	 * @return index of duplicate group
 	 */
-	int FindOrCreateDuplicateGroup(int TriangleID, FMeshIndexMappings& IndexMaps, FDynamicMeshEditResult& ResultOut);
+	GEOMETRYCORE_API int FindOrCreateDuplicateGroup(int TriangleID, FMeshIndexMappings& IndexMaps, FDynamicMeshEditResult& ResultOut);
 
 	/**
 	 * Find "new" UV for input UV element under Index mapping, or create new if missing
@@ -457,7 +457,7 @@ public:
 	 * @param IndexMaps source/destination mapping of already-duplicated UVs
 	 * @return index of duplicate UV in given UV layer
 	 */
-	int FindOrCreateDuplicateUV(int ElementID, int UVLayerIndex, FMeshIndexMappings& IndexMaps);
+	GEOMETRYCORE_API int FindOrCreateDuplicateUV(int ElementID, int UVLayerIndex, FMeshIndexMappings& IndexMaps);
 
 	/**
 	 * Find "new" normal for input normal element under Index mapping, or create new if missing
@@ -468,7 +468,7 @@ public:
 	 *   NewNormalOverlayElements must have size > NormalLayerIndex.
 	 * @return index of duplicate normal in given normal layer
 	 */
-	int FindOrCreateDuplicateNormal(int ElementID, int NormalLayerIndex, FMeshIndexMappings& IndexMaps, FDynamicMeshEditResult* ResultOut = nullptr);
+	GEOMETRYCORE_API int FindOrCreateDuplicateNormal(int ElementID, int NormalLayerIndex, FMeshIndexMappings& IndexMaps, FDynamicMeshEditResult* ResultOut = nullptr);
 
 	/**
 	 * Find "new" color for input color element under Index mapping, or create new if missing
@@ -477,7 +477,7 @@ public:
 	 * @param ResultOut any newly created element indices are stored in NewColorOverlayElements here. 
 	 * @return index of duplicate color in given color layer
 	 */
-	int FindOrCreateDuplicateColor(int ElementID, FMeshIndexMappings& IndexMaps, FDynamicMeshEditResult* ResultOut);
+	GEOMETRYCORE_API int FindOrCreateDuplicateColor(int ElementID, FMeshIndexMappings& IndexMaps, FDynamicMeshEditResult* ResultOut);
 	
 	/**
 	 * Copy all attribute-layer values from one triangle to another, using the IndexMaps to track and re-use shared attribute values.
@@ -486,7 +486,7 @@ public:
 	 * @param IndexMaps mappings passed to FindOrCreateDuplicateX functions to track already-created attributes
 	 * @param ResultOut information about new attributes is stored here (@todo finish populating this, at time of writing only normal overlay elements get tracked)
 	 */
-	void CopyAttributes(int FromTriangleID, int ToTriangleID, FMeshIndexMappings& IndexMaps, FDynamicMeshEditResult& ResultOut);
+	GEOMETRYCORE_API void CopyAttributes(int FromTriangleID, int ToTriangleID, FMeshIndexMappings& IndexMaps, FDynamicMeshEditResult& ResultOut);
 
 
 
@@ -498,7 +498,7 @@ public:
 	 * @param PositionTransform optional transformation function applied to mesh vertex positions
 	 * @param NormalTransform optional transformation function applied to mesh normals
 	 */
-	void AppendMesh(const FDynamicMesh3* AppendMesh, FMeshIndexMappings& IndexMapsOut, 
+	GEOMETRYCORE_API void AppendMesh(const FDynamicMesh3* AppendMesh, FMeshIndexMappings& IndexMapsOut, 
 		TFunction<FVector3d(int, const FVector3d&)> PositionTransform = nullptr,
 		TFunction<FVector3d(int, const FVector3d&)> NormalTransform = nullptr);
 
@@ -509,7 +509,7 @@ public:
 	 * @param IndexMapsOut mesh element index mappings generated in this append operation
 	 * @param PositionTransform optional transformation function applied to mesh vertex positions
 	 */
-	void AppendMesh(const TTriangleMeshAdapter<double>* AppendMesh, FMeshIndexMappings& IndexMapsOut,
+	GEOMETRYCORE_API void AppendMesh(const TTriangleMeshAdapter<double>* AppendMesh, FMeshIndexMappings& IndexMapsOut,
 		TFunction<FVector3d(int, const FVector3d&)> PositionTransform = nullptr);
 
 
@@ -525,7 +525,7 @@ public:
 	 * @param NormalTransform optional transformation function applied to mesh normals
 	 * @param NormalMapOut Mapping from element IDs of FromNormals to new element IDs in ToNormals
 	 */
-	void AppendNormals(const FDynamicMesh3* AppendMesh,
+	GEOMETRYCORE_API void AppendNormals(const FDynamicMesh3* AppendMesh,
 		const FDynamicMeshNormalOverlay* FromNormals, FDynamicMeshNormalOverlay* ToNormals,
 		const FIndexMapi& VertexMap, const FIndexMapi& TriangleMap,
 		TFunction<FVector3d(int, const FVector3d&)> NormalTransform,
@@ -543,7 +543,7 @@ public:
 	 * @param TriangleMap map from AppendMesh triangle IDs to triangle IDs applicable to ToUVs (ie of .Mesh)
 	 * @param UVMapOut Mapping from element IDs of FromUVs to new element IDs in ToUVs
 	 */
-	void AppendUVs(const FDynamicMesh3* AppendMesh,
+	GEOMETRYCORE_API void AppendUVs(const FDynamicMesh3* AppendMesh,
 		const FDynamicMeshUVOverlay* FromUVs, FDynamicMeshUVOverlay* ToUVs,
 		const FIndexMapi& VertexMap, const FIndexMapi& TriangleMap,
 		FIndexMapi& UVMapOut);
@@ -559,7 +559,7 @@ public:
 	* @param TriangleMap map from AppendMesh triangle IDs to triangle IDs applicable to ToOverlay (ie of .Mesh)
 	* @param ColorMapOut Mapping from element IDs of FromUVs to new element IDs in ToOverlay
 	*/
-	void AppendColors(const FDynamicMesh3* AppendMesh,
+	GEOMETRYCORE_API void AppendColors(const FDynamicMesh3* AppendMesh,
 		const FDynamicMeshColorOverlay* FromOverlay, FDynamicMeshColorOverlay* ToOverlay,
 		const FIndexMapi& VertexMap, const FIndexMapi& TriangleMap,
 		FIndexMapi& ColorMapOut);
@@ -573,7 +573,7 @@ public:
 	 * @param ResultOut lists of newly created triangles/vertices/etc
 	 * @param bComputeTriangleMap if true, computes the triangle map section of IndexMaps (which is not needed for the append to work, so is optional)
 	 */
-	void AppendTriangles(const FDynamicMesh3* SourceMesh, const TArrayView<const int>& SourceTriangles, FMeshIndexMappings& IndexMaps, FDynamicMeshEditResult& ResultOut, bool bComputeTriangleMap = true);
+	GEOMETRYCORE_API void AppendTriangles(const FDynamicMesh3* SourceMesh, const TArrayView<const int>& SourceTriangles, FMeshIndexMappings& IndexMaps, FDynamicMeshEditResult& ResultOut, bool bComputeTriangleMap = true);
 
 	/**
 	 * Create multiple meshes out of the source mesh by splitting triangles out.
@@ -584,7 +584,7 @@ public:
 	 * @param TriIDToMeshID
 	 * @return true if needed split, false if there were not multiple mesh ids so no split was needed
 	 */
-	static bool SplitMesh(const FDynamicMesh3* SourceMesh, TArray<FDynamicMesh3>& SplitMeshes, TFunctionRef<int(int)> TriIDToMeshID, int DeleteMeshID = -1);
+	static GEOMETRYCORE_API bool SplitMesh(const FDynamicMesh3* SourceMesh, TArray<FDynamicMesh3>& SplitMeshes, TFunctionRef<int(int)> TriIDToMeshID, int DeleteMeshID = -1);
 
 };
 
diff --git a/Engine/Source/Runtime/GeometryCore/Public/DynamicSubmesh3.h b/Engine/Source/Runtime/GeometryCore/Public/DynamicSubmesh3.h
--- a/Engine/Source/Runtime/GeometryCore/Public/DynamicSubmesh3.h
+++ b/Engine/Source/Runtime/GeometryCore/Public/DynamicSubmesh3.h
@@ -13,7 +13,7 @@ namespace UE
 namespace Geometry
 {
 
-struct GEOMETRYCORE_API FDynamicSubmesh3
+struct FDynamicSubmesh3
 {
 protected:
 	const FDynamicMesh3* BaseMesh;
@@ -290,4 +290,4 @@ protected:
 
 
 } // end namespace UE::Geometry
-} // end namespace UE
\ No newline at end of file
+} // end namespace UE
diff --git a/Engine/Source/Runtime/GeometryCore/Public/EdgeLoop.h b/Engine/Source/Runtime/GeometryCore/Public/EdgeLoop.h
--- a/Engine/Source/Runtime/GeometryCore/Public/EdgeLoop.h
+++ b/Engine/Source/Runtime/GeometryCore/Public/EdgeLoop.h
@@ -15,7 +15,7 @@ namespace Geometry
 /**
  * Sequential lists of vertices/edges in a mesh that form a closed loop
  */
-class GEOMETRYCORE_API FEdgeLoop
+class FEdgeLoop
 {
 public:
 	/** The Mesh that contains this EdgeLoop */
@@ -56,13 +56,13 @@ public:
 	/**
 	 * Initialize the loop data
 	 */
-	void Initialize(const FDynamicMesh3* mesh, const TArray<int>& vertices, const TArray<int> & edges, const TArray<int>* BowtieVerticesIn = nullptr);
+	GEOMETRYCORE_API void Initialize(const FDynamicMesh3* mesh, const TArray<int>& vertices, const TArray<int> & edges, const TArray<int>* BowtieVerticesIn = nullptr);
 
 	/**
 	 * Construct an FEdgeLoop from a list of edges of the mesh
 	 * @param EdgesIn list of sequential connected edges
 	 */
-	void InitializeFromEdges(const TArray<int>& EdgesIn);
+	GEOMETRYCORE_API void InitializeFromEdges(const TArray<int>& EdgesIn);
 
 	/**
 	 * Construct an FEdgeLoop from a list of edges of the mesh
@@ -83,7 +83,7 @@ public:
 	 * @param bAutoOrient if true, and any of the edges are boundary edges, we will re-orient the loop to be consistent with boundary edges
 	 * @return false if we found any parts of vertices that are not connected by an edge
 	 */
-	bool InitializeFromVertices(const TArray<int>& VerticesIn, bool bAutoOrient = true);
+	GEOMETRYCORE_API bool InitializeFromVertices(const TArray<int>& VerticesIn, bool bAutoOrient = true);
 
 	/**
 	 * Construct EdgeLoop from list of vertices of mesh
@@ -112,7 +112,7 @@ public:
 	/**
 	 * Populate the BowtieVertices member
 	 */
-	void CalculateBowtieVertices();
+	GEOMETRYCORE_API void CalculateBowtieVertices();
 
 	/**
 	 * @return number of vertices in the loop
@@ -159,7 +159,7 @@ public:
 	/**
 	 * @return bounding box of the vertices of the EdgeLoop
 	 */
-	FAxisAlignedBox3d GetBounds() const;
+	GEOMETRYCORE_API FAxisAlignedBox3d GetBounds() const;
 
 	/**
 	 * Add the vertices of the loop to the Vertices array
@@ -181,7 +181,7 @@ public:
 	 * oriented such that it corresponds with the boundary edges, and if not, reverse it.
 	 * @return true if the loop was reversed
 	 */
-	bool SetCorrectOrientation();
+	GEOMETRYCORE_API bool SetCorrectOrientation();
 
 
 	/**
@@ -197,26 +197,26 @@ public:
 	/**
 	 * @return true if all edges of this loop are internal edges, ie not on the mesh boundary
 	 */
-	bool IsInternalLoop() const;
+	GEOMETRYCORE_API bool IsInternalLoop() const;
 
 
 	/**
 	 * @param TestMesh use this mesh instead of the internal Mesh pointer
 	 * @return true if all edges of this loop are boundary edges
 	 */
-	bool IsBoundaryLoop(const FDynamicMesh3* TestMesh = nullptr) const;
+	GEOMETRYCORE_API bool IsBoundaryLoop(const FDynamicMesh3* TestMesh = nullptr) const;
 
 
 	/**
 	 * @return index of VertexID in the Vertices list, or -1 if not found
 	 */
-	int FindVertexIndex(int VertexID) const;
+	GEOMETRYCORE_API int FindVertexIndex(int VertexID) const;
 
 
 	/**
 	 * @return index of vertex in the Vertices list that is closest to QueryPoint
 	 */
-	int FindNearestVertexIndex(const FVector3d& QueryPoint) const;
+	GEOMETRYCORE_API int FindNearestVertexIndex(const FVector3d& QueryPoint) const;
 
 
 
@@ -224,7 +224,7 @@ public:
 	 * Exhaustively check that verts and edges of this EdgeLoop are consistent. This is quite expensive.
 	 * @return true if loop is consistent/valid
 	 */
-	bool CheckValidity(EValidityCheckFailMode FailMode = EValidityCheckFailMode::Check) const;
+	GEOMETRYCORE_API bool CheckValidity(EValidityCheckFailMode FailMode = EValidityCheckFailMode::Check) const;
 
 
 
@@ -235,7 +235,7 @@ public:
 	 * @param VertexLoop ordered list of vertices
 	 * @param OutEdgeLoop computed list of sequential connected vertices
 	 */
-	static void VertexLoopToEdgeLoop(const FDynamicMesh3* Mesh, const TArray<int>& VertexLoop, TArray<int>& OutEdgeLoop);
+	static GEOMETRYCORE_API void VertexLoopToEdgeLoop(const FDynamicMesh3* Mesh, const TArray<int>& VertexLoop, TArray<int>& OutEdgeLoop);
 
 };
 
diff --git a/Engine/Source/Runtime/GeometryCore/Public/EdgeSpan.h b/Engine/Source/Runtime/GeometryCore/Public/EdgeSpan.h
--- a/Engine/Source/Runtime/GeometryCore/Public/EdgeSpan.h
+++ b/Engine/Source/Runtime/GeometryCore/Public/EdgeSpan.h
@@ -17,7 +17,7 @@ namespace Geometry
  * Sequential lists of vertices/edges in a mesh that is *not* closed.
  * If the list is closed it should be an FEdgeLoop
  */
-class GEOMETRYCORE_API FEdgeSpan
+class FEdgeSpan
 {
 public:
 	/** The Mesh that contains this EdgeSpan */
@@ -58,14 +58,14 @@ public:
 	/**
 	 * Initialize the loop data
 	 */
-	void Initialize(const FDynamicMesh3* mesh, const TArray<int>& vertices, const TArray<int> & edges, const TArray<int>* BowtieVerticesIn = nullptr);
+	GEOMETRYCORE_API void Initialize(const FDynamicMesh3* mesh, const TArray<int>& vertices, const TArray<int> & edges, const TArray<int>* BowtieVerticesIn = nullptr);
 
 
 	/**
 	 * Construct an FEdgeSpan from a list of edges of the mesh
 	 * @param EdgesIn list of sequential connected edges
 	 */
-	void InitializeFromEdges(const TArray<int>& EdgesIn);
+	GEOMETRYCORE_API void InitializeFromEdges(const TArray<int>& EdgesIn);
 
 
 	/**
@@ -100,7 +100,7 @@ public:
 	 * @param bAutoOrient if true, and any of the edges are boundary edges, we will re-orient the span to be consistent with boundary edges
 	 * @return false if we found any parts of VerticesIn that are not connected by an edge
 	 */
-	bool InitializeFromVertices(const TArray<int>& VerticesIn, bool bAutoOrient = true);
+	GEOMETRYCORE_API bool InitializeFromVertices(const TArray<int>& VerticesIn, bool bAutoOrient = true);
 
 
 	/** Set the bowtie vertices */
@@ -114,7 +114,7 @@ public:
 	/**
 	 * Populate the BowtieVertices member
 	 */
-	void CalculateBowtieVertices();
+	GEOMETRYCORE_API void CalculateBowtieVertices();
 
 
 	/**
@@ -144,20 +144,20 @@ public:
 	/**
 	 * @return bounding box of the vertices of the EdgeSpan
 	 */
-	FAxisAlignedBox3d GetBounds() const;
+	GEOMETRYCORE_API FAxisAlignedBox3d GetBounds() const;
 
 
 	/**
 	 * Extract Polyline from Mesh based on vertex list
 	 */
-	void GetPolyline(FPolyline3d& PolylineOut) const;
+	GEOMETRYCORE_API void GetPolyline(FPolyline3d& PolylineOut) const;
 
 	/**
 	 * If any edges if the span are on a mesh boundary, we can check that the span is
 	 * oriented such that it corresponds with the boundary edges, and if not, reverse it.
 	 * @return true if the span was reversed
 	 */
-	bool SetCorrectOrientation();
+	GEOMETRYCORE_API bool SetCorrectOrientation();
 
 
 	/**
@@ -173,30 +173,30 @@ public:
 	/**
 	 * @return true if all edges of this span are internal edges, ie not on the mesh boundary
 	 */
-	bool IsInternalspan() const;
+	GEOMETRYCORE_API bool IsInternalspan() const;
 
 	/**
 	 * @param TestMesh use this mesh instead of the internal Mesh pointer
 	 * @return true if all edges of this span are boundary edges
 	 */
-	bool IsBoundaryspan(const FDynamicMesh3* TestMesh = nullptr) const;
+	GEOMETRYCORE_API bool IsBoundaryspan(const FDynamicMesh3* TestMesh = nullptr) const;
 
 	/**
 	 * @return index of VertexID in the Vertices list, or -1 if not found
 	 */
-	int FindVertexIndex(int VertexID) const;
+	GEOMETRYCORE_API int FindVertexIndex(int VertexID) const;
 
 	/**
 	 * @return index of vertex in the Vertices list that is closest to QueryPoint
 	 */
-	int FindNearestVertexIndex(const FVector3d& QueryPoint) const;
+	GEOMETRYCORE_API int FindNearestVertexIndex(const FVector3d& QueryPoint) const;
 
 	   
 	/**
 	 * Exhaustively check that verts and edges of this EdgeSpan are consistent. This is quite expensive.
 	 * @return true if span is consistent/valid
 	 */
-	bool CheckValidity(EValidityCheckFailMode FailMode = EValidityCheckFailMode::Check) const;
+	GEOMETRYCORE_API bool CheckValidity(EValidityCheckFailMode FailMode = EValidityCheckFailMode::Check) const;
 
 
 
@@ -206,7 +206,7 @@ public:
 	 * @param Vertexspan ordered list of vertices
 	 * @param OutEdgeSpan computed list of sequential connected vertices
 	 */
-	static void VertexSpanToEdgeSpan(const FDynamicMesh3* Mesh, const TArray<int>& VertexSpan, TArray<int>& OutEdgeSpan);
+	static GEOMETRYCORE_API void VertexSpanToEdgeSpan(const FDynamicMesh3* Mesh, const TArray<int>& VertexSpan, TArray<int>& OutEdgeSpan);
 
 };
 
diff --git a/Engine/Source/Runtime/GeometryCore/Public/Generators/DiscMeshGenerator.h b/Engine/Source/Runtime/GeometryCore/Public/Generators/DiscMeshGenerator.h
--- a/Engine/Source/Runtime/GeometryCore/Public/Generators/DiscMeshGenerator.h
+++ b/Engine/Source/Runtime/GeometryCore/Public/Generators/DiscMeshGenerator.h
@@ -15,7 +15,7 @@ namespace Geometry
 /**
  * Generate planar disc (a circle polygon)
  */
-class GEOMETRYCORE_API FDiscMeshGenerator : public FMeshShapeGenerator
+class FDiscMeshGenerator : public FMeshShapeGenerator
 {
 public:
 	/** Radius */
@@ -43,10 +43,10 @@ public:
 	FIndex2i IndicesMap;
 
 public:
-	FDiscMeshGenerator();
+	GEOMETRYCORE_API FDiscMeshGenerator();
 
 	/** Generate the disc */
-	virtual FMeshShapeGenerator& Generate() override;
+	GEOMETRYCORE_API virtual FMeshShapeGenerator& Generate() override;
 
 
 	/** Create vertex at position under IndicesMap, shifted to Origin*/
@@ -62,18 +62,18 @@ public:
 /**
 * Generate planar disc with a hole
 */
-class GEOMETRYCORE_API FPuncturedDiscMeshGenerator : public FDiscMeshGenerator
+class FPuncturedDiscMeshGenerator : public FDiscMeshGenerator
 {
 public:
 	/** Hole Radius */
 	float HoleRadius;
 
 public:
-	FPuncturedDiscMeshGenerator();
+	GEOMETRYCORE_API FPuncturedDiscMeshGenerator();
 
 	/** Generate the disc */
-	virtual FMeshShapeGenerator& Generate() override;
+	GEOMETRYCORE_API virtual FMeshShapeGenerator& Generate() override;
 };
 
 } // end namespace UE::Geometry
-} // end namespace UE
\ No newline at end of file
+} // end namespace UE
diff --git a/Engine/Source/Runtime/GeometryCore/Public/Generators/FlatTriangulationMeshGenerator.h b/Engine/Source/Runtime/GeometryCore/Public/Generators/FlatTriangulationMeshGenerator.h
--- a/Engine/Source/Runtime/GeometryCore/Public/Generators/FlatTriangulationMeshGenerator.h
+++ b/Engine/Source/Runtime/GeometryCore/Public/Generators/FlatTriangulationMeshGenerator.h
@@ -17,7 +17,7 @@ namespace Geometry
 /**
  * Generate planar triangulation of a Polygon.
  */
-class GEOMETRYCORE_API FFlatTriangulationMeshGenerator : public FMeshShapeGenerator
+class FFlatTriangulationMeshGenerator : public FMeshShapeGenerator
 {
 public:
 	/** Vertices of 2D triangulation. */
@@ -36,10 +36,10 @@ public:
 	FIndex2i IndicesMap;
 
 public:
-	FFlatTriangulationMeshGenerator();
+	GEOMETRYCORE_API FFlatTriangulationMeshGenerator();
 
 	/** Generate the triangulation */
-	virtual FMeshShapeGenerator& Generate() override;
+	GEOMETRYCORE_API virtual FMeshShapeGenerator& Generate() override;
 
 
 	/** Create vertex at position under IndicesMap, shifted to Origin*/
@@ -54,4 +54,4 @@ public:
 
 
 } // end namespace UE::Geometry
-} // end namespace UE
\ No newline at end of file
+} // end namespace UE
diff --git a/Engine/Source/Runtime/GeometryCore/Public/Generators/MarchingCubes.h b/Engine/Source/Runtime/GeometryCore/Public/Generators/MarchingCubes.h
--- a/Engine/Source/Runtime/GeometryCore/Public/Generators/MarchingCubes.h
+++ b/Engine/Source/Runtime/GeometryCore/Public/Generators/MarchingCubes.h
@@ -1066,11 +1066,11 @@ protected:
 	*/
 
 
-	GEOMETRYCORE_API constexpr static int EdgeIndices[12][2] = {
+	constexpr static int EdgeIndices[12][2] = {
 		{0,1}, {1,2}, {2,3}, {3,0}, {4,5}, {5,6}, {6,7}, {7,4}, {0,4}, {1,5}, {2,6}, {3,7}
 	};
 
-	GEOMETRYCORE_API constexpr static int EdgeTable[256] = {
+	constexpr static int EdgeTable[256] = {
 		0x0  , 0x109, 0x203, 0x30a, 0x406, 0x50f, 0x605, 0x70c,
 		0x80c, 0x905, 0xa0f, 0xb06, 0xc0a, 0xd03, 0xe09, 0xf00,
 		0x190, 0x99 , 0x393, 0x29a, 0x596, 0x49f, 0x795, 0x69c,
@@ -1105,7 +1105,7 @@ protected:
 		0x70c, 0x605, 0x50f, 0x406, 0x30a, 0x203, 0x109, 0x0   };
 
 
-	GEOMETRYCORE_API constexpr static int TriTable[256][16] =
+	constexpr static int TriTable[256][16] =
 	{
 		{-1, -1, -1, -1, -1, -1, -1, -1, -1, -1, -1, -1, -1, -1, -1, -1},
 		{0, 8, 3, -1, -1, -1, -1, -1, -1, -1, -1, -1, -1, -1, -1, -1},
@@ -1369,4 +1369,4 @@ protected:
 
 
 } // end namespace UE::Geometry
-} // end namespace UE
\ No newline at end of file
+} // end namespace UE
diff --git a/Engine/Source/Runtime/GeometryCore/Public/Generators/PlanarPolygonMeshGenerator.h b/Engine/Source/Runtime/GeometryCore/Public/Generators/PlanarPolygonMeshGenerator.h
--- a/Engine/Source/Runtime/GeometryCore/Public/Generators/PlanarPolygonMeshGenerator.h
+++ b/Engine/Source/Runtime/GeometryCore/Public/Generators/PlanarPolygonMeshGenerator.h
@@ -19,7 +19,7 @@ namespace Geometry
 /**
  * Generate planar triangulation of a Polygon.
  */
-class GEOMETRYCORE_API FPlanarPolygonMeshGenerator : public FMeshShapeGenerator
+class FPlanarPolygonMeshGenerator : public FMeshShapeGenerator
 {
 public:
 	/** Polygon to triangulate. If Polygon has self-intersections or degenerate edges, result is undefined. */
@@ -32,13 +32,13 @@ public:
 	FIndex2i IndicesMap;
 
 public:
-	FPlanarPolygonMeshGenerator();
+	GEOMETRYCORE_API FPlanarPolygonMeshGenerator();
 
 	/** Initialize the polygon from an array of 2D vertices */
-	void SetPolygon(const TArray<FVector2D>& PolygonVerts);
+	GEOMETRYCORE_API void SetPolygon(const TArray<FVector2D>& PolygonVerts);
 
 	/** Generate the triangulation */
-	virtual FMeshShapeGenerator& Generate() override;
+	GEOMETRYCORE_API virtual FMeshShapeGenerator& Generate() override;
 
 
 	/** Create vertex at position under IndicesMap, shifted to Origin*/
@@ -53,4 +53,4 @@ public:
 
 
 } // end namespace UE::Geometry
-} // end namespace UE
\ No newline at end of file
+} // end namespace UE
diff --git a/Engine/Source/Runtime/GeometryCore/Public/Generators/PolygonEdgeMeshGenerator.h b/Engine/Source/Runtime/GeometryCore/Public/Generators/PolygonEdgeMeshGenerator.h
--- a/Engine/Source/Runtime/GeometryCore/Public/Generators/PolygonEdgeMeshGenerator.h
+++ b/Engine/Source/Runtime/GeometryCore/Public/Generators/PolygonEdgeMeshGenerator.h
@@ -14,7 +14,7 @@ namespace Geometry
 {
 
 /** Mesh generator that generates a quad for each edge of a closed polygon */
-class GEOMETRYCORE_API FPolygonEdgeMeshGenerator : public FMeshShapeGenerator
+class FPolygonEdgeMeshGenerator : public FMeshShapeGenerator
 {
 
 private:
@@ -57,7 +57,7 @@ public:
 	/** Number of radial subdivisions for rounded corners */
 	int NumArcVertices;
 
-	FPolygonEdgeMeshGenerator(const TArray<FFrame3d>& InPolygon,
+	GEOMETRYCORE_API FPolygonEdgeMeshGenerator(const TArray<FFrame3d>& InPolygon,
 		bool bInClosed,
 		const TArray<double>& InOffsetScaleFactors,
 		double InWidth = 1.0,
@@ -69,7 +69,7 @@ public:
 
 	// Generate triangulation
 	// TODO: Enable more subdivisions along the width and length dimensions if requested
-	virtual FMeshShapeGenerator& Generate() final;
+	GEOMETRYCORE_API virtual FMeshShapeGenerator& Generate() final;
 
 private:
 
diff --git a/Engine/Source/Runtime/GeometryCore/Public/Generators/RectangleMeshGenerator.h b/Engine/Source/Runtime/GeometryCore/Public/Generators/RectangleMeshGenerator.h
--- a/Engine/Source/Runtime/GeometryCore/Public/Generators/RectangleMeshGenerator.h
+++ b/Engine/Source/Runtime/GeometryCore/Public/Generators/RectangleMeshGenerator.h
@@ -17,7 +17,7 @@ namespace Geometry
  * Generate planar rectangular mesh with variable number of subdivisions along width and height.
  * By default, center of rectangle is centered at (0,0,0) origin
  */
-class GEOMETRYCORE_API FRectangleMeshGenerator : public FMeshShapeGenerator
+class FRectangleMeshGenerator : public FMeshShapeGenerator
 {
 public:
 	/** Rectangle will be translated so that center is at this point */
@@ -45,10 +45,10 @@ public:
 	FIndex2i IndicesMap;
 
 public:
-	FRectangleMeshGenerator();
+	GEOMETRYCORE_API FRectangleMeshGenerator();
 
 	/** Generate the mesh */
-	virtual FMeshShapeGenerator& Generate() override;
+	GEOMETRYCORE_API virtual FMeshShapeGenerator& Generate() override;
 
 	/** Create vertex at position under IndicesMap, shifted to Origin*/
 	virtual FVector3d MakeVertex(int CornerIndex, double x, double y)
@@ -75,7 +75,7 @@ ENUM_CLASS_FLAGS(ERoundedRectangleCorner);
 /**
 * Adds rounded corners to the rectangle mesh
 */
-class GEOMETRYCORE_API FRoundedRectangleMeshGenerator : public FRectangleMeshGenerator
+class FRoundedRectangleMeshGenerator : public FRectangleMeshGenerator
 {
 public:
 	/** Radius of rounded corners */
@@ -106,12 +106,12 @@ public:
 	}
 
 public:
-	FRoundedRectangleMeshGenerator();
+	GEOMETRYCORE_API FRoundedRectangleMeshGenerator();
 
 	/** Generate the mesh */
-	virtual FMeshShapeGenerator& Generate() override;
+	GEOMETRYCORE_API virtual FMeshShapeGenerator& Generate() override;
 };
 
 
 } // end namespace UE::Geometry
-} // end namespace UE
\ No newline at end of file
+} // end namespace UE
diff --git a/Engine/Source/Runtime/GeometryCore/Public/Generators/RevolveGenerator.h b/Engine/Source/Runtime/GeometryCore/Public/Generators/RevolveGenerator.h
--- a/Engine/Source/Runtime/GeometryCore/Public/Generators/RevolveGenerator.h
+++ b/Engine/Source/Runtime/GeometryCore/Public/Generators/RevolveGenerator.h
@@ -24,7 +24,7 @@ namespace Geometry
  * Revolve generators currently depend on the ability to use DynamicMesh operations to (eg) triangulate endcaps/etc,
  * and so they are not currently FMeshShapeGenerators (this may be resolved in future)
  */
-class GEOMETRYCORE_API FBaseRevolveGenerator
+class FBaseRevolveGenerator
 {
 public:
 
@@ -87,7 +87,7 @@ public:
  * Revolve a planar polyline (in the XY plane) around the +Z axis along a circular path
  * (+X, +Y) in the 2D polyline is mapped to (+X, 0, +Z)
  */
-class GEOMETRYCORE_API FRevolvePlanarPathGenerator : public FBaseRevolveGenerator
+class FRevolvePlanarPathGenerator : public FBaseRevolveGenerator
 {
 public:
 	/** Vertices of the planar polyline/path that will be revolved */
@@ -97,7 +97,7 @@ public:
 	bool bCapped = true;
 
 	/** Generate the mesh */
-	FDynamicMesh3 GenerateMesh();
+	GEOMETRYCORE_API FDynamicMesh3 GenerateMesh();
 
 };
 
@@ -105,7 +105,7 @@ public:
 * Revolve a planar polygon (in the XY plane) around the +Z axis along a circular path
 * (+X, +Y) in the 2D polygon is mapped to (+X, 0, +Z)
 */
-class GEOMETRYCORE_API FRevolvePlanarPolygonGenerator : public FBaseRevolveGenerator
+class FRevolvePlanarPolygonGenerator : public FBaseRevolveGenerator
 {
 public:
 	/** Vertices of the planar polygon that will be revolved */
@@ -115,7 +115,7 @@ public:
 	bool bWeldVertsOnAxis = false;
 
 	/** Generate the mesh */
-	FDynamicMesh3 GenerateMesh();
+	GEOMETRYCORE_API FDynamicMesh3 GenerateMesh();
 };
 
 
@@ -123,7 +123,7 @@ public:
 * Revolve a planar polygon (in the XY plane) around the +Z axis along a spiral path.
 * (+X, +Y) in the 2D polygon is mapped to (+X, 0, +Z)
 */
-class GEOMETRYCORE_API FSpiralRevolvePlanarPolygonGenerator : public FBaseRevolveGenerator
+class FSpiralRevolvePlanarPolygonGenerator : public FBaseRevolveGenerator
 {
 public:
 	/** Vertices of the planar polygon that will be revolved */
@@ -136,7 +136,7 @@ public:
 	double RisePerFullRevolution = 50;
 
 	/** Generate the mesh */
-	FDynamicMesh3 GenerateMesh();
+	GEOMETRYCORE_API FDynamicMesh3 GenerateMesh();
 };
 
 
diff --git a/Engine/Source/Runtime/GeometryCore/Public/Generators/StairGenerator.h b/Engine/Source/Runtime/GeometryCore/Public/Generators/StairGenerator.h
--- a/Engine/Source/Runtime/GeometryCore/Public/Generators/StairGenerator.h
+++ b/Engine/Source/Runtime/GeometryCore/Public/Generators/StairGenerator.h
@@ -18,7 +18,7 @@ namespace Geometry
 /**
  * Base Stair mesh generator class.
  */
-class GEOMETRYCORE_API FStairGenerator : public FMeshShapeGenerator
+class FStairGenerator : public FMeshShapeGenerator
 {
 public:
 	/** If true (default), UVs are scaled so that there is no stretching. If false, UVs are scaled to fill unit square */
@@ -52,14 +52,14 @@ public:
 	}
 
 	/** Generate the mesh */
-	virtual FMeshShapeGenerator& Generate() override;
+	GEOMETRYCORE_API virtual FMeshShapeGenerator& Generate() override;
 
 private:
 	/** Generate solid style stairs */
-	FMeshShapeGenerator& GenerateSolidStairs();
+	GEOMETRYCORE_API FMeshShapeGenerator& GenerateSolidStairs();
 
 	/** Generate floating style stairs */
-	FMeshShapeGenerator& GenerateFloatingStairs();
+	GEOMETRYCORE_API FMeshShapeGenerator& GenerateFloatingStairs();
 
 	/**
 	 * Helper to identify the side based on the face index.
@@ -67,7 +67,7 @@ private:
 	 * This method is only valid after the face description array has been
 	 * initialized.
 	 */
-	ESide FaceToSide(int FaceId);
+	GEOMETRYCORE_API ESide FaceToSide(int FaceId);
 
 protected:
 	enum class EStairStyle
@@ -107,7 +107,7 @@ protected:
 
 protected:
 	/** Reset state data on the generator */
-	virtual void ResetData();
+	GEOMETRYCORE_API virtual void ResetData();
 
 	/**
 	 * Returns a vertex position.
@@ -167,7 +167,7 @@ protected:
 /**
  * Generate an oriented Linear Stair mesh.
  */
-class GEOMETRYCORE_API FLinearStairGenerator : public FStairGenerator
+class FLinearStairGenerator : public FStairGenerator
 {
 public:
 	/** The depth of each step. */
@@ -186,16 +186,16 @@ protected:
 	typedef FStairGenerator Super;
 
 protected:
-	virtual FVector3d GenerateVertex(ESide Side, int VertexColumn, int VertexRow) override;
-	virtual FVector3f GenerateNormal(ESide Side, int VertexId) override;
-	virtual FVector2f GenerateUV(ESide Side, int Step, int VertexId, float UVScale) override;
-	virtual float GetMaxDimension() override;
+	GEOMETRYCORE_API virtual FVector3d GenerateVertex(ESide Side, int VertexColumn, int VertexRow) override;
+	GEOMETRYCORE_API virtual FVector3f GenerateNormal(ESide Side, int VertexId) override;
+	GEOMETRYCORE_API virtual FVector2f GenerateUV(ESide Side, int Step, int VertexId, float UVScale) override;
+	GEOMETRYCORE_API virtual float GetMaxDimension() override;
 };
 
 /**
  * Generate an oriented Floating Stair mesh.
  */
-class GEOMETRYCORE_API FFloatingStairGenerator : public FLinearStairGenerator
+class FFloatingStairGenerator : public FLinearStairGenerator
 {
 public:
 	FFloatingStairGenerator()
@@ -210,14 +210,14 @@ protected:
 	typedef FLinearStairGenerator Super;
 
 protected:
-	virtual FVector3d GenerateVertex(ESide Side, int VertexColumn, int VertexRow) override;
-	virtual FVector2f GenerateUV(ESide Side, int Step, int VertexId, float UVScale) override;
+	GEOMETRYCORE_API virtual FVector3d GenerateVertex(ESide Side, int VertexColumn, int VertexRow) override;
+	GEOMETRYCORE_API virtual FVector2f GenerateUV(ESide Side, int Step, int VertexId, float UVScale) override;
 };
 
 /**
  * Generate an oriented Curved Stair mesh.
  */
-class GEOMETRYCORE_API FCurvedStairGenerator : public FStairGenerator
+class FCurvedStairGenerator : public FStairGenerator
 {
 public:
 	/** Inner radius of the curved staircase */
@@ -247,18 +247,18 @@ protected:
 	FVector3f BackNormal = FVector3f::Zero();
 
 protected:
-	virtual void ResetData() override;
-	virtual FVector3d GenerateVertex(ESide Side, int VertexColumn, int VertexRow) override;
-	virtual FVector3f GenerateNormal(ESide Side, int VertexId) override;
-	virtual FVector2f GenerateUV(ESide Side, int Step, int VertexId, float UVScale) override;
-	virtual float GetMaxDimension() override;
+	GEOMETRYCORE_API virtual void ResetData() override;
+	GEOMETRYCORE_API virtual FVector3d GenerateVertex(ESide Side, int VertexColumn, int VertexRow) override;
+	GEOMETRYCORE_API virtual FVector3f GenerateNormal(ESide Side, int VertexId) override;
+	GEOMETRYCORE_API virtual FVector2f GenerateUV(ESide Side, int Step, int VertexId, float UVScale) override;
+	GEOMETRYCORE_API virtual float GetMaxDimension() override;
 };
 
 
 /**
  * Generate an oriented Curved Stair mesh.
  */
-class GEOMETRYCORE_API FSpiralStairGenerator : public FCurvedStairGenerator
+class FSpiralStairGenerator : public FCurvedStairGenerator
 {
 public:
 	FSpiralStairGenerator()
@@ -273,8 +273,8 @@ protected:
 	typedef FCurvedStairGenerator Super;
 
 protected:
-	virtual FVector3d GenerateVertex(ESide Side, int VertexColumn, int VertexRow) override;
-	virtual FVector2f GenerateUV(ESide Side, int Step, int VertexId, float UVScale) override;
+	GEOMETRYCORE_API virtual FVector3d GenerateVertex(ESide Side, int VertexColumn, int VertexRow) override;
+	GEOMETRYCORE_API virtual FVector2f GenerateUV(ESide Side, int Step, int VertexId, float UVScale) override;
 };
 
 } // end namespace UE::Geometry
diff --git a/Engine/Source/Runtime/GeometryCore/Public/Generators/SweepGenerator.h b/Engine/Source/Runtime/GeometryCore/Public/Generators/SweepGenerator.h
--- a/Engine/Source/Runtime/GeometryCore/Public/Generators/SweepGenerator.h
+++ b/Engine/Source/Runtime/GeometryCore/Public/Generators/SweepGenerator.h
@@ -910,7 +910,7 @@ enum class EProfileSweepQuadSplit : uint8
  * it cannot actually use the utility function from FSweepGeneratorBase, and so it doesn't inherit from 
  * that class.
  */
-class GEOMETRYCORE_API FProfileSweepGenerator : public FMeshShapeGenerator
+class FProfileSweepGenerator : public FMeshShapeGenerator
 {
 public:
 
@@ -975,7 +975,7 @@ public:
 public:
 
 	/** Generate the mesh */
-	virtual FMeshShapeGenerator& Generate() override;
+	GEOMETRYCORE_API virtual FMeshShapeGenerator& Generate() override;
 
 	/** If the sweep curve is not closed, this will store the vertex ids of the first and last instances
 	 * of the profile curve. Note that even if the profile curve is closed, depending on the welding,
@@ -989,14 +989,14 @@ public:
 	// a case where we would actually use them.
 protected:
 
-	void InitializeUvBuffer(const TArray<int32>& VertPositionOffsets, 
+	GEOMETRYCORE_API void InitializeUvBuffer(const TArray<int32>& VertPositionOffsets, 
 		int32& NumUvRowsOut, int32& NumUvColumnsOut);
-	void AdjustNormalsForTriangle(int32 TriIndex, int32 FirstIndex, int32 SecondIndex, int32 ThirdIndex,
+	GEOMETRYCORE_API void AdjustNormalsForTriangle(int32 TriIndex, int32 FirstIndex, int32 SecondIndex, int32 ThirdIndex,
 		TArray<FVector3d>& WeightedNormals);
-	void AdjustNormalsForTriangle(int32 TriIndex, int32 FirstIndex, int32 SecondIndex, int32 ThirdIndex,
+	GEOMETRYCORE_API void AdjustNormalsForTriangle(int32 TriIndex, int32 FirstIndex, int32 SecondIndex, int32 ThirdIndex,
 		TArray<FVector3d>& WeightedNormals, const FVector3d& AbNormalized);
 };
 
 
 } // end namespace UE::Geometry
-} // end namespace UE
\ No newline at end of file
+} // end namespace UE
diff --git a/Engine/Source/Runtime/GeometryCore/Public/MathUtil.h b/Engine/Source/Runtime/GeometryCore/Public/MathUtil.h
--- a/Engine/Source/Runtime/GeometryCore/Public/MathUtil.h
+++ b/Engine/Source/Runtime/GeometryCore/Public/MathUtil.h
@@ -18,76 +18,76 @@ template<>
 struct TMathUtilConstants<float>
 {
 	/** Machine Epsilon - float approx 1e-7, double approx 2e-16 */
-	GEOMETRYCORE_API static constexpr float Epsilon = FLT_EPSILON;
+	static constexpr float Epsilon = FLT_EPSILON;
 	/** Zero tolerance for math operations (eg like parallel tests) - float 1e-6, double 1e-8 */
-	GEOMETRYCORE_API static constexpr float ZeroTolerance = 1e-06f;
+	static constexpr float ZeroTolerance = 1e-06f;
 
 	/** largest possible number for type */
-	GEOMETRYCORE_API static constexpr float MaxReal = FLT_MAX;
+	static constexpr float MaxReal = FLT_MAX;
 
 	/** 3.14159... */
-	GEOMETRYCORE_API static constexpr float Pi = 3.1415926535897932384626433832795f;
-	GEOMETRYCORE_API static constexpr float FourPi = 4.0f * Pi;
-	GEOMETRYCORE_API static constexpr float TwoPi = 2.0f*Pi;
-	GEOMETRYCORE_API static constexpr float HalfPi = 0.5f*Pi;
+	static constexpr float Pi = 3.1415926535897932384626433832795f;
+	static constexpr float FourPi = 4.0f * Pi;
+	static constexpr float TwoPi = 2.0f*Pi;
+	static constexpr float HalfPi = 0.5f*Pi;
 
 	/** 1.0 / Pi */
-	GEOMETRYCORE_API static constexpr float InvPi = 1.0f / Pi;
+	static constexpr float InvPi = 1.0f / Pi;
 	/** 1.0 / (2*Pi) */
-	GEOMETRYCORE_API static constexpr float InvTwoPi = 1.0f / TwoPi;
+	static constexpr float InvTwoPi = 1.0f / TwoPi;
 
 	/** pi / 180 */
-	GEOMETRYCORE_API static constexpr float DegToRad = Pi / 180.0f;
+	static constexpr float DegToRad = Pi / 180.0f;
 	/** 180 / pi */
-	GEOMETRYCORE_API static constexpr float RadToDeg = 180.0f / Pi;
+	static constexpr float RadToDeg = 180.0f / Pi;
 
 	//static constexpr float LN_2;
 	//static constexpr float LN_10;
 	//static constexpr float INV_LN_2;
 	//static constexpr float INV_LN_10;
 
-	GEOMETRYCORE_API static constexpr float Sqrt2 = 1.4142135623730950488016887242097f;
-	GEOMETRYCORE_API static constexpr float InvSqrt2 = 1.0f / Sqrt2;
-	GEOMETRYCORE_API static constexpr float Sqrt3 = 1.7320508075688772935274463415059f;
-	GEOMETRYCORE_API static constexpr float InvSqrt3 = 1.0f / Sqrt3;
+	static constexpr float Sqrt2 = 1.4142135623730950488016887242097f;
+	static constexpr float InvSqrt2 = 1.0f / Sqrt2;
+	static constexpr float Sqrt3 = 1.7320508075688772935274463415059f;
+	static constexpr float InvSqrt3 = 1.0f / Sqrt3;
 };
 
 template<>
 struct TMathUtilConstants<double>
 {
 	/** Machine Epsilon - float approx 1e-7, double approx 2e-16 */
-	GEOMETRYCORE_API static constexpr double Epsilon = DBL_EPSILON;
+	static constexpr double Epsilon = DBL_EPSILON;
 	/** Zero tolerance for math operations (eg like parallel tests) - float 1e-6, double 1e-8 */
-	GEOMETRYCORE_API static constexpr double ZeroTolerance = 1e-08;
+	static constexpr double ZeroTolerance = 1e-08;
 
 	/** largest possible number for type */
-	GEOMETRYCORE_API static constexpr double MaxReal = DBL_MAX;
+	static constexpr double MaxReal = DBL_MAX;
 
 	/** 3.14159... */
-	GEOMETRYCORE_API static constexpr double Pi = 3.1415926535897932384626433832795;
-	GEOMETRYCORE_API static constexpr double FourPi = 4.0 * Pi;
-	GEOMETRYCORE_API static constexpr double TwoPi = 2.0 * Pi;
-	GEOMETRYCORE_API static constexpr double HalfPi = 0.5 * Pi;
+	static constexpr double Pi = 3.1415926535897932384626433832795;
+	static constexpr double FourPi = 4.0 * Pi;
+	static constexpr double TwoPi = 2.0 * Pi;
+	static constexpr double HalfPi = 0.5 * Pi;
 
 	/** 1.0 / Pi */
-	GEOMETRYCORE_API static constexpr double InvPi = 1.0 / Pi;
+	static constexpr double InvPi = 1.0 / Pi;
 	/** 1.0 / (2*Pi) */
-	GEOMETRYCORE_API static constexpr double InvTwoPi = 1.0 / TwoPi;
+	static constexpr double InvTwoPi = 1.0 / TwoPi;
 
 	/** pi / 180 */
-	GEOMETRYCORE_API static constexpr double DegToRad = Pi / 180.0;
+	static constexpr double DegToRad = Pi / 180.0;
 	/** 180 / pi */
-	GEOMETRYCORE_API static constexpr double RadToDeg = 180.0 / Pi;
+	static constexpr double RadToDeg = 180.0 / Pi;
 
 	//static constexpr double LN_2;
 	//static constexpr double LN_10;
 	//static constexpr double INV_LN_2;
 	//static constexpr double INV_LN_10;
 
-	GEOMETRYCORE_API static constexpr double Sqrt2 = 1.4142135623730950488016887242097;
-	GEOMETRYCORE_API static constexpr double InvSqrt2 = 1.0 / Sqrt2;
-	GEOMETRYCORE_API static constexpr double Sqrt3 = 1.7320508075688772935274463415059;
-	GEOMETRYCORE_API static constexpr double InvSqrt3 = 1.0 / Sqrt3;
+	static constexpr double Sqrt2 = 1.4142135623730950488016887242097;
+	static constexpr double InvSqrt2 = 1.0 / Sqrt2;
+	static constexpr double Sqrt3 = 1.7320508075688772935274463415059;
+	static constexpr double InvSqrt3 = 1.0 / Sqrt3;
 };
 
 
@@ -95,21 +95,21 @@ struct TMathUtilConstants<double>
 template<>
 struct TMathUtilConstants<int32>
 {
-	GEOMETRYCORE_API static constexpr int32 Epsilon = 0;
-	GEOMETRYCORE_API static constexpr int32 ZeroTolerance = 0;
-	GEOMETRYCORE_API static constexpr int32 MaxReal = ((int32)0x7fffffff);
-	GEOMETRYCORE_API static constexpr int32 Pi = 3;
-	GEOMETRYCORE_API static constexpr int32 FourPi = 4 * Pi;
-	GEOMETRYCORE_API static constexpr int32 TwoPi = 2 * Pi;
-	GEOMETRYCORE_API static constexpr int32 HalfPi = 1;
-	GEOMETRYCORE_API static constexpr int32 InvPi = 1;
-	GEOMETRYCORE_API static constexpr int32 InvTwoPi = 1;
-	GEOMETRYCORE_API static constexpr int32 DegToRad = 1;
-	GEOMETRYCORE_API static constexpr int32 RadToDeg = 1;
-	GEOMETRYCORE_API static constexpr int32 Sqrt2 = 1;
-	GEOMETRYCORE_API static constexpr int32 InvSqrt2 = 1;
-	GEOMETRYCORE_API static constexpr int32 Sqrt3 = 2;
-	GEOMETRYCORE_API static constexpr int32 InvSqrt3 = 1;
+	static constexpr int32 Epsilon = 0;
+	static constexpr int32 ZeroTolerance = 0;
+	static constexpr int32 MaxReal = ((int32)0x7fffffff);
+	static constexpr int32 Pi = 3;
+	static constexpr int32 FourPi = 4 * Pi;
+	static constexpr int32 TwoPi = 2 * Pi;
+	static constexpr int32 HalfPi = 1;
+	static constexpr int32 InvPi = 1;
+	static constexpr int32 InvTwoPi = 1;
+	static constexpr int32 DegToRad = 1;
+	static constexpr int32 RadToDeg = 1;
+	static constexpr int32 Sqrt2 = 1;
+	static constexpr int32 InvSqrt2 = 1;
+	static constexpr int32 Sqrt3 = 2;
+	static constexpr int32 InvSqrt3 = 1;
 };
 
 
@@ -117,21 +117,21 @@ struct TMathUtilConstants<int32>
 template<>
 struct TMathUtilConstants<int64>
 {
-	GEOMETRYCORE_API static constexpr int64 Epsilon = 0;
-	GEOMETRYCORE_API static constexpr int64 ZeroTolerance = 0;
-	GEOMETRYCORE_API static constexpr int64 MaxReal = ((int64)0x7fffffffffffffff);
-	GEOMETRYCORE_API static constexpr int64 Pi = 3;
-	GEOMETRYCORE_API static constexpr int64 FourPi = 4 * Pi;
-	GEOMETRYCORE_API static constexpr int64 TwoPi = 2 * Pi;
-	GEOMETRYCORE_API static constexpr int64 HalfPi = 1;
-	GEOMETRYCORE_API static constexpr int64 InvPi = 1;
-	GEOMETRYCORE_API static constexpr int64 InvTwoPi = 1;
-	GEOMETRYCORE_API static constexpr int64 DegToRad = 1;
-	GEOMETRYCORE_API static constexpr int64 RadToDeg = 1;
-	GEOMETRYCORE_API static constexpr int64 Sqrt2 = 1;
-	GEOMETRYCORE_API static constexpr int64 InvSqrt2 = 1;
-	GEOMETRYCORE_API static constexpr int64 Sqrt3 = 2;
-	GEOMETRYCORE_API static constexpr int64 InvSqrt3 = 1;
+	static constexpr int64 Epsilon = 0;
+	static constexpr int64 ZeroTolerance = 0;
+	static constexpr int64 MaxReal = ((int64)0x7fffffffffffffff);
+	static constexpr int64 Pi = 3;
+	static constexpr int64 FourPi = 4 * Pi;
+	static constexpr int64 TwoPi = 2 * Pi;
+	static constexpr int64 HalfPi = 1;
+	static constexpr int64 InvPi = 1;
+	static constexpr int64 InvTwoPi = 1;
+	static constexpr int64 DegToRad = 1;
+	static constexpr int64 RadToDeg = 1;
+	static constexpr int64 Sqrt2 = 1;
+	static constexpr int64 InvSqrt2 = 1;
+	static constexpr int64 Sqrt3 = 2;
+	static constexpr int64 InvSqrt3 = 1;
 };
 
 
diff --git a/Engine/Source/Runtime/GeometryCore/Public/MeshBoundaryLoops.h b/Engine/Source/Runtime/GeometryCore/Public/MeshBoundaryLoops.h
--- a/Engine/Source/Runtime/GeometryCore/Public/MeshBoundaryLoops.h
+++ b/Engine/Source/Runtime/GeometryCore/Public/MeshBoundaryLoops.h
@@ -15,7 +15,7 @@ namespace UE
 namespace Geometry
 {
 
-class GEOMETRYCORE_API FMeshBoundaryLoops
+class FMeshBoundaryLoops
 {
 public:
 	/** Mesh we are finding loops on */
@@ -76,7 +76,7 @@ public:
 	 * .SpanBehavior and .FailureBehavior control what happens if we run into problem cases
 	 * @return false if errors occurred, in this case output set is incomplete
 	 */
-	bool Compute();
+	GEOMETRYCORE_API bool Compute();
 
 
 
@@ -100,35 +100,35 @@ public:
 
 
 	/** @return index of loop with maximum number of vertices */
-	int GetMaxVerticesLoopIndex() const;
+	GEOMETRYCORE_API int GetMaxVerticesLoopIndex() const;
 
 	/** @return index of loop with longest arc length, or -1 if no loops */
-	int GetLongestLoopIndex() const;
+	GEOMETRYCORE_API int GetLongestLoopIndex() const;
 
 	/**
 	 * @return pair (LoopIndex,VertexIndexInLoop) of VertexID in EdgeLoops, or FIndex2i::Invalid if not found
 	 */
-	FIndex2i FindVertexInLoop(int VertexID) const;
+	GEOMETRYCORE_API FIndex2i FindVertexInLoop(int VertexID) const;
 
 	/**
 	 * @return index of loop that contains vertex, or -1 if not found
 	 */
-	int FindLoopContainingVertex(int VertexID) const;
+	GEOMETRYCORE_API int FindLoopContainingVertex(int VertexID) const;
 
 	/**
 	 * @return index of loop that contains edge, or -1 if not found
 	 */
-	int FindLoopContainingEdge(int EdgeID) const;
+	GEOMETRYCORE_API int FindLoopContainingEdge(int EdgeID) const;
 
 	/**
 	 * @return index of loop that best matches input triangles, or -1 if not found
 	 */
-	int FindLoopTrianglesHint(const TArray<int>& BorderHintTris) const;
+	GEOMETRYCORE_API int FindLoopTrianglesHint(const TArray<int>& BorderHintTris) const;
 
 	/**
 	 * @return index of loop that best matches input edges, or -1 if not found
 	 */
-	int FindLoopEdgesHint(const TSet<int>& BorderHintEdges) const;
+	GEOMETRYCORE_API int FindLoopEdgesHint(const TSet<int>& BorderHintEdges) const;
 
 
 protected:
@@ -137,12 +137,12 @@ protected:
 
 	// [TODO] cache this : a dictionary? we will not need very many, but we will
 	//   need each multiple times!
-	FVector3d GetVertexNormal(int vid);
+	GEOMETRYCORE_API FVector3d GetVertexNormal(int vid);
 
 	// ok, bdry_edges[0...bdry_edges_count] contains the boundary edges coming out of bowtie_v.
 	// We want to pick the best one to continue the loop that came in to bowtie_v on incoming_e.
 	// If the loops are all sane, then we will get the smallest loops by "turning left" at bowtie_v.
-	int FindLeftTurnEdge(int incoming_e, int bowtie_v, TArray<int>& bdry_edges, int bdry_edges_count, TArray<bool>& used_edges);
+	GEOMETRYCORE_API int FindLeftTurnEdge(int incoming_e, int bowtie_v, TArray<int>& bdry_edges, int bdry_edges_count, TArray<bool>& used_edges);
 	
 	struct Subloops
 	{
@@ -154,7 +154,7 @@ protected:
 	// This is called when loopV contains one or more "bowtie" vertices.
 	// These vertices *might* be duplicated in loopV (but not necessarily)
 	// If they are, we have to break loopV into subloops that don't contain duplicates.
-	bool ExtractSubloops(TArray<int>& loopV, TArray<int>& loopE, TArray<int>& bowties, Subloops& SubloopsOut);
+	GEOMETRYCORE_API bool ExtractSubloops(TArray<int>& loopV, TArray<int>& loopE, TArray<int>& bowties, Subloops& SubloopsOut);
 
 
 
@@ -171,28 +171,28 @@ protected:
 	// Check if the loop from bowtieV to bowtieV inside loopV contains any other bowtie verts.
 	// Also returns start and end indices in loopV of "clean" loop
 	// Note that start may be < end, if the "clean" loop wraps around the end
-	static bool IsSimpleBowtieLoop(const TArray<int>& LoopVerts, const TArray<int>& BowtieVerts, int BowtieVertex, int& start_i, int& end_i);
+	static GEOMETRYCORE_API bool IsSimpleBowtieLoop(const TArray<int>& LoopVerts, const TArray<int>& BowtieVerts, int BowtieVertex, int& start_i, int& end_i);
 
 	// check if forward path from loopV[i1] to loopV[i2] contains any bowtie verts other than bowtieV
-	static bool IsSimplePath(const TArray<int>& LoopVerts, const TArray<int>& BowtieVerts, int BowtieVertex, int i1, int i2);
+	static GEOMETRYCORE_API bool IsSimplePath(const TArray<int>& LoopVerts, const TArray<int>& BowtieVerts, int BowtieVertex, int i1, int i2);
 
 
 	// Read out the span from loop[i0] to loop [i1-1] into an array.
 	// If bMarkInvalid, then these values are set to -1 in loop
-	static void ExtractSpan(TArray<int>& Loop, int i0, int i1, bool bMarkInvalid, TArray<int>& OutSpan);
+	static GEOMETRYCORE_API void ExtractSpan(TArray<int>& Loop, int i0, int i1, bool bMarkInvalid, TArray<int>& OutSpan);
 
 	// count number of valid vertices in l between loop[i0] and loop[i1-1]
-	static int CountSpan(const TArray<int>& Loop, int i0, int i1);
+	static GEOMETRYCORE_API int CountSpan(const TArray<int>& Loop, int i0, int i1);
 
 	// find the index of item in loop, starting at start index
-	static int FindIndex(const TArray<int>& Loop, int Start, int Item);
+	static GEOMETRYCORE_API int FindIndex(const TArray<int>& Loop, int Start, int Item);
 
 	// count number of times item appears in loop
-	static int CountInList(const TArray<int>& Loop, int Item);
+	static GEOMETRYCORE_API int CountInList(const TArray<int>& Loop, int Item);
 
 
 };
 
 
 } // end namespace UE::Geometry
-} // end namespace UE
\ No newline at end of file
+} // end namespace UE
diff --git a/Engine/Source/Runtime/GeometryCore/Public/MeshRegionBoundaryLoops.h b/Engine/Source/Runtime/GeometryCore/Public/MeshRegionBoundaryLoops.h
--- a/Engine/Source/Runtime/GeometryCore/Public/MeshRegionBoundaryLoops.h
+++ b/Engine/Source/Runtime/GeometryCore/Public/MeshRegionBoundaryLoops.h
@@ -18,7 +18,7 @@ namespace Geometry
  * Extract FEdgeLoops on the boundary of a set of triangles of a mesh.
  * @todo this was an early port and possibly can be optimized
  */
-class GEOMETRYCORE_API FMeshRegionBoundaryLoops
+class FMeshRegionBoundaryLoops
 {
 public:
 
@@ -37,15 +37,15 @@ public:
 public:
 	FMeshRegionBoundaryLoops() {}
 
-	FMeshRegionBoundaryLoops(const FDynamicMesh3* MeshIn, const TArray<int>& RegionTris, bool bAutoCompute = true);
+	GEOMETRYCORE_API FMeshRegionBoundaryLoops(const FDynamicMesh3* MeshIn, const TArray<int>& RegionTris, bool bAutoCompute = true);
 
-	void SetMesh(const FDynamicMesh3* MeshIn, const TArray<int>& RegionTris);
+	GEOMETRYCORE_API void SetMesh(const FDynamicMesh3* MeshIn, const TArray<int>& RegionTris);
 
 	/**
 	 * Find set of FEdgeLoops on the border of the input triangle set
 	 * @return false if errors occurred, in this case output set is incomplete
 	 */
-	bool Compute();
+	GEOMETRYCORE_API bool Compute();
 
 
 	/** @return number of loops found by Compute() */
@@ -72,7 +72,7 @@ public:
 	}
 
 	/** @return index of loop with maximum number of vertices */
-	int GetMaxVerticesLoopIndex() const;
+	GEOMETRYCORE_API int GetMaxVerticesLoopIndex() const;
 	
 	template<typename ElementType>
 	using ElementIDAndValue = TPair<int32, ElementType>;
@@ -113,7 +113,7 @@ public:
 	* @param Loop output loop will be stored here. This value is garbage if false is returned.
 	* @return true if a single well-formed loop was found, false if non-manifold or failure case encountered
 	*/
-	static bool GetTriangleSetBoundaryLoop(const FDynamicMesh3& Mesh, const TArray<int32>& Tris, FEdgeLoop& Loop);
+	static GEOMETRYCORE_API bool GetTriangleSetBoundaryLoop(const FDynamicMesh3& Mesh, const TArray<int32>& Tris, FEdgeLoop& Loop);
 
 protected:
 
@@ -126,33 +126,33 @@ protected:
 
 	// returns true for both internal and mesh boundary edges
 	// tid_in and tid_out are triangles 'in' and 'out' of set, respectively
-	bool IsEdgeOnBoundary(int eid, int& tid_in, int& tid_out) const;
+	GEOMETRYCORE_API bool IsEdgeOnBoundary(int eid, int& tid_in, int& tid_out) const;
 
 	// return same indices as GetEdgeV, but oriented based on attached triangle
-	FIndex2i GetOrientedEdgeVerts(int eID, int tid_in);
+	GEOMETRYCORE_API FIndex2i GetOrientedEdgeVerts(int eID, int tid_in);
 
 	// returns first two boundary edges, and count of total boundary edges
-	int GetVertexBoundaryEdges(int vID, int& e0, int& e1);
+	GEOMETRYCORE_API int GetVertexBoundaryEdges(int vID, int& e0, int& e1);
 
 	// e needs to be large enough (ie call GetVtxBoundaryEdges, or as large as max one-ring)
 	// returns count, ie number of elements of e that were filled
-	int GetAllVertexBoundaryEdges(int vID, TArray<int>& e);
+	GEOMETRYCORE_API int GetAllVertexBoundaryEdges(int vID, TArray<int>& e);
 
 
 	// [TODO] cache this : a dictionary? we will not need very many, but we will
 	//   need each multiple times!
-	FVector3d GetVertexNormal(int vid);
+	GEOMETRYCORE_API FVector3d GetVertexNormal(int vid);
 
 	// ok, bdry_edges[0...bdry_edges_count] contains the boundary edges coming out of bowtie_v.
 	// We want to pick the best one to continue the loop that came : to bowtie_v on incoming_e.
 	// If the loops are all sane, then we will get the smallest loops by "turning left" at bowtie_v.
-	int FindLeftTurnEdge(int incoming_e, int bowtie_v, TArray<int>& bdry_edges, int bdry_edges_count, const FIndexFlagSet& used_edges);
+	GEOMETRYCORE_API int FindLeftTurnEdge(int incoming_e, int bowtie_v, TArray<int>& bdry_edges, int bdry_edges_count, const FIndexFlagSet& used_edges);
 
 
 	// This is called when loopV contains one or more "bowtie" vertices.
 	// These vertices *might* be duplicated : loopV (but not necessarily)
 	// If they are, we have to break loopV into subloops that don't contain duplicates.
-	bool TryExtractSubloops(TArray<int>& loopV, const TArray<int>& loopE, const TArray<int>& bowties, TArray<FEdgeLoop>& SubLoopsOut);
+	GEOMETRYCORE_API bool TryExtractSubloops(TArray<int>& loopV, const TArray<int>& loopE, const TArray<int>& bowties, TArray<FEdgeLoop>& SubLoopsOut);
 
 
 };
diff --git a/Engine/Source/Runtime/GeometryCore/Public/Operations/EmbedSurfacePath.h b/Engine/Source/Runtime/GeometryCore/Public/Operations/EmbedSurfacePath.h
--- a/Engine/Source/Runtime/GeometryCore/Public/Operations/EmbedSurfacePath.h
+++ b/Engine/Source/Runtime/GeometryCore/Public/Operations/EmbedSurfacePath.h
@@ -27,7 +27,7 @@ enum class ESurfacePointType
 	Triangle = 2
 };
 
-struct GEOMETRYCORE_API FMeshSurfacePoint
+struct FMeshSurfacePoint
 {
 	int ElementID;
 	FVector3d BaryCoord;
@@ -46,7 +46,7 @@ struct GEOMETRYCORE_API FMeshSurfacePoint
 	{
 	}
 
-	FVector3d Pos(const FDynamicMesh3 *Mesh) const;
+	GEOMETRYCORE_API FVector3d Pos(const FDynamicMesh3 *Mesh) const;
 };
 
 /**
@@ -61,7 +61,7 @@ struct GEOMETRYCORE_API FMeshSurfacePoint
 /**
  * Represent a path on the surface of a mesh via barycentric coordinates and triangle references
  */
-class GEOMETRYCORE_API FMeshSurfacePath
+class FMeshSurfacePath
 {
 public:
 	FDynamicMesh3 *Mesh;
@@ -84,7 +84,7 @@ public:
 	/**
 	 * @return True if the Path exactly sticks to the mesh surface, and never jumps to disconnected elements
 	 */
-	bool IsConnected() const;
+	GEOMETRYCORE_API bool IsConnected() const;
 
 	bool IsClosed() const
 	{
@@ -97,7 +97,7 @@ public:
 		bIsClosed = false;
 	}
 
-	bool AddViaPlanarWalk(int StartTri, int StartVID, FVector3d StartPt, int EndTri, int EndVertID, FVector3d EndPt,
+	GEOMETRYCORE_API bool AddViaPlanarWalk(int StartTri, int StartVID, FVector3d StartPt, int EndTri, int EndVertID, FVector3d EndPt,
 		FVector3d WalkPlaneNormal, TFunction<FVector3d(const FDynamicMesh3*, int)> VertexToPosnFn = nullptr,
 		bool bAllowBackwardsSearch = true, double AcceptEndPtOutsideDist = FMathd::ZeroTolerance, 
 		double PtOnPlaneThresholdSq = FMathf::ZeroTolerance*100, double BackwardsTolerance = FMathd::ZeroTolerance*10);
@@ -108,7 +108,7 @@ public:
 	 *
 	 * @return true if succeeded
 	 */
-	bool ClosePath();
+	GEOMETRYCORE_API bool ClosePath();
 
 	/**
 	* @return EOperationValidationResult::Ok if we can apply operation, or error code if we cannot
@@ -141,7 +141,7 @@ public:
 	 * @param PathVertices Indices of the vertices on the path after embedding succeeds; NOTE these will not be 1:1 with the input Path
 	 * @return true if embedding succeeded.
 	 */
-	bool EmbedSimplePath(bool bUpdatePath, TArray<int>& PathVertices, bool bDoNotDuplicateFirstVertexID = true, double SnapElementThresholdSq = FMathf::ZeroTolerance*100);
+	GEOMETRYCORE_API bool EmbedSimplePath(bool bUpdatePath, TArray<int>& PathVertices, bool bDoNotDuplicateFirstVertexID = true, double SnapElementThresholdSq = FMathf::ZeroTolerance*100);
 
 	// TODO: add functionality to delete 'inside' of path -- but how do we determine what is inside?  ref MeshFacesFromLoop.cs in geometry3sharp
 };
@@ -158,4 +158,4 @@ bool GEOMETRYCORE_API EmbedProjectedPaths(FDynamicMesh3* Mesh, const TArrayView<
 
 
 } // end namespace UE::Geometry
-} // end namespace UE
\ No newline at end of file
+} // end namespace UE
diff --git a/Engine/Source/Runtime/GeometryCore/Public/Operations/LocalPlanarSimplify.h b/Engine/Source/Runtime/GeometryCore/Public/Operations/LocalPlanarSimplify.h
--- a/Engine/Source/Runtime/GeometryCore/Public/Operations/LocalPlanarSimplify.h
+++ b/Engine/Source/Runtime/GeometryCore/Public/Operations/LocalPlanarSimplify.h
@@ -22,7 +22,7 @@ class FDynamicMesh3;
 // Class to support local planar simplification, to reduce excess edges created by plane cuts, edge loop insertions, and similar operations
 // (Note a custom version of this logic also exists on the FMeshBoolean class)
 //
-class GEOMETRYCORE_API FLocalPlanarSimplify
+class FLocalPlanarSimplify
 {
 public:
 	// Simplification settings
@@ -55,7 +55,7 @@ public:
 	 * @param InOutEdges		Edges to simplify along -- will be updated to remove any edges that have been collapsed by simplification
 	 * @param ProcessCollapse	Optional function to be called whenever the simplification collapses an edge
 	 */
-	void SimplifyAlongEdges(FDynamicMesh3& Mesh, TSet<int32>& InOutEdges, TUniqueFunction<void(const DynamicMeshInfo::FEdgeCollapseInfo&)> ProcessCollapse = nullptr) const;
+	GEOMETRYCORE_API void SimplifyAlongEdges(FDynamicMesh3& Mesh, TSet<int32>& InOutEdges, TUniqueFunction<void(const DynamicMeshInfo::FEdgeCollapseInfo&)> ProcessCollapse = nullptr) const;
 
 
 
@@ -72,7 +72,7 @@ public:
 	 * @param The normal of the first triangle attached to the vertex.
 	 * @return Whether all the triangles were coplanar
 	 */
-	static bool IsFlat(const FDynamicMesh3& Mesh, int VID, double DotTolerance, FVector3d& OutFirstNormal);
+	static GEOMETRYCORE_API bool IsFlat(const FDynamicMesh3& Mesh, int VID, double DotTolerance, FVector3d& OutFirstNormal);
 
 	/**
 	 * Test if the triangles connected to a vertex could be flattened by folding along the given edge
@@ -83,7 +83,7 @@ public:
 	 * @param The normal of the first triangle attached to the vertex.
 	 * @return Whether all the triangles were coplanar
 	 */
-	static bool IsDevelopableAlongEdge(const FDynamicMesh3& Mesh, int EID, int VID, double DotTolerance, FVector3d& NormalA, bool& bIsFlat);
+	static GEOMETRYCORE_API bool IsDevelopableAlongEdge(const FDynamicMesh3& Mesh, int EID, int VID, double DotTolerance, FVector3d& NormalA, bool& bIsFlat);
 
 	/**
 	 * Test if a given edge collapse would cause a triangle flip or other unacceptable decrease in mesh quality
@@ -97,7 +97,7 @@ public:
 	 * @param TryToImproveTriQualityThreshold Threshold for triangle quality (see comment for class member, above)
 	 * @param bHasMultipleNormals If false, assume the triangulation is locally flat, so we can use ExpectNormal instead of recomputing the expected normal
 	 */
-	static bool CollapseWouldHurtTriangleQuality(
+	static GEOMETRYCORE_API bool CollapseWouldHurtTriangleQuality(
 		const FDynamicMesh3& Mesh, const FVector3d& RemoveVNormal,
 		int32 RemoveV, const FVector3d& RemoveVPos, int32 KeepV, const FVector3d& KeepVPos,
 		double TryToImproveTriQualityThreshold, bool bHasMultipleNormals = false
@@ -106,7 +106,7 @@ public:
 	/**
 	 * Test if a given edge collapse would change the mesh shape, mesh triangle group shape, or UVs unacceptably
 	 */
-	static bool CollapseWouldChangeShapeOrUVs(
+	static GEOMETRYCORE_API bool CollapseWouldChangeShapeOrUVs(
 		const FDynamicMesh3& Mesh, const TSet<int>& PathEdgeSet, double DotTolerance, int SourceEID,
 		int32 RemoveV, const FVector3d& RemoveVPos, int32 KeepV, const FVector3d& KeepVPos,
 		const FVector3d& EdgeDir, bool bPreserveTriangleGroups, bool bPreserveUVsForMesh,
@@ -116,4 +116,4 @@ public:
 };
 
 } // end namespace UE::Geometry
-} // end namespace UE
\ No newline at end of file
+} // end namespace UE
diff --git a/Engine/Source/Runtime/GeometryCore/Public/Operations/MeshBoolean.h b/Engine/Source/Runtime/GeometryCore/Public/Operations/MeshBoolean.h
--- a/Engine/Source/Runtime/GeometryCore/Public/Operations/MeshBoolean.h
+++ b/Engine/Source/Runtime/GeometryCore/Public/Operations/MeshBoolean.h
@@ -26,7 +26,7 @@ namespace Geometry
 /**
  * MeshBoolean -- perform a boolean operation on two input meshes.
  */
-class GEOMETRYCORE_API FMeshBoolean
+class FMeshBoolean
 {
 public:
 
@@ -178,7 +178,7 @@ public:
 	 * Compute the Boolean operation
 	 * @return true if operation succeeds
 	 */
-	bool Compute();
+	GEOMETRYCORE_API bool Compute();
 
 protected:
 	/** If this returns true, abort computation.  */
@@ -189,14 +189,14 @@ protected:
 
 private:
 
-	int FindNearestEdge(const FDynamicMesh3& OnMesh, const TArray<int>& EIDs, FVector3d Pos);
+	GEOMETRYCORE_API int FindNearestEdge(const FDynamicMesh3& OnMesh, const TArray<int>& EIDs, FVector3d Pos);
 
-	bool MergeEdges(const FMeshIndexMappings& IndexMaps, FDynamicMesh3* CutMesh[2], const TArray<int> CutBoundaryEdges[2], const TMap<int, int>& AllVIDMatches);
+	GEOMETRYCORE_API bool MergeEdges(const FMeshIndexMappings& IndexMaps, FDynamicMesh3* CutMesh[2], const TArray<int> CutBoundaryEdges[2], const TMap<int, int>& AllVIDMatches);
 
-	void SimplifyAlongNewEdges(int NumMeshesToProcess, FDynamicMesh3* CutMesh[2], TArray<int> CutBoundaryEdges[2], TMap<int, int>& AllVIDMatches);
+	GEOMETRYCORE_API void SimplifyAlongNewEdges(int NumMeshesToProcess, FDynamicMesh3* CutMesh[2], TArray<int> CutBoundaryEdges[2], TMap<int, int>& AllVIDMatches);
 
 };
 
 
 } // end namespace UE::Geometry
-} // end namespace UE
\ No newline at end of file
+} // end namespace UE
diff --git a/Engine/Source/Runtime/GeometryCore/Public/Operations/MeshMeshCut.h b/Engine/Source/Runtime/GeometryCore/Public/Operations/MeshMeshCut.h
--- a/Engine/Source/Runtime/GeometryCore/Public/Operations/MeshMeshCut.h
+++ b/Engine/Source/Runtime/GeometryCore/Public/Operations/MeshMeshCut.h
@@ -19,7 +19,7 @@ namespace UE
 namespace Geometry
 {
 
-class GEOMETRYCORE_API FMeshSelfCut
+class FMeshSelfCut
 {
 public:
 	//
@@ -68,7 +68,7 @@ public:
 	 * Split mesh along the provided intersections
 	 * @param Intersections A set of mesh intersections, for example as returned by TMeshAABBTree3::FindAllSelfIntersections
 	 */
-	bool Cut(const MeshIntersection::FIntersectionsQueryResult& Intersections);
+	GEOMETRYCORE_API bool Cut(const MeshIntersection::FIntersectionsQueryResult& Intersections);
 
 	void ResetOutputs()
 	{
@@ -82,7 +82,7 @@ public:
 /**
  * Cut a mesh where it crosses a second mesh -- resolving all intersections, but not deleting geometry.  Optionally resolve intersections mutually for both meshes.
  */
-class GEOMETRYCORE_API FMeshMeshCut
+class FMeshMeshCut
 {
 public:
 
@@ -144,7 +144,7 @@ public:
 	 * Split mesh(es) along the provided intersections
 	 * @param Intersections A set of mesh-mesh intersections, for example as returned by TMeshAABBTree3::FindAllIntersections
 	 */
-	bool Cut(const MeshIntersection::FIntersectionsQueryResult& Intersections);
+	GEOMETRYCORE_API bool Cut(const MeshIntersection::FIntersectionsQueryResult& Intersections);
 
 
 	void ResetOutputs()
@@ -159,4 +159,4 @@ public:
 
 
 } // end namespace UE::Geometry
-} // end namespace UE
\ No newline at end of file
+} // end namespace UE
