[Engine]

diff --git a/Engine/Source/Runtime/GeometryCore/Public/Operations/MeshSelfUnion.h b/Engine/Source/Runtime/GeometryCore/Public/Operations/MeshSelfUnion.h
--- a/Engine/Source/Runtime/GeometryCore/Public/Operations/MeshSelfUnion.h
+++ b/Engine/Source/Runtime/GeometryCore/Public/Operations/MeshSelfUnion.h
@@ -25,7 +25,7 @@ namespace Geometry
 /**
  * MeshSelfUnion -- perform a "Mesh Boolean" style union of a mesh on itself, resolving any self intersections and welding the new boundaries as needed
  */
-class GEOMETRYCORE_API FMeshSelfUnion
+class FMeshSelfUnion
 {
 public:
 
@@ -128,7 +128,7 @@ public:
 	 * on the positive side of the cutting plane.
 	 * @return true if operation succeeds
 	 */
-	bool Compute();
+	GEOMETRYCORE_API bool Compute();
 
 protected:
 	/** If this returns true, abort computation.  */
@@ -139,14 +139,14 @@ protected:
 
 private:
 
-	int FindNearestEdge(const TArray<int>& EIDs, const TArray<int>& BoundaryNbrEdges, FVector3d Pos);
+	GEOMETRYCORE_API int FindNearestEdge(const TArray<int>& EIDs, const TArray<int>& BoundaryNbrEdges, FVector3d Pos);
 
-	bool MergeEdges(const TArray<int>& CutBoundaryEdges, const TMap<int, int>& AllVIDMatches);
+	GEOMETRYCORE_API bool MergeEdges(const TArray<int>& CutBoundaryEdges, const TMap<int, int>& AllVIDMatches);
 
-	void SimplifyAlongNewEdges(TArray<int>& CutBoundaryEdges, TMap<int, int>& FoundMatches);
+	GEOMETRYCORE_API void SimplifyAlongNewEdges(TArray<int>& CutBoundaryEdges, TMap<int, int>& FoundMatches);
 
 };
 
 
 } // end namespace UE::Geometry
-} // end namespace UE
\ No newline at end of file
+} // end namespace UE
diff --git a/Engine/Source/Runtime/GeometryCore/Public/Parameterization/UVPacking.h b/Engine/Source/Runtime/GeometryCore/Public/Parameterization/UVPacking.h
--- a/Engine/Source/Runtime/GeometryCore/Public/Parameterization/UVPacking.h
+++ b/Engine/Source/Runtime/GeometryCore/Public/Parameterization/UVPacking.h
@@ -21,7 +21,7 @@ namespace Geometry
  * generic mesh class.  The UV islands must already be identified, this
  * class simply scales/rotates/translates the islands to fit.
  */
-class GEOMETRYCORE_API FUVPacker
+class FUVPacker
 {
 public:
 
@@ -54,7 +54,7 @@ public:
 	 * All UV islands are packed into standard positive-unit-square.
 	 * Only supports single-pixel border size.
 	 */
-	bool StandardPack(IUVMeshView* Mesh, int NumIslands, TFunctionRef<void(int, TArray<int32>&)> CopyIsland);
+	GEOMETRYCORE_API bool StandardPack(IUVMeshView* Mesh, int NumIslands, TFunctionRef<void(int, TArray<int32>&)> CopyIsland);
 
 	/// Version of StandardPack that takes an array of arrays instead of a TFunctionRef, for convenience
 	bool StandardPack(IUVMeshView* Mesh, const TArray<TArray<int>>& UVIslands)
@@ -70,7 +70,7 @@ public:
 	 * and translate each islands separately so that it's bbox-min is at the origin.
 	 * So the islands are "stacked" and all fit in the unit box.
 	 */
-	bool StackPack(IUVMeshView* Mesh, int NumIslands, TFunctionRef<void(int, TArray<int32>&)> CopyIsland);
+	GEOMETRYCORE_API bool StackPack(IUVMeshView* Mesh, int NumIslands, TFunctionRef<void(int, TArray<int32>&)> CopyIsland);
 
 	/// Version of StackPack that takes an array of arrays instead of a TFunctionRef, for convenience
 	bool StackPack(IUVMeshView* Mesh, const TArray<TArray<int>>& UVIslands)
@@ -86,7 +86,7 @@ protected:
 	/**
 	 * Compute common stats used by the packing algorithms to transform UV islands
 	 */
-	void GetIslandStats(IUVMeshView* Mesh, const TArray<int32>& Island, FAxisAlignedBox2d& IslandBoundsOut, double& IslandScaleFactorOut, double& UVAreaOut);
+	GEOMETRYCORE_API void GetIslandStats(IUVMeshView* Mesh, const TArray<int32>& Island, FAxisAlignedBox2d& IslandBoundsOut, double& IslandScaleFactorOut, double& UVAreaOut);
 
 };
 
diff --git a/Engine/Source/Runtime/GeometryCore/Public/Sampling/MeshSurfacePointSampling.h b/Engine/Source/Runtime/GeometryCore/Public/Sampling/MeshSurfacePointSampling.h
--- a/Engine/Source/Runtime/GeometryCore/Public/Sampling/MeshSurfacePointSampling.h
+++ b/Engine/Source/Runtime/GeometryCore/Public/Sampling/MeshSurfacePointSampling.h
@@ -21,7 +21,7 @@ class FDynamicMesh3;
  * FMeshSurfacePointSampling computes oriented point samples on the surface of
  * a Mesh using various sampling strategies.
  */
-class GEOMETRYCORE_API FMeshSurfacePointSampling
+class FMeshSurfacePointSampling
 {
 public:
 	// 
@@ -142,7 +142,7 @@ public:
 	 * also introduces various biases to increase performance. Generally increasing SubSampleDensity will result in 
 	 * more tightly-packed results, but at increasingly expensive computation time.
 	 */
-	void ComputePoissonSampling(const FDynamicMesh3& Mesh, FProgressCancel* Progress = nullptr);
+	GEOMETRYCORE_API void ComputePoissonSampling(const FDynamicMesh3& Mesh, FProgressCancel* Progress = nullptr);
 
 
 	// ability to incrementally increase existing sample sizes to more tightly pack them?
diff --git a/Engine/Source/Runtime/GeometryCore/Public/Selections/MeshConnectedComponents.h b/Engine/Source/Runtime/GeometryCore/Public/Selections/MeshConnectedComponents.h
--- a/Engine/Source/Runtime/GeometryCore/Public/Selections/MeshConnectedComponents.h
+++ b/Engine/Source/Runtime/GeometryCore/Public/Selections/MeshConnectedComponents.h
@@ -19,7 +19,7 @@ namespace Geometry
  * By default the actual mesh connectivity is used, but an optional connectivity predicate
  * can be provided to specify when two elements should be considered connected.
  */
-class GEOMETRYCORE_API FMeshConnectedComponents
+class FMeshConnectedComponents
 {
 public:
 	const FDynamicMesh3* Mesh;
@@ -54,7 +54,7 @@ public:
 	 * Triangle connectivity is based on edge connectivity, ie bowtie-vertices are not connections between triangles.
 	 * @param TrisConnectedPredicate optional function that specifies whether two edge-connected triangles should be considered connected by the search
 	 */
-	void FindConnectedTriangles(TFunction<bool(int32, int32)> TrisConnectedPredicate = nullptr);
+	GEOMETRYCORE_API void FindConnectedTriangles(TFunction<bool(int32, int32)> TrisConnectedPredicate = nullptr);
 
 	/**
 	 * Find all connected triangle components of a subset of triangles of the Mesh and store in Components array.
@@ -62,7 +62,7 @@ public:
 	 * @param TriangleROI list of triangles to search across
 	 * @param TrisConnectedPredicate optional function that specifies whether two edge-connected triangles should be considered connected by the search
 	 */
-	void FindConnectedTriangles(const TArray<int>& TriangleROI, TFunction<bool(int32, int32)> TrisConnectedPredicate = nullptr);
+	GEOMETRYCORE_API void FindConnectedTriangles(const TArray<int>& TriangleROI, TFunction<bool(int32, int32)> TrisConnectedPredicate = nullptr);
 
 	/**
 	 * Find all connected triangle components of a subset of triangles of the Mesh and store in Components array.
@@ -70,7 +70,7 @@ public:
 	 * @param IndexFilterFunc defines set of triangles to search across, return true for triangle IDs that are to be considered
 	 * @param TrisConnectedPredicate optional function that specifies whether two edge-connected triangles should be considered connected by the search
 	 */
-	void FindConnectedTriangles(TFunctionRef<bool(int)> IndexFilterFunc, TFunction<bool(int32, int32)> TrisConnectedPredicate = nullptr);
+	GEOMETRYCORE_API void FindConnectedTriangles(TFunctionRef<bool(int)> IndexFilterFunc, TFunction<bool(int32, int32)> TrisConnectedPredicate = nullptr);
 
 	/**
 	 * Find all connected triangle components that contain one or more Seed Triangles and store in Components array.
@@ -79,14 +79,14 @@ public:
 	 * @param SeedTriangles list of start triangles, each component contains at least one of these triangles
 	 * @param TrisConnectedPredicate optional function that specifies whether two edge-connected triangles should be considered connected by the search
 	 */
-	void FindTrianglesConnectedToSeeds(const TArray<int>& SeedTriangles, TFunction<bool(int32, int32)> TrisConnectedPredicate = nullptr);
+	GEOMETRYCORE_API void FindTrianglesConnectedToSeeds(const TArray<int>& SeedTriangles, TFunction<bool(int32, int32)> TrisConnectedPredicate = nullptr);
 
 	/**
 	 * Initialize the internal FComponent list from the input ComponentLists, skipping any empty input lists
 	 * @param bValidateIDs if true, test that each value corresponds to a valid triangle ID on the Mesh
 	 * @return true if all IDs are valid, or if check was skipped
 	 */
-	bool InitializeFromTriangleComponents(const TArray<TArray<int32>>& ComponentLists, bool bValidateIDs);
+	GEOMETRYCORE_API bool InitializeFromTriangleComponents(const TArray<TArray<int32>>& ComponentLists, bool bValidateIDs);
 
 	/**
 	* Initialize the internal FComponent list from the input ComponentLists, skipping any empty input lists
@@ -94,27 +94,27 @@ public:
 	* @param bValidateIDs if true, test that each value corresponds to a valid triangle ID on the Mesh
 	* @return true if all IDs are valid, or if check was skipped
 	*/
-	bool InitializeFromTriangleComponents(TArray<TArray<int32>>& ComponentLists, bool bMoveSubLists, bool bValidateIDs);
+	GEOMETRYCORE_API bool InitializeFromTriangleComponents(TArray<TArray<int32>>& ComponentLists, bool bMoveSubLists, bool bValidateIDs);
 
 	/**
 	* Find all connected vertex components of the Mesh and store in Components array.
 	* @param VertsConnectedPredicate optional function that specifies whether two edge-connected vertices should be considered connected by the search
 	*/
-	void FindConnectedVertices(TFunction<bool(int32, int32)> VertsConnectedPredicate = nullptr);
+	GEOMETRYCORE_API void FindConnectedVertices(TFunction<bool(int32, int32)> VertsConnectedPredicate = nullptr);
 
 	/**
 	* Find all connected vertex components of a subset of vertices of the Mesh and store in Components array.
 	* @param VertexROI list of vertices to search across
 	* @param VertsConnectedPredicate optional function that specifies whether two edge-connected vertices should be considered connected by the search
 	*/
-	void FindConnectedVertices(const TArray<int>& VertexROI, TFunction<bool(int32, int32)> VertsConnectedPredicate = nullptr);
+	GEOMETRYCORE_API void FindConnectedVertices(const TArray<int>& VertexROI, TFunction<bool(int32, int32)> VertsConnectedPredicate = nullptr);
 
 	/**
 	* Find all connected vertex components of a subset of vertices of the Mesh and store in Components array.
 	* @param IndexFilterFunc defines set of vertices to search across, return true for vertex IDs that are to be considered
 	* @param VertsConnectedPredicate optional function that specifies whether two edge-connected vertices should be considered connected by the search
 	*/
-	void FindConnectedVertices(TFunctionRef<bool(int)> IndexFilterFunc, TFunction<bool(int32, int32)>VertsConnectedPredicate = nullptr);
+	GEOMETRYCORE_API void FindConnectedVertices(TFunctionRef<bool(int)> IndexFilterFunc, TFunction<bool(int32, int32)>VertsConnectedPredicate = nullptr);
 
 	/**
 	* Find all connected vertex components that contain one or more Seed Vertices and store in Components array.
@@ -122,14 +122,14 @@ public:
 	* @param SeedVertices list of start vertices, each component contains at least one of these vertices
 	* @param VertsConnectedPredicate optional function that specifies whether two edge-connected vertices should be considered connected by the search
 	*/
-	void FindVerticesConnectedToSeeds(const TArray<int>& SeedVertices, TFunction<bool(int32, int32)> VertsConnectedPredicate = nullptr);
+	GEOMETRYCORE_API void FindVerticesConnectedToSeeds(const TArray<int>& SeedVertices, TFunction<bool(int32, int32)> VertsConnectedPredicate = nullptr);
 
 	/**
 	 * Initialize the internal FComponent list from the input ComponentLists, skipping any empty input lists
 	 * @param bValidateIDs if true, test that each value corresponds to a valid vertex ID on the Mesh
 	 * @return true if all IDs are valid, or if check was skipped
 	 */
-	bool InitializeFromVertexComponents(const TArray<TArray<int32>>& ComponentLists, bool bValidateIDs);
+	GEOMETRYCORE_API bool InitializeFromVertexComponents(const TArray<TArray<int32>>& ComponentLists, bool bValidateIDs);
 
 	/**
 	* Initialize the internal FComponent list from the input ComponentLists, skipping any empty input lists
@@ -137,7 +137,7 @@ public:
 	* @param bValidateIDs if true, test that each value corresponds to a valid vertex ID on the Mesh
 	* @return true if all IDs are valid, or if check was skipped
 	*/
-	bool InitializeFromVertexComponents(TArray<TArray<int32>>& ComponentLists, bool bMoveSubLists, bool bValidateIDs);
+	GEOMETRYCORE_API bool InitializeFromVertexComponents(TArray<TArray<int32>>& ComponentLists, bool bMoveSubLists, bool bValidateIDs);
 
 	//
 	// Query functions. Only valid to call after a Calculation function has been called.
@@ -164,13 +164,13 @@ public:
 	/** 
 	 * @return index of largest component by element count 
 	 */
-	int32 GetLargestIndexByCount() const;
+	GEOMETRYCORE_API int32 GetLargestIndexByCount() const;
 
 	/**
 	 * Sort the Components array by component element count
 	 * @param bLargestFirst if true, sort by decreasing count, otherwise by increasing count
 	 */
-	void SortByCount(bool bLargestFirst = true);
+	GEOMETRYCORE_API void SortByCount(bool bLargestFirst = true);
 
 
 
@@ -190,16 +190,16 @@ protected:
 	//
 	// Internal functions to calculate ROI
 	//
-	void FindTriComponents(FInterval1i ActiveRange, TArray<uint8>& ActiveSet, TFunction<bool(int32, int32)> TriConnectedPredicate);
-	void FindTriComponents(const TArray<int32>& SeedList, TArray<uint8>& ActiveSet, TFunction<bool(int32, int32)> TriConnectedPredicate);
-	void FindTriComponent(FComponent* Component, TArray<int32>& ComponentQueue, TArray<uint8>& ActiveSet);
-	void FindTriComponent(FComponent* Component, TArray<int32>& ComponentQueue, TArray<uint8>& ActiveSet, TFunctionRef<bool(int32, int32)> TriConnectedPredicate);
-	void RemoveFromActiveSet(const FComponent* Component, TArray<uint8>& ActiveSet);
+	GEOMETRYCORE_API void FindTriComponents(FInterval1i ActiveRange, TArray<uint8>& ActiveSet, TFunction<bool(int32, int32)> TriConnectedPredicate);
+	GEOMETRYCORE_API void FindTriComponents(const TArray<int32>& SeedList, TArray<uint8>& ActiveSet, TFunction<bool(int32, int32)> TriConnectedPredicate);
+	GEOMETRYCORE_API void FindTriComponent(FComponent* Component, TArray<int32>& ComponentQueue, TArray<uint8>& ActiveSet);
+	GEOMETRYCORE_API void FindTriComponent(FComponent* Component, TArray<int32>& ComponentQueue, TArray<uint8>& ActiveSet, TFunctionRef<bool(int32, int32)> TriConnectedPredicate);
+	GEOMETRYCORE_API void RemoveFromActiveSet(const FComponent* Component, TArray<uint8>& ActiveSet);
 
-	void FindVertComponents(FInterval1i ActiveRange, TArray<uint8>& ActiveSet, TFunction<bool(int32, int32)> VertsConnectedPredicate);
-	void FindVertComponents(const TArray<int32>& SeedList, TArray<uint8>& ActiveSet, TFunction<bool(int32, int32)> VertsConnectedPredicate);
-	void FindVertComponent(FComponent* Component, TArray<int32>& ComponentQueue, TArray<uint8>& ActiveSet);
-	void FindVertComponent(FComponent* Component, TArray<int32>& ComponentQueue, TArray<uint8>& ActiveSet, TFunctionRef<bool(int32, int32)> VertsConnectedPredicate);
+	GEOMETRYCORE_API void FindVertComponents(FInterval1i ActiveRange, TArray<uint8>& ActiveSet, TFunction<bool(int32, int32)> VertsConnectedPredicate);
+	GEOMETRYCORE_API void FindVertComponents(const TArray<int32>& SeedList, TArray<uint8>& ActiveSet, TFunction<bool(int32, int32)> VertsConnectedPredicate);
+	GEOMETRYCORE_API void FindVertComponent(FComponent* Component, TArray<int32>& ComponentQueue, TArray<uint8>& ActiveSet);
+	GEOMETRYCORE_API void FindVertComponent(FComponent* Component, TArray<int32>& ComponentQueue, TArray<uint8>& ActiveSet, TFunctionRef<bool(int32, int32)> VertsConnectedPredicate);
 
 public:
 	/**
@@ -212,7 +212,7 @@ public:
 	 * @param DoneBuffer optional set used to track which triangles have already been processed. If passed as nullptr, an TSet will be locally allocated
 	 * @param CanGrowPredicate determines whether two connected mesh triangles should be considered connected while growing
 	 */
-	static void GrowToConnectedTriangles(const FDynamicMesh3* Mesh, 
+	static GEOMETRYCORE_API void GrowToConnectedTriangles(const FDynamicMesh3* Mesh, 
 		const TArray<int>& InputROI, 
 		TArray<int>& ResultROI,
 		TArray<int32>* QueueBuffer = nullptr, 
@@ -231,7 +231,7 @@ public:
 	 * @param QueueBuffer optional buffer used as internal Queue. If passed as nullptr, a TArray will be locally allocated
 	 * @param CanGrowPredicate determines whether two connected mesh triangles should be considered connected while growing
 	 */
-	static void GrowToConnectedTriangles(const FDynamicMesh3* Mesh,
+	static GEOMETRYCORE_API void GrowToConnectedTriangles(const FDynamicMesh3* Mesh,
 		const TArray<int>& InputROI,
 		TSet<int>& ResultROI,
 		TArray<int32>* QueueBuffer = nullptr,
@@ -247,7 +247,7 @@ public:
 	 * @param QueueBuffer optional buffer used as internal Queue. If passed as nullptr, a TArray will be locally allocated
 	 * @param CanGrowPredicate determines whether two connected mesh vertices should be considered connected while growing
 	 */
-	static void GrowToConnectedVertices(
+	static GEOMETRYCORE_API void GrowToConnectedVertices(
 		const FDynamicMesh3& Mesh,
 		const TArray<int>& InputROI,
 		TSet<int>& ResultROI,
@@ -264,7 +264,7 @@ public:
 	 * @param QueueBuffer optional buffer used as internal Queue. If passed as nullptr, a TArray will be locally allocated
 	 * @param CanGrowPredicate determines whether two connected mesh edges should be considered connected while growing
 	 */
-	static void GrowToConnectedEdges(
+	static GEOMETRYCORE_API void GrowToConnectedEdges(
 		const FDynamicMesh3& Mesh,
 		const TArray<int>& InputROI,
 		TSet<int>& ResultROI,
diff --git a/Engine/Source/Runtime/GeometryCore/Public/Selections/MeshEdgeSelection.h b/Engine/Source/Runtime/GeometryCore/Public/Selections/MeshEdgeSelection.h
--- a/Engine/Source/Runtime/GeometryCore/Public/Selections/MeshEdgeSelection.h
+++ b/Engine/Source/Runtime/GeometryCore/Public/Selections/MeshEdgeSelection.h
@@ -18,7 +18,7 @@ namespace Geometry
 /**
  * Currently a thin wrapper of a TSet<int> of Edge IDs paired with a Mesh; the backing storage will likely change as we need to optimize in the future
  */
-class GEOMETRYCORE_API FMeshEdgeSelection
+class FMeshEdgeSelection
 {
 private:
 	const FDynamicMesh3* Mesh;
@@ -34,10 +34,10 @@ public:
 
 
 	// convert vertex selection to edge selection. Require at least minCount verts of edge to be selected
-	FMeshEdgeSelection(const FDynamicMesh3* mesh, const FMeshVertexSelection& convertV, int minCount = 2);
+	GEOMETRYCORE_API FMeshEdgeSelection(const FDynamicMesh3* mesh, const FMeshVertexSelection& convertV, int minCount = 2);
 
 	// convert face selection to edge selection. Require at least minCount tris of edge to be selected
-	FMeshEdgeSelection(const FDynamicMesh3* mesh, const FMeshFaceSelection& convertT, int minCount = 1);
+	GEOMETRYCORE_API FMeshEdgeSelection(const FDynamicMesh3* mesh, const FMeshFaceSelection& convertT, int minCount = 1);
 
 
 
@@ -155,7 +155,7 @@ public:
 	}
 
 
-	void SelectBoundaryTriEdges(const FMeshFaceSelection& Triangles);
+	GEOMETRYCORE_API void SelectBoundaryTriEdges(const FMeshFaceSelection& Triangles);
 
 	void Deselect(int tid) {
 		remove(tid);
diff --git a/Engine/Source/Runtime/GeometryCore/Public/Selections/MeshFaceSelection.h b/Engine/Source/Runtime/GeometryCore/Public/Selections/MeshFaceSelection.h
--- a/Engine/Source/Runtime/GeometryCore/Public/Selections/MeshFaceSelection.h
+++ b/Engine/Source/Runtime/GeometryCore/Public/Selections/MeshFaceSelection.h
@@ -17,7 +17,7 @@ class FMeshVertexSelection;
 class FMeshEdgeSelection;
 
 
-class GEOMETRYCORE_API FMeshFaceSelection
+class FMeshFaceSelection
 {
 private:
 	const FDynamicMesh3* Mesh;
@@ -33,7 +33,7 @@ public:
 
 	// convert vertex selection to face selection. Require at least minCount verts of
 	// tri to be selected (valid values are 1,2,3)
-	FMeshFaceSelection(const FDynamicMesh3* mesh, const FMeshVertexSelection& convertV, int minCount = 3);
+	GEOMETRYCORE_API FMeshFaceSelection(const FDynamicMesh3* mesh, const FMeshVertexSelection& convertV, int minCount = 3);
 
 	// select a group
 	FMeshFaceSelection(const FDynamicMesh3* mesh, int group_id) : Mesh(mesh)
@@ -353,7 +353,7 @@ public:
 	 *  
 	 *  Return false from FilterF to prevent triangles from being included.
 	 */
-	void ExpandToOneRingNeighbours(const TUniqueFunction<bool(int)>& FilterF = nullptr);
+	GEOMETRYCORE_API void ExpandToOneRingNeighbours(const TUniqueFunction<bool(int)>& FilterF = nullptr);
 
 	/**
 	 *  Expand selection by N vertex one-rings. This is *significantly* faster
@@ -362,7 +362,7 @@ public:
 	 *  
 	 *  Return false from FilterF to prevent triangles from being included.
 	 */
-	void ExpandToOneRingNeighbours(int nRings, const TUniqueFunction<bool(int)>& FilterF = nullptr);
+	GEOMETRYCORE_API void ExpandToOneRingNeighbours(int nRings, const TUniqueFunction<bool(int)>& FilterF = nullptr);
 
 	/**
 	 *  remove all triangles in vertex one-rings of current selection to set.
@@ -371,7 +371,7 @@ public:
 	 *  
 	 *  Return false from FilterF to prevent triangles from being deselected.
 	 */
-	void ContractBorderByOneRingNeighbours(int NumRings = 1, bool bContractFromMeshBoundary = false, const TUniqueFunction<bool(int)>& FilterF = nullptr);
+	GEOMETRYCORE_API void ContractBorderByOneRingNeighbours(int NumRings = 1, bool bContractFromMeshBoundary = false, const TUniqueFunction<bool(int)>& FilterF = nullptr);
 
 	/**
 	 *  Grow selection outwards from seed triangle, until it hits boundaries defined by triangle and edge filters.
diff --git a/Engine/Source/Runtime/GeometryCore/Public/Selections/MeshVertexSelection.h b/Engine/Source/Runtime/GeometryCore/Public/Selections/MeshVertexSelection.h
--- a/Engine/Source/Runtime/GeometryCore/Public/Selections/MeshVertexSelection.h
+++ b/Engine/Source/Runtime/GeometryCore/Public/Selections/MeshVertexSelection.h
@@ -15,7 +15,7 @@ namespace Geometry
 class FMeshFaceSelection;
 class FMeshEdgeSelection;
 
-class GEOMETRYCORE_API FMeshVertexSelection
+class FMeshVertexSelection
 {
 private:
 	const FDynamicMesh3* Mesh;
@@ -29,10 +29,10 @@ public:
 	}
 
 	// convert face selection to vertex selection. 
-	FMeshVertexSelection(const FDynamicMesh3* mesh, const FMeshFaceSelection& convertT);
+	GEOMETRYCORE_API FMeshVertexSelection(const FDynamicMesh3* mesh, const FMeshFaceSelection& convertT);
 
 	// convert edge selection to vertex selection. 
-	FMeshVertexSelection(const FDynamicMesh3* mesh, const FMeshEdgeSelection& convertE);
+	GEOMETRYCORE_API FMeshVertexSelection(const FDynamicMesh3* mesh, const FMeshEdgeSelection& convertE);
 
 
 	TSet<int> AsSet() const
@@ -146,14 +146,14 @@ public:
 			add(tri.A); add(tri.B); add(tri.C);
 		}
 	}
-	void SelectTriangleVertices(const FMeshFaceSelection& Triangles);
+	GEOMETRYCORE_API void SelectTriangleVertices(const FMeshFaceSelection& Triangles);
 
 
 	/**
 	 *  for each vertex of input triangle set, select vertex if all
 	 *  one-ring triangles are contained in triangle set (ie vertex is not on boundary of triangle set).
 	 */
-	void SelectInteriorVertices(const FMeshFaceSelection& triangles);
+	GEOMETRYCORE_API void SelectInteriorVertices(const FMeshFaceSelection& triangles);
 
 
 	/**
diff --git a/Engine/Source/Runtime/GeometryCore/Public/Spatial/GeometrySet3.h b/Engine/Source/Runtime/GeometryCore/Public/Spatial/GeometrySet3.h
--- a/Engine/Source/Runtime/GeometryCore/Public/Spatial/GeometrySet3.h
+++ b/Engine/Source/Runtime/GeometryCore/Public/Spatial/GeometrySet3.h
@@ -24,34 +24,34 @@ namespace Geometry
  * 
  * Since Points and Curves have no area to hit, hit-tests are done via nearest-point-on-ray.
  */
-class GEOMETRYCORE_API FGeometrySet3
+class FGeometrySet3
 {
 public:
 	/**
 	 * @param bPoints if true, discard all points
 	 * @param bCurves if true, discard all polycurves
 	 */
-	void Reset(bool bPoints = true, bool bCurves = true);
+	GEOMETRYCORE_API void Reset(bool bPoints = true, bool bCurves = true);
 
 	/** Add a point with given PointID at the given Position*/
-	void AddPoint(int PointID, const FVector3d& Position);
+	GEOMETRYCORE_API void AddPoint(int PointID, const FVector3d& Position);
 	/** Add a polycurve with given CurveID and the give Polyline */
-	void AddCurve(int CurveID, const FPolyline3d& Polyline);
+	GEOMETRYCORE_API void AddCurve(int CurveID, const FPolyline3d& Polyline);
 
 	/** Remove a point with given PointID. */
-	void RemovePoint(int PointID);
+	GEOMETRYCORE_API void RemovePoint(int PointID);
 	/** Remove a curve with given CurveID. */
-	void RemoveCurve(int CurveID);
+	GEOMETRYCORE_API void RemoveCurve(int CurveID);
 
 	/** Update the Position of previously-added PointID */
-	void UpdatePoint(int PointID, const FVector3d& Position);
+	GEOMETRYCORE_API void UpdatePoint(int PointID, const FVector3d& Position);
 	/** Update the Polyline of previously-added CurveID */
-	void UpdateCurve(int CurveID, const FPolyline3d& Polyline);
+	GEOMETRYCORE_API void UpdateCurve(int CurveID, const FPolyline3d& Polyline);
 
 	/**
 	 * FNearest is returned by nearest-point queries
 	 */
-	struct GEOMETRYCORE_API FNearest
+	struct FNearest
 	{
 		/** ID of point or curve */
 		int ID;
@@ -80,7 +80,7 @@ public:
 	 * @return true if the nearest point on Ray to some point in the set passed the PointWithinToleranceTest.
 	 * @warning PointWithinToleranceTest is called in parallel and hence must be thread-safe/re-entrant!
 	 */
-	bool FindNearestPointToRay(const FRay3d& Ray, FNearest& ResultOut,
+	GEOMETRYCORE_API bool FindNearestPointToRay(const FRay3d& Ray, FNearest& ResultOut,
 		TFunction<bool(const FVector3d&, const FVector3d&)> PointWithinToleranceTest) const;
 
 	/**
@@ -92,7 +92,7 @@ public:
 	 * @return true if at least one result was added (ie, passed PointWithinToleranceTest).
 	 * @warning PointWithinToleranceTest is called in parallel and hence must be thread-safe/re-entrant!
 	 */
-	bool CollectPointsNearRay(const FRay3d& Ray, TArray<FNearest>& ResultsOut,
+	GEOMETRYCORE_API bool CollectPointsNearRay(const FRay3d& Ray, TArray<FNearest>& ResultsOut,
 		TFunction<bool(const FVector3d&, const FVector3d&)> PointWithinToleranceTest) const;
 
 	/**
@@ -102,7 +102,7 @@ public:
 	 * @return true if the nearest point on Ray to some curve in the set passed the PointWithinToleranceTest.
 	 * @warning PointWithinToleranceTest is called in parallel and hence must be thread-safe/re-entrant!
 	 */
-	bool FindNearestCurveToRay(const FRay3d& Ray, FNearest& ResultOut,
+	GEOMETRYCORE_API bool FindNearestCurveToRay(const FRay3d& Ray, FNearest& ResultOut,
 		TFunction<bool(const FVector3d&, const FVector3d&)> PointWithinToleranceTest) const;
 
 	/**
@@ -114,7 +114,7 @@ public:
 	 * @return true if at least one result was added (ie, passed PointWithinToleranceTest).
 	 * @warning PointWithinToleranceTest is called in parallel and hence must be thread-safe/re-entrant!
 	 */
-	bool CollectCurvesNearRay(const FRay3d& Ray, TArray<FNearest>& ResultsOut,
+	GEOMETRYCORE_API bool CollectCurvesNearRay(const FRay3d& Ray, TArray<FNearest>& ResultsOut,
 		TFunction<bool(const FVector3d&, const FVector3d&)> PointWithinToleranceTest) const;
 
 	/**
diff --git a/Engine/Source/Runtime/GeometryCore/Public/Spatial/PointSetHashTable.h b/Engine/Source/Runtime/GeometryCore/Public/Spatial/PointSetHashTable.h
--- a/Engine/Source/Runtime/GeometryCore/Public/Spatial/PointSetHashTable.h
+++ b/Engine/Source/Runtime/GeometryCore/Public/Spatial/PointSetHashTable.h
@@ -28,7 +28,7 @@ namespace Geometry
  *
  * @todo support larger search radius than cellsize
  */
-class GEOMETRYCORE_API FPointSetHashtable
+class FPointSetHashtable
 {
 protected:
 	typedef TArray<int> PointList;
@@ -56,7 +56,7 @@ public:
 	 * @param CellSize the size of the cells/voxels in the grid.
 	 * @param Origin World origin of the grid (not strictly necessary since grid is sparse, can set to origin for example)
 	 */
-	void Build(double CellSize, const FVector3d& Origin);
+	GEOMETRYCORE_API void Build(double CellSize, const FVector3d& Origin);
 
 	/**
 	 * Find all points within given query distance from query point.
@@ -67,10 +67,10 @@ public:
 	 * @param ResultOut indices of discovered points are stored in this list
 	 * @return true on success
 	 */
-	bool FindPointsInBall(const FVector3d& QueryPt, double QueryRadius, TArray<int>& ResultOut);
+	GEOMETRYCORE_API bool FindPointsInBall(const FVector3d& QueryPt, double QueryRadius, TArray<int>& ResultOut);
 
 };
 
 
 } // end namespace UE::Geometry
-} // end namespace UE
\ No newline at end of file
+} // end namespace UE
diff --git a/Engine/Source/Runtime/GeometryCore/Public/Spatial/PriorityOrderPoints.h b/Engine/Source/Runtime/GeometryCore/Public/Spatial/PriorityOrderPoints.h
--- a/Engine/Source/Runtime/GeometryCore/Public/Spatial/PriorityOrderPoints.h
+++ b/Engine/Source/Runtime/GeometryCore/Public/Spatial/PriorityOrderPoints.h
@@ -22,7 +22,7 @@ using namespace UE::Math;
  * such that if we only look at the first N points, they will cover the overall space as well as possible
  * while still favoring higher ImportanceWeight points in each local region.
  */
-struct GEOMETRYCORE_API FPriorityOrderPoints
+struct FPriorityOrderPoints
 {
 
 	/// Inputs
@@ -50,8 +50,8 @@ public:
 	 * @param EarlyStop				If >= 0, stop re-ordering once we've sorted this many points; beyond this the ordering will be arbitrary
 	 * @param OffsetResFactor		Factor that controls how closely the points are allowed to space, with higher values -> more clumping allowed
 	 */
-	void ComputeUniformSpaced(TArrayView<const FVector3d> Points, TArrayView<const float> ImportanceWeights, int32 EarlyStop = -1, int32 OffsetResFactor = 4);
-	void ComputeUniformSpaced(TArrayView<const FVector3f> Points, TArrayView<const float> ImportanceWeights, int32 EarlyStop = -1, int32 OffsetResFactor = 4);
+	GEOMETRYCORE_API void ComputeUniformSpaced(TArrayView<const FVector3d> Points, TArrayView<const float> ImportanceWeights, int32 EarlyStop = -1, int32 OffsetResFactor = 4);
+	GEOMETRYCORE_API void ComputeUniformSpaced(TArrayView<const FVector3f> Points, TArrayView<const float> ImportanceWeights, int32 EarlyStop = -1, int32 OffsetResFactor = 4);
 
 	/**
 	 * Compute an ordering for the input points that attempts to keep points 'well spaced' / un-clumped
@@ -62,15 +62,15 @@ public:
 	 * @param EarlyStop					If >= 0, stop re-ordering once we've sorted this many points; beyond this the ordering will be arbitrary
 	 * @param OffsetResFactor			Factor that controls how closely the points are allowed to space, with higher values -> more clumping allowed
 	 */
-	void ComputeUniformSpaced(TArrayView<const FVector3d> Points, TArrayView<const float> ImportanceWeights, TArrayView<const float> SecondImportanceWeights, int32 EarlyStop = -1, int32 OffsetResFactor = 4);
-	void ComputeUniformSpaced(TArrayView<const FVector3f> Points, TArrayView<const float> ImportanceWeights, TArrayView<const float> SecondImportanceWeights, int32 EarlyStop = -1, int32 OffsetResFactor = 4);
+	GEOMETRYCORE_API void ComputeUniformSpaced(TArrayView<const FVector3d> Points, TArrayView<const float> ImportanceWeights, TArrayView<const float> SecondImportanceWeights, int32 EarlyStop = -1, int32 OffsetResFactor = 4);
+	GEOMETRYCORE_API void ComputeUniformSpaced(TArrayView<const FVector3f> Points, TArrayView<const float> ImportanceWeights, TArrayView<const float> SecondImportanceWeights, int32 EarlyStop = -1, int32 OffsetResFactor = 4);
 
 	/**
 	 * Compute an ordering that only sorts by descending ImportanceWeights, and does not attempt to spatially distribute the points
 	 * 
 	 * @param ImportanceWeights		Points with a higher ImportanceWeight will appear earlier in the ordering
 	 */
-	void ComputeDescendingImportance(TArrayView<const float> ImportanceWeights);
+	GEOMETRYCORE_API void ComputeDescendingImportance(TArrayView<const float> ImportanceWeights);
 
 };
 
diff --git a/Engine/Source/Runtime/GeometryCore/Public/Spatial/SparseDynamicOctree3.h b/Engine/Source/Runtime/GeometryCore/Public/Spatial/SparseDynamicOctree3.h
--- a/Engine/Source/Runtime/GeometryCore/Public/Spatial/SparseDynamicOctree3.h
+++ b/Engine/Source/Runtime/GeometryCore/Public/Spatial/SparseDynamicOctree3.h
@@ -43,7 +43,7 @@ class FDynamicFlagArray
 public:
 	TBitArray<> BitArray;
 	uint32 MaxIndex = 0;
-	GEOMETRYCORE_API static constexpr int32 GrowChunkSize = 0xFFF;
+	static constexpr int32 GrowChunkSize = 0xFFF;
 
 	FDynamicFlagArray()
 	{
@@ -81,8 +81,8 @@ public:
  */
 struct FSparseOctreeCell
 {
-	GEOMETRYCORE_API static constexpr uint32 InvalidID = TNumericLimits<uint32>::Max();;
-	GEOMETRYCORE_API static constexpr uint8 InvalidLevel = TNumericLimits<uint8>::Max();
+	static constexpr uint32 InvalidID = TNumericLimits<uint32>::Max();;
+	static constexpr uint8 InvalidLevel = TNumericLimits<uint8>::Max();
 
 	/** ID of cell (index into cell list) */
 	uint32 CellID;
@@ -161,7 +161,7 @@ struct FSparseOctreeCell
  * The octree is dynamic. Objects can be removed and re-inserted.
  * 
  */
-class GEOMETRYCORE_API FSparseDynamicOctree3
+class FSparseDynamicOctree3
 {
 	// potential optimizations/improvements
 	//    - Cell sizes are known at each level...can keep a lookup table? will this improve performance?
@@ -223,7 +223,7 @@ public:
 	 * @param ObjectID ID of the object
 	 * @return true if ObjectID is stored in this octree
 	 */
-	bool ContainsObject(int32 ObjectID) const;
+	GEOMETRYCORE_API bool ContainsObject(int32 ObjectID) const;
 
 	/**
 	 * Insert ObjectID into the Octree
@@ -231,14 +231,14 @@ public:
 	 * @param Bounds bounding box of the object
 	 *
 	 */
-	void InsertObject(int32 ObjectID, const FAxisAlignedBox3d& Bounds);
+	GEOMETRYCORE_API void InsertObject(int32 ObjectID, const FAxisAlignedBox3d& Bounds);
 
 	/**
 	 * Remove an object from the octree
 	 * @param ObjectID ID of the object
 	 * @return true if the object was in the tree and removed
 	 */
-	bool RemoveObject(int32 ObjectID);
+	GEOMETRYCORE_API bool RemoveObject(int32 ObjectID);
 
 	/**
 	 * Update the position of an object in the octree. This is more efficient than doing a remove+insert
@@ -246,7 +246,7 @@ public:
 	 * @param NewBounds new bounding box of the object
 	 * @return true if object was reinserted, false if the object was fine in it's existing cell
 	 */
-	bool ReinsertObject(int32 ObjectID, const FAxisAlignedBox3d& NewBounds, uint32 CellIDHint = InvalidCellID);
+	GEOMETRYCORE_API bool ReinsertObject(int32 ObjectID, const FAxisAlignedBox3d& NewBounds, uint32 CellIDHint = InvalidCellID);
 
 	/**
 	 * Check if the object needs to be reinserted, if it has NewBounds. 
@@ -255,7 +255,7 @@ public:
 	 * @param CellIDOut returned CellID of the object, if it is in the tree. This can be passed to ReinsertObject() to save some computation.
 	 * @return false if ObjectID is already in the tree, and NewBounds fits within the cell that currently contains ObjectID
 	 */
-	bool CheckIfObjectNeedsReinsert(int32 ObjectID, const FAxisAlignedBox3d& NewBounds, uint32& CellIDOut) const;
+	GEOMETRYCORE_API bool CheckIfObjectNeedsReinsert(int32 ObjectID, const FAxisAlignedBox3d& NewBounds, uint32& CellIDOut) const;
 
 
 	/**
@@ -266,7 +266,7 @@ public:
 	 * @param MaxDistance maximum hit distance
 	 * @return ObjectID of hit object, or -1 on miss
 	 */
-	int32 FindNearestHitObject(const FRay3d& Ray,
+	GEOMETRYCORE_API int32 FindNearestHitObject(const FRay3d& Ray,
 		TFunctionRef<FAxisAlignedBox3d(int)> GetObjectBoundsFunc,
 		TFunctionRef<double(int, const FRay3d&)> HitObjectDistFunc,
 		double MaxDistance = TNumericLimits<double>::Max()) const;
@@ -276,7 +276,7 @@ public:
 	 * @param Point query point
 	 * @param ObjectIDFunc this function is called for each ObjectID
 	 */
-	void ContainmentQuery(const FVector3d& Point,
+	GEOMETRYCORE_API void ContainmentQuery(const FVector3d& Point,
 		TFunctionRef<void(int)> ObjectIDFunc) const;
 
 	/**
@@ -285,7 +285,7 @@ public:
 	 * @param ObjectIDFunc this function is called for each ObjectID. Returns true to continue query and false to abort
 	 * @return true if query finished, false if it exited
 	 */
-	bool ContainmentQueryCancellable(const FVector3d& Point,
+	GEOMETRYCORE_API bool ContainmentQueryCancellable(const FVector3d& Point,
 		TFunctionRef<bool(int)> ObjectIDFunc) const;
 
 	/**
@@ -293,7 +293,7 @@ public:
 	 * @param Bounds query box
 	 * @param ObjectIDFunc this function is called for each ObjectID
 	 */
-	void RangeQuery(const FAxisAlignedBox3d& Bounds,
+	GEOMETRYCORE_API void RangeQuery(const FAxisAlignedBox3d& Bounds,
 		TFunctionRef<void(int)> ObjectIDFunc) const;
 
 	/**
@@ -301,7 +301,7 @@ public:
 	 * @param Bounds query box
 	 * @param ObjectIDsOut collected ObjectIDs are stored here 
 	 */
-	void RangeQuery(const FAxisAlignedBox3d& Bounds,
+	GEOMETRYCORE_API void RangeQuery(const FAxisAlignedBox3d& Bounds,
 		TArray<int>& ObjectIDsOut ) const;
 
 	/**
@@ -310,7 +310,7 @@ public:
 	 * @param Bounds query box
 	 * @param ObjectIDsOut collected ObjectIDs are stored here
 	 */
-	void ParallelRangeQuery(const FAxisAlignedBox3d& Bounds,
+	GEOMETRYCORE_API void ParallelRangeQuery(const FAxisAlignedBox3d& Bounds,
 		TArray<int>& ObjectIDsOut) const;
 
 
@@ -322,7 +322,7 @@ public:
 	 * @param bVerbose if true, print some debug info via UE_LOG
 	 * @param bFailOnMissingObjects if true, assume ObjectIDs are dense and that all ObjectIDs must be in the tree
 	 */
-	void CheckValidity(
+	GEOMETRYCORE_API void CheckValidity(
 		TFunctionRef<bool(int)> IsValidObjectIDFunc,
 		TFunctionRef<FAxisAlignedBox3d(int)> GetObjectBoundsFunc,
 		EValidityCheckFailMode FailMode = EValidityCheckFailMode::Check,
@@ -344,7 +344,7 @@ public:
 	/**
 	 * Populate given FStatistics with info about the octree
 	 */
-	void ComputeStatistics(FStatistics& StatsOut) const;
+	GEOMETRYCORE_API void ComputeStatistics(FStatistics& StatsOut) const;
 
 protected:
 	// this identifier is used for unknown cells
@@ -445,18 +445,18 @@ protected:
 		return InvalidCellID;
 	}
 
-	FSparseOctreeCell FindCurrentContainingCell(const FAxisAlignedBox3d& Bounds) const;
+	GEOMETRYCORE_API FSparseOctreeCell FindCurrentContainingCell(const FAxisAlignedBox3d& Bounds) const;
 
 
-	void Insert_Spill(int32 ObjectID, const FAxisAlignedBox3d& Bounds);
-	void Insert_NewRoot(int32 ObjectID, const FAxisAlignedBox3d& Bounds, FSparseOctreeCell NewRootCell);
-	void Insert_ToCell(int32 ObjectID, const FAxisAlignedBox3d& Bounds, const FSparseOctreeCell& ExistingCell);
-	void Insert_NewChildCell(int32 ObjectID, const FAxisAlignedBox3d& Bounds, int ParentCellID, FSparseOctreeCell NewChildCell, int ChildIdx);
+	GEOMETRYCORE_API void Insert_Spill(int32 ObjectID, const FAxisAlignedBox3d& Bounds);
+	GEOMETRYCORE_API void Insert_NewRoot(int32 ObjectID, const FAxisAlignedBox3d& Bounds, FSparseOctreeCell NewRootCell);
+	GEOMETRYCORE_API void Insert_ToCell(int32 ObjectID, const FAxisAlignedBox3d& Bounds, const FSparseOctreeCell& ExistingCell);
+	GEOMETRYCORE_API void Insert_NewChildCell(int32 ObjectID, const FAxisAlignedBox3d& Bounds, int ParentCellID, FSparseOctreeCell NewChildCell, int ChildIdx);
 
-	double FindNearestRayCellIntersection(const FSparseOctreeCell& Cell, const FRay3d& Ray) const;
+	GEOMETRYCORE_API double FindNearestRayCellIntersection(const FSparseOctreeCell& Cell, const FRay3d& Ray) const;
 
 
-	void BranchRangeQuery(const FSparseOctreeCell* ParentCell,
+	GEOMETRYCORE_API void BranchRangeQuery(const FSparseOctreeCell* ParentCell,
 		const FAxisAlignedBox3d& Bounds,
 		TArray<int>& ObjectIDs) const;
 
@@ -536,4 +536,4 @@ private:
 
 
 } // end namespace UE::Geometry
-} // end namespace UE
\ No newline at end of file
+} // end namespace UE
diff --git a/Engine/Source/Runtime/GeometryCore/Public/Spatial/ZOrderCurvePoints.h b/Engine/Source/Runtime/GeometryCore/Public/Spatial/ZOrderCurvePoints.h
--- a/Engine/Source/Runtime/GeometryCore/Public/Spatial/ZOrderCurvePoints.h
+++ b/Engine/Source/Runtime/GeometryCore/Public/Spatial/ZOrderCurvePoints.h
@@ -21,7 +21,7 @@ using namespace UE::Math;
  * Generate a "z-order curve" ordering for a point set by distributing the points to a quad- or octree and then reading the points indices back in tree traversal order
  * This is useful for giving the points spatial locality -- so points that are close together in the ordering tend to be close together spatially, also
  */
-struct GEOMETRYCORE_API FZOrderCurvePoints
+struct FZOrderCurvePoints
 {
 	/// Inputs
 
@@ -46,10 +46,10 @@ public:
 	 * @param Points				Points to re-order
 	 * @param Bounds				Optional bounding box of the points; if an empty box is passed in, it will be computed for you
 	 */
-	void Compute(TArrayView<const FVector2d> Points, const FAxisAlignedBox2d& Bounds = FAxisAlignedBox2d::Empty());
-	void Compute(TArrayView<const FVector2f> Points, const FAxisAlignedBox2f& Bounds = FAxisAlignedBox2f::Empty());
-	void Compute(TArrayView<const FVector3d> Points, const FAxisAlignedBox3d& Bounds = FAxisAlignedBox3d::Empty());
-	void Compute(TArrayView<const FVector3f> Points, const FAxisAlignedBox3f& Bounds = FAxisAlignedBox3f::Empty());
+	GEOMETRYCORE_API void Compute(TArrayView<const FVector2d> Points, const FAxisAlignedBox2d& Bounds = FAxisAlignedBox2d::Empty());
+	GEOMETRYCORE_API void Compute(TArrayView<const FVector2f> Points, const FAxisAlignedBox2f& Bounds = FAxisAlignedBox2f::Empty());
+	GEOMETRYCORE_API void Compute(TArrayView<const FVector3d> Points, const FAxisAlignedBox3d& Bounds = FAxisAlignedBox3d::Empty());
+	GEOMETRYCORE_API void Compute(TArrayView<const FVector3f> Points, const FAxisAlignedBox3f& Bounds = FAxisAlignedBox3f::Empty());
 
 };
 
@@ -57,7 +57,7 @@ public:
  * Generate a "Biased Randomized Insertion Order" (BRIO) ordering for a point set by randomly bucketing the points, then applying Z-Order-Curve sorting (above) to each bucket
  * This gives an ordering with point locality but enough randomization to typically avoid worst case behavior if using the ordering for Delaunay meshing
  */
-struct GEOMETRYCORE_API FBRIOPoints
+struct FBRIOPoints
 {
 	/// Inputs
 
@@ -82,10 +82,10 @@ public:
 	 * @param Points				Points to re-order
 	 * @param Bounds				Optional bounding box of the points; if an empty box is passed in, it will be computed for you
 	 */
-	void Compute(TArrayView<const FVector2d> Points, const FAxisAlignedBox2d& Bounds = FAxisAlignedBox2d::Empty(), const FRandomStream& Random = FRandomStream());
-	void Compute(TArrayView<const FVector2f> Points, const FAxisAlignedBox2f& Bounds = FAxisAlignedBox2f::Empty(), const FRandomStream& Random = FRandomStream());
-	void Compute(TArrayView<const FVector3d> Points, const FAxisAlignedBox3d& Bounds = FAxisAlignedBox3d::Empty(), const FRandomStream& Random = FRandomStream());
-	void Compute(TArrayView<const FVector3f> Points, const FAxisAlignedBox3f& Bounds = FAxisAlignedBox3f::Empty(), const FRandomStream& Random = FRandomStream());
+	GEOMETRYCORE_API void Compute(TArrayView<const FVector2d> Points, const FAxisAlignedBox2d& Bounds = FAxisAlignedBox2d::Empty(), const FRandomStream& Random = FRandomStream());
+	GEOMETRYCORE_API void Compute(TArrayView<const FVector2f> Points, const FAxisAlignedBox2f& Bounds = FAxisAlignedBox2f::Empty(), const FRandomStream& Random = FRandomStream());
+	GEOMETRYCORE_API void Compute(TArrayView<const FVector3d> Points, const FAxisAlignedBox3d& Bounds = FAxisAlignedBox3d::Empty(), const FRandomStream& Random = FRandomStream());
+	GEOMETRYCORE_API void Compute(TArrayView<const FVector3f> Points, const FAxisAlignedBox3f& Bounds = FAxisAlignedBox3f::Empty(), const FRandomStream& Random = FRandomStream());
 
 };
 
diff --git a/Engine/Source/Runtime/GeometryCore/Public/Util/ProgressCancel.h b/Engine/Source/Runtime/GeometryCore/Public/Util/ProgressCancel.h
--- a/Engine/Source/Runtime/GeometryCore/Public/Util/ProgressCancel.h
+++ b/Engine/Source/Runtime/GeometryCore/Public/Util/ProgressCancel.h
@@ -183,7 +183,7 @@ struct FGeometryResult
  * 1) provide progress info back to caller
  * 2) allow caller to cancel the computation
  */
-class GEOMETRYCORE_API FProgressCancel
+class FProgressCancel
 {
 private:
 	bool WasCancelled = false;  // will be set to true if CancelF() ever returns true
@@ -220,9 +220,9 @@ private:
 	// critical section for accesses to the Progress Message
 	mutable FCriticalSection MessageCS;
 
-	void StartWorkScope(FProgressData& SaveProgressFrameOut, float StepSize, const FText& Message);
+	GEOMETRYCORE_API void StartWorkScope(FProgressData& SaveProgressFrameOut, float StepSize, const FText& Message);
 
-	void EndWorkScope(const FProgressData& SavedProgressFrame);
+	GEOMETRYCORE_API void EndWorkScope(const FProgressData& SavedProgressFrame);
 
 public:
 	TFunction<bool()> CancelF = []() { return false; };
@@ -254,7 +254,7 @@ public:
 
 	// Simple helper to track progress in a local scope on an optional FProgressCancel
 	// Will still work if the ProgressCancel is null (just does nothing in that case)
-	class GEOMETRYCORE_API FProgressScope
+	class FProgressScope
 	{
 		FProgressCancel* ProgressCancel;
 		FProgressData SavedProgressData;
@@ -266,7 +266,7 @@ public:
 		 * @param ProgressCancel		Progress will be tracked on this. If null, the FProgressScope will do nothing.
 		 * @param ProgressAmount		Amount to increase progress w/in this scope (as a fraction of the current outer-scope active progress range)
 		 */
-		FProgressScope(FProgressCancel* ProgressCancel, float ProgressAmount, const FText& Message = FText());
+		GEOMETRYCORE_API FProgressScope(FProgressCancel* ProgressCancel, float ProgressAmount, const FText& Message = FText());
 
 		/**
 		 * Create a dummy/inactive FProgressScope
@@ -305,7 +305,7 @@ public:
 		/**
 		 * Advance to the end of the scope's progress range and close the scope
 		 */
-		void Done();
+		GEOMETRYCORE_API void Done();
 
 		/**
 		 * @param Amount	Amount to increase the progress fraction, as a fraction of the current active progress range
@@ -378,7 +378,7 @@ public:
 	 * @param Message				Optional message describing the work to be done
 	 * @return						A new FProgressScope that covers work from the current progress to the target progress value (relative to the current scope)
 	 */
-	static FProgressScope CreateScopeTo(FProgressCancel* ProgressCancel, float ProgressTo, const FText& Message = FText());
+	static GEOMETRYCORE_API FProgressScope CreateScopeTo(FProgressCancel* ProgressCancel, float ProgressTo, const FText& Message = FText());
 
 	float GetProgress() const
 	{
@@ -513,4 +513,4 @@ bool UE::Geometry::FGeometryResult::CheckAndSetCancelled(FProgressCancel* Progre
 		return true;
 	}
 	return false;
-}
\ No newline at end of file
+}
diff --git a/Engine/Source/Runtime/GeometryCore/Public/Util/SmallListSet.h b/Engine/Source/Runtime/GeometryCore/Public/Util/SmallListSet.h
--- a/Engine/Source/Runtime/GeometryCore/Public/Util/SmallListSet.h
+++ b/Engine/Source/Runtime/GeometryCore/Public/Util/SmallListSet.h
@@ -32,11 +32,11 @@ namespace Geometry
  * 
  * @todo look at usage of TFunction, are we making unnecessary copies?
  */
-class GEOMETRYCORE_API FSmallListSet
+class FSmallListSet
 {
 protected:
 	/** This value is used to indicate Null in internal pointers */
-	static const int32 NullValue; // = -1;		// note: cannot be constexpr because we pass as reference to several functions, requires C++17
+	static GEOMETRYCORE_API const int32 NullValue; // = -1;		// note: cannot be constexpr because we pass as reference to several functions, requires C++17
 
 	/** size of initial linear-memory portion of lists */
 	static constexpr int32 BLOCKSIZE = 8;
@@ -80,7 +80,7 @@ public:
 	/**
 	 * set new number of lists
 	 */
-	void Resize(int32 NewSize);
+	GEOMETRYCORE_API void Resize(int32 NewSize);
 
 	/**
 	 * Reset to initial state
@@ -107,13 +107,13 @@ public:
 	/**
 	 * Create a list at the given ListIndex
 	 */
-	void AllocateAt(int32 ListIndex);
+	GEOMETRYCORE_API void AllocateAt(int32 ListIndex);
 
 
 	/**
 	 * Insert Value into list at ListIndex
 	 */
-	void Insert(int32 ListIndex, int32 Value);
+	GEOMETRYCORE_API void Insert(int32 ListIndex, int32 Value);
 
 
 
@@ -121,21 +121,21 @@ public:
 	 * remove Value from the list at ListIndex
 	 * @return false if Value was not in this list
 	 */
-	bool Remove(int32 ListIndex, int32 Value);
+	GEOMETRYCORE_API bool Remove(int32 ListIndex, int32 Value);
 
 
 
 	/**
 	 * Move list at FromIndex to ToIndex
 	 */
-	void Move(int32 FromIndex, int32 ToIndex);
+	GEOMETRYCORE_API void Move(int32 FromIndex, int32 ToIndex);
 
 
 
 	/**
 	 * Remove all elements from the list at ListIndex
 	 */
-	void Clear(int32 ListIndex);
+	GEOMETRYCORE_API void Clear(int32 ListIndex);
 
 
 	/**
@@ -165,14 +165,14 @@ public:
 	 * Search for the given Value in list at ListIndex
 	 * @return true if found
 	 */
-	bool Contains(int32 ListIndex, int32 Value) const;
+	GEOMETRYCORE_API bool Contains(int32 ListIndex, int32 Value) const;
 
 
 	/**
 	 * Search the list at ListIndex for a value where PredicateFunc(value) returns true
 	 * @return the found value, or the InvalidValue argument if not found
 	 */
-	int32 Find(int32 ListIndex, const TFunction<bool(int32)>& PredicateFunc, int32 InvalidValue = -1) const;
+	GEOMETRYCORE_API int32 Find(int32 ListIndex, const TFunction<bool(int32)>& PredicateFunc, int32 InvalidValue = -1) const;
 
 
 
@@ -180,13 +180,13 @@ public:
 	 * Search the list at ListIndex for a value where PredicateFunc(value) returns true, and replace it with NewValue
 	 * @return true if the value was found and replaced
 	 */
-	bool Replace(int32 ListIndex, const TFunction<bool(int32)>& PredicateFunc, int32 NewValue);
+	GEOMETRYCORE_API bool Replace(int32 ListIndex, const TFunction<bool(int32)>& PredicateFunc, int32 NewValue);
 
 
 	/**
 	 * Call ApplyFunc on each element of the list at ListIndex
 	 */
-	void Enumerate(int32 ListIndex, TFunctionRef<void(int32)> ApplyFunc) const;
+	GEOMETRYCORE_API void Enumerate(int32 ListIndex, TFunctionRef<void(int32)> ApplyFunc) const;
 
 	/**
 	 * Serialization operator for FSmallListSet.
@@ -207,7 +207,7 @@ public:
 	* @param bCompactData Only serialize unique data and/or recompute redundant data when loading.
 	* @param bUseCompression Use compression to serialize data; the resulting size will likely be smaller but serialization will take significantly longer.
 	*/
-	void Serialize(FArchive& Ar, bool bCompactData, bool bUseCompression);
+	GEOMETRYCORE_API void Serialize(FArchive& Ar, bool bCompactData, bool bUseCompression);
 
 	friend bool operator==(const FSmallListSet& Lhs, const FSmallListSet& Rhs)
 	{
@@ -524,7 +524,7 @@ protected:
 
 
 	// grab a block from the free list, or allocate a new one
-	int32 AllocateBlock();
+	GEOMETRYCORE_API int32 AllocateBlock();
 
 	// push a link-node onto the free list
 	inline void AddFreeLink(int32 ptr)
@@ -535,7 +535,7 @@ protected:
 
 
 	// remove val from the linked-list attached to block_ptr
-	bool RemoveFromLinkedList(int32 block_ptr, int32 val);
+	GEOMETRYCORE_API bool RemoveFromLinkedList(int32 block_ptr, int32 val);
 
 
 public:
@@ -550,4 +550,4 @@ public:
 
 
 } // end namespace UE::Geometry
-} // end namespace UE
\ No newline at end of file
+} // end namespace UE
diff --git a/Engine/Source/Runtime/GeometryCore/Public/Util/UniqueIndexSet.h b/Engine/Source/Runtime/GeometryCore/Public/Util/UniqueIndexSet.h
--- a/Engine/Source/Runtime/GeometryCore/Public/Util/UniqueIndexSet.h
+++ b/Engine/Source/Runtime/GeometryCore/Public/Util/UniqueIndexSet.h
@@ -22,7 +22,7 @@ namespace Geometry
  * This is faster for construction than a TSet, but does require allocating and/or clearing the bit-array.
  * Still generally quite a bit more efficient, particularly for large sets.
  */
-class GEOMETRYCORE_API FUniqueIndexSet
+class FUniqueIndexSet
 {
 	// TODO: second bit set that tracks which blocks of ints have been set.
 	// This would allow for clears to be much quicker for subsets.
@@ -37,7 +37,7 @@ class GEOMETRYCORE_API FUniqueIndexSet
 	// TODO: support growing bit set
 
 public:
-	~FUniqueIndexSet();
+	GEOMETRYCORE_API ~FUniqueIndexSet();
 
 	/**
 	 * Initialize the set with maximum index. 
diff --git a/Engine/Source/Runtime/GeometryFramework/Public/Changes/MeshChange.h b/Engine/Source/Runtime/GeometryFramework/Public/Changes/MeshChange.h
--- a/Engine/Source/Runtime/GeometryFramework/Public/Changes/MeshChange.h
+++ b/Engine/Source/Runtime/GeometryFramework/Public/Changes/MeshChange.h
@@ -27,11 +27,11 @@ namespace UE { namespace Geometry { class FDynamicMesh3; } }
  * FMeshChange represents an undoable change to a FDynamicMesh3.
  * Currently only valid to call Apply/Revert when the Object is a one of several components backed by FDynamicMesh: UDynamicMeshComponent, UOctreeDynamicMeshComponent, UPreviewMesh
  */
-class GEOMETRYFRAMEWORK_API FMeshChange : public FToolCommandChange
+class FMeshChange : public FToolCommandChange
 {
 public:
-	FMeshChange();
-	FMeshChange(TUniquePtr<UE::Geometry::FDynamicMeshChange> DynamicMeshChangeIn);
+	GEOMETRYFRAMEWORK_API FMeshChange();
+	GEOMETRYFRAMEWORK_API FMeshChange(TUniquePtr<UE::Geometry::FDynamicMeshChange> DynamicMeshChangeIn);
 
 	TUniquePtr<UE::Geometry::FDynamicMeshChange> DynamicMeshChange;
 
@@ -39,13 +39,13 @@ public:
 	TFunction<void(FMeshChange*, UObject*, bool)> OnChangeAppliedFunc;
 
 	/** Makes the change to the object */
-	virtual void Apply(UObject* Object) override;
+	GEOMETRYFRAMEWORK_API virtual void Apply(UObject* Object) override;
 
 	/** Reverts change to the object */
-	virtual void Revert(UObject* Object) override;
+	GEOMETRYFRAMEWORK_API virtual void Revert(UObject* Object) override;
 
 	/** Describes this change (for debugging) */
-	virtual FString ToString() const override;
+	GEOMETRYFRAMEWORK_API virtual FString ToString() const override;
 
 
 	/** 
@@ -54,22 +54,22 @@ public:
 	 * This allows FMeshChange subclasses to customize the change behavior if necessary.
 	 * The default behavior just forwards the call to DynamicMeshChange->Apply(Mesh, bRevert).
 	 */
-	virtual void ApplyChangeToMesh(UE::Geometry::FDynamicMesh3* Mesh, bool bRevert) const;
+	GEOMETRYFRAMEWORK_API virtual void ApplyChangeToMesh(UE::Geometry::FDynamicMesh3* Mesh, bool bRevert) const;
 };
 
 
 
 
 
-UINTERFACE()
-class GEOMETRYFRAMEWORK_API UMeshCommandChangeTarget : public UInterface
+UINTERFACE(MinimalAPI)
+class UMeshCommandChangeTarget : public UInterface
 {
 	GENERATED_BODY()
 };
 /**
  * IMeshCommandChangeTarget is an interface which is used to apply a mesh change
  */
-class GEOMETRYFRAMEWORK_API IMeshCommandChangeTarget
+class IMeshCommandChangeTarget
 {
 	GENERATED_BODY()
 public:
diff --git a/Engine/Source/Runtime/GeometryFramework/Public/Changes/MeshReplacementChange.h b/Engine/Source/Runtime/GeometryFramework/Public/Changes/MeshReplacementChange.h
--- a/Engine/Source/Runtime/GeometryFramework/Public/Changes/MeshReplacementChange.h
+++ b/Engine/Source/Runtime/GeometryFramework/Public/Changes/MeshReplacementChange.h
@@ -22,13 +22,13 @@ class UObject;
  * FMeshReplacementChange represents an undoable *complete* change to a FDynamicMesh3.
  * Currently only valid to call Apply/Revert when the Object is a UDynamicMeshComponent
  */
-class GEOMETRYFRAMEWORK_API FMeshReplacementChange : public FToolCommandChange
+class FMeshReplacementChange : public FToolCommandChange
 {
 	TSharedPtr<const FDynamicMesh3, ESPMode::ThreadSafe> Before, After;
 
 public:
-	FMeshReplacementChange();
-	FMeshReplacementChange(TSharedPtr<const FDynamicMesh3, ESPMode::ThreadSafe> Before, TSharedPtr<const FDynamicMesh3, ESPMode::ThreadSafe> After);
+	GEOMETRYFRAMEWORK_API FMeshReplacementChange();
+	GEOMETRYFRAMEWORK_API FMeshReplacementChange(TSharedPtr<const FDynamicMesh3, ESPMode::ThreadSafe> Before, TSharedPtr<const FDynamicMesh3, ESPMode::ThreadSafe> After);
 
 	const TSharedPtr<const FDynamicMesh3, ESPMode::ThreadSafe>& GetMesh(bool bRevert) const
 	{
@@ -39,28 +39,28 @@ public:
 	TFunction<void(FMeshReplacementChange*, UObject*, bool)> OnChangeAppliedFunc;
 
 	/** Makes the change to the object */
-	virtual void Apply(UObject* Object) override;
+	GEOMETRYFRAMEWORK_API virtual void Apply(UObject* Object) override;
 
 	/** Reverts change to the object */
-	virtual void Revert(UObject* Object) override;
+	GEOMETRYFRAMEWORK_API virtual void Revert(UObject* Object) override;
 
 	/** Describes this change (for debugging) */
-	virtual FString ToString() const override;
+	GEOMETRYFRAMEWORK_API virtual FString ToString() const override;
 };
 
 
 
 
 
-UINTERFACE()
-class GEOMETRYFRAMEWORK_API UMeshReplacementCommandChangeTarget : public UInterface
+UINTERFACE(MinimalAPI)
+class UMeshReplacementCommandChangeTarget : public UInterface
 {
 	GENERATED_BODY()
 };
 /**
  * IMeshReplacementCommandChangeTarget is an interface which is used to apply a mesh replacement change
  */
-class GEOMETRYFRAMEWORK_API IMeshReplacementCommandChangeTarget
+class IMeshReplacementCommandChangeTarget
 {
 	GENERATED_BODY()
 public:
diff --git a/Engine/Source/Runtime/GeometryFramework/Public/Changes/MeshVertexChange.h b/Engine/Source/Runtime/GeometryFramework/Public/Changes/MeshVertexChange.h
--- a/Engine/Source/Runtime/GeometryFramework/Public/Changes/MeshVertexChange.h
+++ b/Engine/Source/Runtime/GeometryFramework/Public/Changes/MeshVertexChange.h
@@ -17,7 +17,7 @@ PREDECLARE_USE_GEOMETRY_CLASS(FDynamicMesh3);
  * @todo support optionally storing old/new normals and tangents
  * @todo support applying to a StaticMeshComponent/MeshDescription ?
  */
-class GEOMETRYFRAMEWORK_API FMeshVertexChange : public FToolCommandChange
+class FMeshVertexChange : public FToolCommandChange
 {
 public:
 	bool bHaveVertexPositions = true;
@@ -34,13 +34,13 @@ public:
 	TArray<FVector3f> NewNormals;
 
 	/** Makes the change to the object */
-	virtual void Apply(UObject* Object) override;
+	GEOMETRYFRAMEWORK_API virtual void Apply(UObject* Object) override;
 
 	/** Reverts change to the object */
-	virtual void Revert(UObject* Object) override;
+	GEOMETRYFRAMEWORK_API virtual void Revert(UObject* Object) override;
 
 	/** Describes this change (for debugging) */
-	virtual FString ToString() const override;
+	GEOMETRYFRAMEWORK_API virtual FString ToString() const override;
 };
 
 
@@ -57,7 +57,7 @@ ENUM_CLASS_FLAGS(EMeshVertexChangeComponents);
 /**
  * FMeshVertexChangeBuilder can be used to construct a FMeshVertexChange.
  */
-class GEOMETRYFRAMEWORK_API FMeshVertexChangeBuilder
+class FMeshVertexChangeBuilder
 {
 public:
 	TUniquePtr<FMeshVertexChange> Change;
@@ -71,44 +71,44 @@ public:
 	/** If set, this function is called whenever a newly-seen VertexID is saved, parameters are (VertexID, Index) into saved-vertices array */
 	TUniqueFunction<void(int32, int32)> OnNewVertexSaved = nullptr;
 
-	FMeshVertexChangeBuilder();
-	explicit FMeshVertexChangeBuilder(EMeshVertexChangeComponents Components);
+	GEOMETRYFRAMEWORK_API FMeshVertexChangeBuilder();
+	GEOMETRYFRAMEWORK_API explicit FMeshVertexChangeBuilder(EMeshVertexChangeComponents Components);
 
-	void SaveVertexInitial(const FDynamicMesh3* Mesh, int32 VertexID);
-	void SaveVertexFinal(const FDynamicMesh3* Mesh, int32 VertexID);
+	GEOMETRYFRAMEWORK_API void SaveVertexInitial(const FDynamicMesh3* Mesh, int32 VertexID);
+	GEOMETRYFRAMEWORK_API void SaveVertexFinal(const FDynamicMesh3* Mesh, int32 VertexID);
 
 	template<typename Enumerable>
 	void SaveVertices(const FDynamicMesh3* Mesh, Enumerable Enum, bool bInitial);
 
-	void SaveOverlayNormals(const FDynamicMesh3* Mesh, const TArray<int32>& ElementIDs, bool bInitial);
-	void SaveOverlayNormals(const FDynamicMesh3* Mesh, const TSet<int32>& ElementIDs, bool bInitial);
+	GEOMETRYFRAMEWORK_API void SaveOverlayNormals(const FDynamicMesh3* Mesh, const TArray<int32>& ElementIDs, bool bInitial);
+	GEOMETRYFRAMEWORK_API void SaveOverlayNormals(const FDynamicMesh3* Mesh, const TSet<int32>& ElementIDs, bool bInitial);
 
 
 public:
 	// currently only used in vertex sculpt tool. cannot be used if bSaveColors = true
-	void UpdateVertex(int32 VertexID, const FVector3d& OldPosition, const FVector3d& NewPosition);
+	GEOMETRYFRAMEWORK_API void UpdateVertex(int32 VertexID, const FVector3d& OldPosition, const FVector3d& NewPosition);
 
 	// currently only used in element paint tool. Can only be used if bSaveColors=true and bSavePositions=false
-	void UpdateVertexColor(int32 VertexID, const FVector3f& OldColor, const FVector3f& NewColor);
+	GEOMETRYFRAMEWORK_API void UpdateVertexColor(int32 VertexID, const FVector3f& OldColor, const FVector3f& NewColor);
 
 protected:
-	void UpdateVertexFinal(int32 VertexID, const FVector3d& NewPosition);
+	GEOMETRYFRAMEWORK_API void UpdateVertexFinal(int32 VertexID, const FVector3d& NewPosition);
 
-	void UpdateOverlayNormal(int32 ElementID, const FVector3f& OldNormal, const FVector3f& NewNormal);
-	void UpdateOverlayNormalFinal(int32 ElementID, const FVector3f& NewNormal);
+	GEOMETRYFRAMEWORK_API void UpdateOverlayNormal(int32 ElementID, const FVector3f& OldNormal, const FVector3f& NewNormal);
+	GEOMETRYFRAMEWORK_API void UpdateOverlayNormalFinal(int32 ElementID, const FVector3f& NewNormal);
 };
 
 
 
-UINTERFACE()
-class GEOMETRYFRAMEWORK_API UMeshVertexCommandChangeTarget : public UInterface
+UINTERFACE(MinimalAPI)
+class UMeshVertexCommandChangeTarget : public UInterface
 {
 	GENERATED_BODY()
 };
 /**
  * IMeshVertexCommandChangeTarget is an interface which is used to apply a FMeshVertexChange
  */
-class GEOMETRYFRAMEWORK_API IMeshVertexCommandChangeTarget
+class IMeshVertexCommandChangeTarget
 {
 	GENERATED_BODY()
 public:
diff --git a/Engine/Source/Runtime/GeometryFramework/Public/Components/BaseDynamicMeshComponent.h b/Engine/Source/Runtime/GeometryFramework/Public/Components/BaseDynamicMeshComponent.h
--- a/Engine/Source/Runtime/GeometryFramework/Public/Components/BaseDynamicMeshComponent.h
+++ b/Engine/Source/Runtime/GeometryFramework/Public/Components/BaseDynamicMeshComponent.h
@@ -93,8 +93,8 @@ enum class EDynamicMeshVertexColorTransformMode : uint8
 /**
  * UBaseDynamicMeshComponent is a base interface for a UMeshComponent based on a UDynamicMesh.
  */
-UCLASS(Abstract, hidecategories = (LOD), ClassGroup = Rendering)
-class GEOMETRYFRAMEWORK_API UBaseDynamicMeshComponent : 
+UCLASS(Abstract, hidecategories = (LOD), ClassGroup = Rendering, MinimalAPI)
+class UBaseDynamicMeshComponent : 
 	public UMeshComponent, 
 	public IToolFrameworkComponent, 
 	public IMeshVertexCommandChangeTarget, 
@@ -260,7 +260,7 @@ public:
 	 * Configure the active Color Override
 	 */
 	UFUNCTION(BlueprintCallable, Category = "Dynamic Mesh Component|Rendering")
-	virtual void SetColorOverrideMode(EDynamicMeshComponentColorOverrideMode NewMode);
+	GEOMETRYFRAMEWORK_API virtual void SetColorOverrideMode(EDynamicMeshComponentColorOverrideMode NewMode);
 
 	/**
 	 * @return active Color Override mode
@@ -278,7 +278,7 @@ public:
 	 * Configure the Color used with Constant Color Override Mode
 	 */
 	UFUNCTION(BlueprintCallable, Category = "Dynamic Mesh Component|Rendering")
-	virtual void SetConstantOverrideColor(FColor NewColor);
+	GEOMETRYFRAMEWORK_API virtual void SetConstantOverrideColor(FColor NewColor);
 
 	/**
 	 * @return active Color used for Constant Color Override Mode
@@ -300,7 +300,7 @@ public:
 	 * Configure the active Color Space Transform Mode
 	 */
 	UFUNCTION(BlueprintCallable, Category = "Dynamic Mesh Component|Rendering")
-	virtual void SetVertexColorSpaceTransformMode(EDynamicMeshVertexColorTransformMode NewMode);
+	GEOMETRYFRAMEWORK_API virtual void SetVertexColorSpaceTransformMode(EDynamicMeshVertexColorTransformMode NewMode);
 
 	/**
 	 * @return active Color Override mode
@@ -325,7 +325,7 @@ public:
 	 * Configure the Color used with Constant Color Override Mode
 	 */
 	UFUNCTION(BlueprintCallable, Category = "Dynamic Mesh Component|Rendering")
-	virtual void SetEnableFlatShading(bool bEnable);
+	GEOMETRYFRAMEWORK_API virtual void SetEnableFlatShading(bool bEnable);
 
 	/**
 	 * @return active Color used for Constant Color Override Mode
@@ -343,13 +343,13 @@ public:
 public:
 
 	UFUNCTION(BlueprintCallable, Category = "Dynamic Mesh Component")
-	virtual void SetShadowsEnabled(bool bEnabled);
+	GEOMETRYFRAMEWORK_API virtual void SetShadowsEnabled(bool bEnabled);
 
 	UFUNCTION(BlueprintCallable, Category = "Dynamic Mesh Component")
 	virtual bool GetShadowsEnabled() const { return CastShadow; }
 
 	UFUNCTION(BlueprintCallable, Category = "Dynamic Mesh Component")
-	virtual void SetViewModeOverridesEnabled(bool bEnabled);
+	GEOMETRYFRAMEWORK_API virtual void SetViewModeOverridesEnabled(bool bEnabled);
 
 	UFUNCTION(BlueprintCallable, Category = "Dynamic Mesh Component")
 	virtual bool GetViewModeOverridesEnabled() const { return bEnableViewModeOverrides; }
@@ -373,13 +373,13 @@ public:
 	 * Set an active override render material. This should replace all materials during rendering.
 	 */
 	UFUNCTION(BlueprintCallable, Category = "Dynamic Mesh Component")
-	virtual void SetOverrideRenderMaterial(UMaterialInterface* Material);
+	GEOMETRYFRAMEWORK_API virtual void SetOverrideRenderMaterial(UMaterialInterface* Material);
 
 	/**
 	 * Clear any active override render material
 	 */
 	UFUNCTION(BlueprintCallable, Category = "Dynamic Mesh Component")
-	virtual void ClearOverrideRenderMaterial();
+	GEOMETRYFRAMEWORK_API virtual void ClearOverrideRenderMaterial();
 	
 	/**
 	 * @return true if an override render material is currently enabled for the given MaterialIndex
@@ -417,13 +417,13 @@ public:
 	 * Set an active secondary render material. 
 	 */
 	UFUNCTION(BlueprintCallable, Category = "Dynamic Mesh Component")
-	virtual void SetSecondaryRenderMaterial(UMaterialInterface* Material);
+	GEOMETRYFRAMEWORK_API virtual void SetSecondaryRenderMaterial(UMaterialInterface* Material);
 
 	/**
 	 * Clear any active secondary render material
 	 */
 	UFUNCTION(BlueprintCallable, Category = "Dynamic Mesh Component")
-	virtual void ClearSecondaryRenderMaterial();
+	GEOMETRYFRAMEWORK_API virtual void ClearSecondaryRenderMaterial();
 
 	/**
 	 * @return true if a secondary render material is set
@@ -446,13 +446,13 @@ public:
 	 * Show/Hide the secondary triangle buffers. Does not invalidate SceneProxy.
 	 */
 	UFUNCTION(BlueprintCallable, Category = "Dynamic Mesh Component")
-	virtual void SetSecondaryBuffersVisibility(bool bSetVisible);
+	GEOMETRYFRAMEWORK_API virtual void SetSecondaryBuffersVisibility(bool bSetVisible);
 
 	/**
 	 * @return true if secondary buffers are currently set to be visible
 	 */
 	UFUNCTION(BlueprintCallable, Category = "Dynamic Mesh Component")
-	virtual bool GetSecondaryBuffersVisibility() const;
+	GEOMETRYFRAMEWORK_API virtual bool GetSecondaryBuffersVisibility() const;
 
 protected:
 	UPROPERTY()
@@ -482,16 +482,16 @@ public:
 	 * the rendering queue and forces an immediate rebuild of the SceneProxy.
 	 */
 	UFUNCTION(BlueprintCallable, Category = "Dynamic Mesh Component")
-	virtual void SetEnableRaytracing(bool bSetEnabled);
+	GEOMETRYFRAMEWORK_API virtual void SetEnableRaytracing(bool bSetEnabled);
 
 	/**
 	 * @return true if raytracing support is currently enabled
 	 */
 	UFUNCTION(BlueprintCallable, Category = "Dynamic Mesh Component")
-	virtual bool GetEnableRaytracing() const;
+	GEOMETRYFRAMEWORK_API virtual bool GetEnableRaytracing() const;
 
 protected:
-	virtual void OnRenderingStateChanged(bool bForceImmedateRebuild);
+	GEOMETRYFRAMEWORK_API virtual void OnRenderingStateChanged(bool bForceImmedateRebuild);
 
 
 
@@ -501,17 +501,17 @@ protected:
 public:
 
 	// UMeshComponent Interface.
-	virtual int32 GetNumMaterials() const override;
-	virtual UMaterialInterface* GetMaterial(int32 ElementIndex) const override;
-	virtual FMaterialRelevance GetMaterialRelevance(ERHIFeatureLevel::Type InFeatureLevel) const override;
-	virtual void SetMaterial(int32 ElementIndex, UMaterialInterface* Material) override;
-	virtual void GetUsedMaterials(TArray<UMaterialInterface*>& OutMaterials, bool bGetDebugMaterials = false) const override;
+	GEOMETRYFRAMEWORK_API virtual int32 GetNumMaterials() const override;
+	GEOMETRYFRAMEWORK_API virtual UMaterialInterface* GetMaterial(int32 ElementIndex) const override;
+	GEOMETRYFRAMEWORK_API virtual FMaterialRelevance GetMaterialRelevance(ERHIFeatureLevel::Type InFeatureLevel) const override;
+	GEOMETRYFRAMEWORK_API virtual void SetMaterial(int32 ElementIndex, UMaterialInterface* Material) override;
+	GEOMETRYFRAMEWORK_API virtual void GetUsedMaterials(TArray<UMaterialInterface*>& OutMaterials, bool bGetDebugMaterials = false) const override;
 
-	virtual void SetNumMaterials(int32 NumMaterials);
+	GEOMETRYFRAMEWORK_API virtual void SetNumMaterials(int32 NumMaterials);
 
 	//~ UObject Interface.
 #if WITH_EDITOR
-	void PostEditChangeProperty(FPropertyChangedEvent& PropertyChangedEvent) override;
+	GEOMETRYFRAMEWORK_API void PostEditChangeProperty(FPropertyChangedEvent& PropertyChangedEvent) override;
 #endif
 
 	UPROPERTY()
@@ -532,24 +532,24 @@ public:
 	/**
 	 * Set the wireframe material used for all BaseDynamicMeshComponent-derived Components
 	 */
-	static void SetDefaultWireframeMaterial(UMaterialInterface* Material);
+	static GEOMETRYFRAMEWORK_API void SetDefaultWireframeMaterial(UMaterialInterface* Material);
 
 	/**
 	 * Set the vertex color material used for all BaseDynamicMeshComponent-derived Components
 	 */
-	static void SetDefaultVertexColorMaterial(UMaterialInterface* Material);
+	static GEOMETRYFRAMEWORK_API void SetDefaultVertexColorMaterial(UMaterialInterface* Material);
 
 protected:
-	static void InitializeDefaultMaterials();
+	static GEOMETRYFRAMEWORK_API void InitializeDefaultMaterials();
 	friend class FGeometryFrameworkModule;			// FGeometryFrameworkModule needs to call the above function
 
-	static UMaterialInterface* GetDefaultWireframeMaterial_RenderThread();
-	static UMaterialInterface* GetDefaultVertexColorMaterial_RenderThread();
+	static GEOMETRYFRAMEWORK_API UMaterialInterface* GetDefaultWireframeMaterial_RenderThread();
+	static GEOMETRYFRAMEWORK_API UMaterialInterface* GetDefaultVertexColorMaterial_RenderThread();
 	friend class FBaseDynamicMeshSceneProxy;		// FBaseDynamicMeshSceneProxy needs to call these functions...
 
 private:
 	// these Materials are used by the render thread. Once the engine is running they should not be modified without 
 	// using SetDefaultWireframeMaterial/SetDefaultVertexColorMaterial
-	static UMaterialInterface* DefaultWireframeMaterial;
-	static UMaterialInterface* DefaultVertexColorMaterial;
+	static GEOMETRYFRAMEWORK_API UMaterialInterface* DefaultWireframeMaterial;
+	static GEOMETRYFRAMEWORK_API UMaterialInterface* DefaultVertexColorMaterial;
 };
diff --git a/Engine/Source/Runtime/GeometryFramework/Public/Components/BaseDynamicMeshSceneProxy.h b/Engine/Source/Runtime/GeometryFramework/Public/Components/BaseDynamicMeshSceneProxy.h
--- a/Engine/Source/Runtime/GeometryFramework/Public/Components/BaseDynamicMeshSceneProxy.h
+++ b/Engine/Source/Runtime/GeometryFramework/Public/Components/BaseDynamicMeshSceneProxy.h
@@ -409,7 +409,7 @@ protected:
  * for a UBaseDynamicMeshComponent, where the assumption is that mesh data
  * will be stored in FMeshRenderBufferSet instances
  */
-class GEOMETRYFRAMEWORK_API FBaseDynamicMeshSceneProxy : public FPrimitiveSceneProxy
+class FBaseDynamicMeshSceneProxy : public FPrimitiveSceneProxy
 {
 	using FIndex2i = UE::Geometry::FIndex2i;
 	using FIndex3i = UE::Geometry::FIndex3i;
@@ -483,9 +483,9 @@ protected:
 	bool bEnableViewModeOverrides = true;
 
 public:
-	FBaseDynamicMeshSceneProxy(UBaseDynamicMeshComponent* Component);
+	GEOMETRYFRAMEWORK_API FBaseDynamicMeshSceneProxy(UBaseDynamicMeshComponent* Component);
 
-	virtual ~FBaseDynamicMeshSceneProxy();
+	GEOMETRYFRAMEWORK_API virtual ~FBaseDynamicMeshSceneProxy();
 
 
 	//
@@ -510,12 +510,12 @@ public:
 	 * Allocates a set of render buffers. FPrimitiveSceneProxy will keep track of these
 	 * buffers and destroy them on destruction.
 	 */
-	virtual FMeshRenderBufferSet* AllocateNewRenderBufferSet();
+	GEOMETRYFRAMEWORK_API virtual FMeshRenderBufferSet* AllocateNewRenderBufferSet();
 
 	/**
 	 * Explicitly release a set of RenderBuffers
 	 */
-	virtual void ReleaseRenderBufferSet(FMeshRenderBufferSet* BufferSet);
+	GEOMETRYFRAMEWORK_API virtual void ReleaseRenderBufferSet(FMeshRenderBufferSet* BufferSet);
 
 
 	/**
@@ -918,12 +918,12 @@ public:
 	/**
 	 * @return number of active materials
 	 */
-	virtual int32 GetNumMaterials() const;
+	GEOMETRYFRAMEWORK_API virtual int32 GetNumMaterials() const;
 
 	/**
 	 * Safe GetMaterial function that will never return nullptr
 	 */
-	virtual UMaterialInterface* GetMaterial(int32 k) const;
+	GEOMETRYFRAMEWORK_API virtual UMaterialInterface* GetMaterial(int32 k) const;
 
 	/**
 	 * Set whether or not to validate mesh batch materials against the component materials.
@@ -939,7 +939,7 @@ public:
 	 * the check in FPrimitiveSceneProxy::VerifyUsedMaterial() will fail if an override
 	 * material is set, if materials change, etc, etc
 	 */
-	virtual void UpdatedReferencedMaterials();
+	GEOMETRYFRAMEWORK_API virtual void UpdatedReferencedMaterials();
 
 
 	//
@@ -950,7 +950,7 @@ public:
 	/**
 	 * Render set of active RenderBuffers returned by GetActiveRenderBufferSets
 	 */
-	virtual void GetDynamicMeshElements(
+	GEOMETRYFRAMEWORK_API virtual void GetDynamicMeshElements(
 		const TArray<const FSceneView*>& Views, 
 		const FSceneViewFamily& ViewFamily, 
 		uint32 VisibilityMap, 
@@ -960,7 +960,7 @@ public:
 	/**
 	 * Draw a single-frame FMeshBatch for a FMeshRenderBufferSet
 	 */
-	virtual void DrawBatch(FMeshElementCollector& Collector,
+	GEOMETRYFRAMEWORK_API virtual void DrawBatch(FMeshElementCollector& Collector,
 		const FMeshRenderBufferSet& RenderBuffers,
 		const FDynamicMeshIndexBuffer32& IndexBuffer,
 		FMaterialRenderProxy* UseMaterial,
@@ -971,16 +971,16 @@ public:
 
 #if RHI_RAYTRACING
 
-	virtual bool IsRayTracingRelevant() const override;
-	virtual bool HasRayTracingRepresentation() const override;
+	GEOMETRYFRAMEWORK_API virtual bool IsRayTracingRelevant() const override;
+	GEOMETRYFRAMEWORK_API virtual bool HasRayTracingRepresentation() const override;
 
-	virtual void GetDynamicRayTracingInstances(FRayTracingMaterialGatheringContext& Context, TArray<FRayTracingInstance>& OutRayTracingInstances) override;
+	GEOMETRYFRAMEWORK_API virtual void GetDynamicRayTracingInstances(FRayTracingMaterialGatheringContext& Context, TArray<FRayTracingInstance>& OutRayTracingInstances) override;
 
 
 	/**
 	* Draw a single-frame raytracing FMeshBatch for a FMeshRenderBufferSet
 	*/
-	virtual void DrawRayTracingBatch(
+	GEOMETRYFRAMEWORK_API virtual void DrawRayTracingBatch(
 		FRayTracingMaterialGatheringContext& Context,
 		const FMeshRenderBufferSet& RenderBuffers,
 		const FDynamicMeshIndexBuffer32& IndexBuffer,
