Add a convex decomposition sphere covering mode that uses vox-solidify/difference/offset operations to better place the negative space sphere sampling

diff --git a/Engine/Plugins/Experimental/GeometryCollectionPlugin/Source/GeometryCollectionNodes/Private/Dataflow/GeometryCollectionUtilityNodes.cpp b/Engine/Plugins/Experimental/GeometryCollectionPlugin/Source/GeometryCollectionNodes/Private/Dataflow/GeometryCollectionUtilityNodes.cpp
--- a/Engine/Plugins/Experimental/GeometryCollectionPlugin/Source/GeometryCollectionNodes/Private/Dataflow/GeometryCollectionUtilityNodes.cpp
+++ b/Engine/Plugins/Experimental/GeometryCollectionPlugin/Source/GeometryCollectionNodes/Private/Dataflow/GeometryCollectionUtilityNodes.cpp
@@ -193,6 +193,19 @@ void FCreateNonOverlappingConvexHullsDataflowNode::Evaluate(Dataflow::FContext&
 	}
 }
 
+// local helper to convert the dataflow enum
+static UE::Geometry::FNegativeSpaceSampleSettings::ESampleMethod ConvertNegativeSpaceSampleMethodDataflowEnum(ENegativeSpaceSampleMethodDataflowEnum SampleMethod)
+{
+	switch (SampleMethod)
+	{
+	case ENegativeSpaceSampleMethodDataflowEnum::Uniform:
+		return UE::Geometry::FNegativeSpaceSampleSettings::ESampleMethod::Uniform;
+	case ENegativeSpaceSampleMethodDataflowEnum::VoxelSearch:
+		return UE::Geometry::FNegativeSpaceSampleSettings::ESampleMethod::VoxelSearch;
+	}
+	return UE::Geometry::FNegativeSpaceSampleSettings::ESampleMethod::Uniform;
+}
+
 FGenerateClusterConvexHullsFromLeafHullsDataflowNode::FGenerateClusterConvexHullsFromLeafHullsDataflowNode(const Dataflow::FNodeParameters& InParam, FGuid InGuid)
 	: FDataflowNode(InParam, InGuid)
 {
@@ -236,6 +249,7 @@ void FGenerateClusterConvexHullsFromLeafHullsDataflowNode::Evaluate(Dataflow::FC
 				NegativeSpaceSettings.MinRadius = GetValue(Context, &MinRadius);
 				NegativeSpaceSettings.ReduceRadiusMargin = GetValue(Context, &NegativeSpaceTolerance);
 				NegativeSpaceSettings.MinSpacing = GetValue(Context, &MinSampleSpacing);
+				NegativeSpaceSettings.SampleMethod = ConvertNegativeSpaceSampleMethodDataflowEnum(SampleMethod);
 				NegativeSpaceSettings.Sanitize();
 				bHasNegativeSpace = UE::FractureEngine::Convex::ComputeConvexHullsNegativeSpace(*GeomCollection, NegativeSpace, NegativeSpaceSettings, bHasSelectionFilter, SelectionArray);
 			}
@@ -319,6 +333,7 @@ void FGenerateClusterConvexHullsFromChildrenHullsDataflowNode::Evaluate(Dataflow
 				NegativeSpaceSettings.MinRadius = GetValue(Context, &MinRadius);
 				NegativeSpaceSettings.ReduceRadiusMargin = GetValue(Context, &NegativeSpaceTolerance);
 				NegativeSpaceSettings.MinSpacing = GetValue(Context, &MinSampleSpacing);
+				NegativeSpaceSettings.SampleMethod = ConvertNegativeSpaceSampleMethodDataflowEnum(SampleMethod);
 				NegativeSpaceSettings.Sanitize();
 				bHasNegativeSpace = UE::FractureEngine::Convex::ComputeConvexHullsNegativeSpace(*GeomCollection, NegativeSpace, NegativeSpaceSettings, bHasSelectionFilter, SelectionArray);
 			}
@@ -400,6 +415,7 @@ void FMergeConvexHullsDataflowNode::Evaluate(Dataflow::FContext& Context, const
 			NegativeSpaceSettings.MinRadius = GetValue(Context, &MinRadius);
 			NegativeSpaceSettings.ReduceRadiusMargin = GetValue(Context, &NegativeSpaceTolerance);
 			NegativeSpaceSettings.MinSpacing = GetValue(Context, &MinSampleSpacing);
+			NegativeSpaceSettings.SampleMethod = ConvertNegativeSpaceSampleMethodDataflowEnum(SampleMethod);
 			NegativeSpaceSettings.Sanitize();
 			bHasNegativeSpace = UE::FractureEngine::Convex::ComputeConvexHullsNegativeSpace(InCollection, NegativeSpace, NegativeSpaceSettings, bHasSelectionFilter, SelectionArray, false);
 		}
diff --git a/Engine/Plugins/Experimental/GeometryCollectionPlugin/Source/GeometryCollectionNodes/Public/Dataflow/GeometryCollectionUtilityNodes.h b/Engine/Plugins/Experimental/GeometryCollectionPlugin/Source/GeometryCollectionNodes/Public/Dataflow/GeometryCollectionUtilityNodes.h
--- a/Engine/Plugins/Experimental/GeometryCollectionPlugin/Source/GeometryCollectionNodes/Public/Dataflow/GeometryCollectionUtilityNodes.h
+++ b/Engine/Plugins/Experimental/GeometryCollectionPlugin/Source/GeometryCollectionNodes/Public/Dataflow/GeometryCollectionUtilityNodes.h
@@ -231,6 +231,18 @@ public:
 	UE::Geometry::FSphereCovering Spheres;
 };
 
+//~ Dataflow-specific copy of the negative space sampling method enum in ConvexDecomposition3.h,
+//~ so that it can be exposed as a UENUM
+// Method to distribute sampling spheres
+UENUM()
+enum class ENegativeSpaceSampleMethodDataflowEnum : uint8
+{
+	// Place sample spheres in a uniform grid pattern
+	Uniform,
+	// Use voxel-based subtraction and offsetting methods to specifically target concavities
+	VoxelSearch
+};
+
 /**
  *
  * Generates cluster convex hulls for leafs hulls
@@ -278,6 +290,10 @@ public:
 	UPROPERTY(EditAnywhere, Category = NegativeSpace, meta = (DataflowInput))
 	bool bProtectNegativeSpace = false;
 
+	/** Method to use to find and sample negative space */
+	UPROPERTY(EditAnywhere, Category = NegativeSpace, meta = (EditCondition = "bProtectNegativeSpace", EditConditionHides))
+	ENegativeSpaceSampleMethodDataflowEnum SampleMethod = ENegativeSpaceSampleMethodDataflowEnum::Uniform;
+
 	/** Approximate number of spheres to consider when covering negative space */
 	UPROPERTY(EditAnywhere, Category = NegativeSpace, meta = (DataflowInput, ClampMin = 1, EditCondition = "bProtectNegativeSpace", EditConditionHides))
 	int32 TargetNumSamples = 50;
@@ -343,6 +359,10 @@ public:
 	UPROPERTY(EditAnywhere, Category = NegativeSpace, meta = (DataflowInput))
 	bool bProtectNegativeSpace = false;
 
+	/** Method to use to find and sample negative space */
+	UPROPERTY(EditAnywhere, Category = NegativeSpace, meta = (EditCondition = "bProtectNegativeSpace", EditConditionHides))
+	ENegativeSpaceSampleMethodDataflowEnum SampleMethod = ENegativeSpaceSampleMethodDataflowEnum::Uniform;
+
 	/** Approximate number of spheres to consider when covering negative space */
 	UPROPERTY(EditAnywhere, Category = NegativeSpace, meta = (DataflowInput, ClampMin = 1, EditCondition = "bProtectNegativeSpace", EditConditionHides))
 	int32 TargetNumSamples = 50;
@@ -401,6 +421,10 @@ public:
 	UPROPERTY(EditAnywhere, Category = NegativeSpace, meta = (DataflowInput))
 	bool bProtectNegativeSpace = false;
 
+	/** Method to use to find and sample negative space */
+	UPROPERTY(EditAnywhere, Category = NegativeSpace, meta = (EditCondition = "bProtectNegativeSpace", EditConditionHides))
+	ENegativeSpaceSampleMethodDataflowEnum SampleMethod = ENegativeSpaceSampleMethodDataflowEnum::Uniform;
+
 	/** Approximate number of spheres to consider when covering negative space */
 	UPROPERTY(EditAnywhere, Category = NegativeSpace, meta = (DataflowInput, ClampMin = 1, EditCondition = "bProtectNegativeSpace", EditConditionHides))
 	int32 TargetNumSamples = 50;
diff --git a/Engine/Source/Runtime/GeometryCore/Private/CompGeom/ConvexDecomposition3.cpp b/Engine/Source/Runtime/GeometryCore/Private/CompGeom/ConvexDecomposition3.cpp
--- a/Engine/Source/Runtime/GeometryCore/Private/CompGeom/ConvexDecomposition3.cpp
+++ b/Engine/Source/Runtime/GeometryCore/Private/CompGeom/ConvexDecomposition3.cpp
@@ -14,6 +14,7 @@
 #include "DynamicMesh/MeshTransforms.h"
 #include "DynamicMesh/DynamicMeshAABBTree3.h"
 #include "Spatial/FastWinding.h"
+#include "Implicit/Morphology.h"
 
 #include "Async/ParallelFor.h"
 
@@ -81,6 +82,199 @@ bool FSphereCovering::AddNegativeSpace(const TFastWindingTree<FDynamicMesh3>& Sp
 			}
 		}
 	}
+	else if (SampleSettings.SampleMethod == FNegativeSpaceSampleSettings::ESampleMethod::VoxelSearch)
+	{
+		const FDynamicMesh3* Mesh = Spatial.GetTree()->GetMesh();
+
+		// Make a convex hull of the input surface as the domain to search
+		FConvexHull3d MeshHull;
+		MeshHull.Solve(Mesh->MaxVertexID(), [&](int32 VID) {return Mesh->GetVertex(VID);}, [&](int32 VID) {return Mesh->IsVertex(VID);});
+
+		// Make a light wrapper to pass the convex hull to our AABB Tree and Fast Winding classes
+		struct FHullMeshWrapper
+		{
+			const FDynamicMesh3* SourceMesh;
+			const FConvexHull3d* Hull;
+			inline int32 MaxVertexID() const
+			{
+				return SourceMesh->MaxVertexID();
+			}
+			inline bool IsVertex(int32 VID) const
+			{
+				return SourceMesh->IsVertex(VID);
+			}
+			inline FVector3d GetVertex(int32 VID) const
+			{
+				return SourceMesh->GetVertex(VID);
+			}
+			inline void GetTriVertices(int32 TID, FVector3d& V0, FVector3d& V1, FVector3d& V2) const
+			{
+				FIndex3i Tri = GetTriangle(TID);
+				V0 = GetVertex(Tri.A);
+				V1 = GetVertex(Tri.B);
+				V2 = GetVertex(Tri.C);
+			}
+			inline int32 GetChangeStamp() const
+			{
+				return 0;
+			}
+			inline int32 MaxTriangleID() const
+			{
+				return Hull->GetTriangles().Num();
+			}
+			inline int32 TriangleCount() const
+			{
+				return MaxTriangleID();
+			}
+			inline bool IsTriangle(int32 TID) const
+			{
+				return true;
+			}
+			inline FIndex3i GetTriangle(int32 TID) const
+			{
+				return Hull->GetTriangles()[TID];
+			}
+		};
+		FHullMeshWrapper HullMeshWrap{ Mesh, &MeshHull };
+		TMeshAABBTree3<FHullMeshWrapper> HullAABB(&HullMeshWrap, true);
+		TFastWindingTree<FHullMeshWrapper> HullWinding(&HullAABB, true);
+
+		// Compute the empty space inside the convex hull as (Convex Hull - Original Mesh)
+		FAxisAlignedBox3d HullBox = HullAABB.GetBoundingBox();
+		HullBox.Expand(1);
+		FMarchingCubes MarchingCubes;
+		MarchingCubes.CubeSize = FMath::Max(HullBox.MaxDim() / 128.0, SampleSettings.ReduceRadiusMargin * .5);
+		MarchingCubes.Bounds = HullBox;
+		MarchingCubes.RootMode = ERootfindingModes::Bisection;
+		MarchingCubes.RootModeSteps = 3;
+		MarchingCubes.IsoValue = 0;
+		MarchingCubes.Implicit = [&HullWinding, &Spatial](const FVector3d& Pt) -> double
+		{
+			// Volume is anything inside the hull and outside the input surface; note the input surface has negative winding
+			return (HullWinding.FastWindingNumber(Pt) > .5) && !(Spatial.FastWindingNumber(Pt) < -.5) ? 1.0 : -1.0;
+		};
+		TArray<FVector3d> MCSeeds;
+		MCSeeds.Reserve(Mesh->VertexCount());
+		for (int32 VertIdx : Mesh->VertexIndicesItr())
+		{
+			FVector3d Vertex = Mesh->GetVertex(VertIdx);
+			MCSeeds.Add(Vertex);
+		}
+		FDynamicMesh3 HullMinusComponents(&MarchingCubes.GenerateContinuation(MCSeeds));
+		HullMinusComponents.DiscardAttributes();
+
+		// Contract the surface by the ReduceRadiusMargin to avoid sampling too-small concave regions
+		FDynamicMeshAABBTree3 MorphologyBVTree(&HullMinusComponents);
+		TImplicitMorphology<FDynamicMesh3> ImplicitMorphology;
+		ImplicitMorphology.MorphologyOp = TImplicitMorphology<FDynamicMesh3>::EMorphologyOp::Contract;
+		ImplicitMorphology.Source = &HullMinusComponents;
+		ImplicitMorphology.SourceSpatial = &MorphologyBVTree;
+		ImplicitMorphology.Distance = SampleSettings.ReduceRadiusMargin;
+		ImplicitMorphology.GridCellSize = MarchingCubes.CubeSize;
+		ImplicitMorphology.MeshCellSize = MarchingCubes.CubeSize;
+		FDynamicMesh3 MorphologyMesh(&ImplicitMorphology.Generate());
+		MorphologyMesh.DiscardAttributes();
+
+		// Make sure the mesh is compact to simplify downsampling below
+		MorphologyMesh.CompactInPlace();
+
+		auto AddSample = [this, &Mesh, &Spatial, &SampleSettings, &bAddedPoints](FVector3d Pos)
+		{
+			double Winding = Spatial.FastWindingNumber(Pos);
+			if (Winding < -.5)
+			{
+				return;
+			}
+
+			double NearDistSq;
+			int NearTID = Spatial.GetTree()->FindNearestTriangle(Pos, NearDistSq);
+			double R = FMath::Sqrt(NearDistSq) - SampleSettings.ReduceRadiusMargin;
+
+			// Walk away from the closest point until we're far enough away to create a sample
+			// (give up if we haven't found a valid sample after a few steps, or if we stepped inside the shape)
+			int32 Steps = 0;
+			while (R < SampleSettings.MinRadius && Steps++ < 3)
+			{
+				bool bFoundValidSample = false;
+				FDistPoint3Triangle3d Query = TMeshQueries<FDynamicMesh3>::TriangleDistance(*Mesh, NearTID, Pos);
+				FVector3d Away = Pos - Query.ClosestTrianglePoint;
+				if (!Away.Normalize())
+				{
+					return;
+				}
+				
+				// Move away and re-test the sample
+				Pos += Away * ((SampleSettings.MinRadius - R) * 1.1);
+				if (Spatial.FastWindingNumber(Pos) >= -.5)
+				{
+					NearTID = Spatial.GetTree()->FindNearestTriangle(Pos, NearDistSq);
+					R = FMath::Sqrt(NearDistSq) - SampleSettings.ReduceRadiusMargin;
+					if (R >= SampleSettings.MinRadius)
+					{
+						break;
+					}
+				}
+				else  // give up if we stepped into the volume
+				{
+					return;
+				}
+			}
+			if (R >= SampleSettings.MinRadius)
+			{
+				bAddedPoints = true;
+				Position.Add(Pos);
+				Radius.Add(R);
+			}
+		};
+
+		// if we have more vertices than we want samples, downsample so that we:
+		//  (1) ~uniformly cover space and
+		//  (2) prioritize samples at 'features' (sharp angles) of the offset mesh
+		if (MorphologyMesh.MaxVertexID() > SampleSettings.TargetNumSamples)
+		{
+			TArray<float> VertexAngleMetric;
+			VertexAngleMetric.SetNumZeroed(MorphologyMesh.MaxVertexID());
+			TArray<FVector3d> TriNormals;
+			TriNormals.SetNumZeroed(MorphologyMesh.MaxTriangleID());
+			for (int32 TID : MorphologyMesh.TriangleIndicesItr())
+			{
+				TriNormals[TID] = MorphologyMesh.GetTriNormal(TID);
+			}
+			TArray<FVector3d> VertexPositions;
+			VertexPositions.SetNumZeroed(MorphologyMesh.MaxVertexID());
+			for (int32 VID : MorphologyMesh.VertexIndicesItr())
+			{
+				VertexPositions[VID] = MorphologyMesh.GetVertex(VID);
+				// Note: Angle metric favors vertices on edges with larger dihedral angles
+				float AngleMetric = 0;
+				MorphologyMesh.EnumerateVertexEdges(VID, [&](int32 EID)
+					{
+						FIndex2i EdgeT = MorphologyMesh.GetEdgeT(EID);
+						if (EdgeT.B != FDynamicMesh3::InvalidID)
+						{
+							AngleMetric = FMath::Max(AngleMetric, float(1 - TriNormals[EdgeT.A].Dot(TriNormals[EdgeT.B])));
+						}
+					});
+				VertexAngleMetric[VID] = AngleMetric;
+			}
+			FPriorityOrderPoints Ordering;
+			Ordering.ComputeUniformSpaced(VertexPositions, VertexAngleMetric, SampleSettings.TargetNumSamples);
+			int32 NumSamples = FMath::Min(Ordering.Order.Num(), SampleSettings.TargetNumSamples);
+			for (int32 SampleIdx = 0; SampleIdx < NumSamples; ++SampleIdx)
+			{
+				int32 VID = Ordering.Order[SampleIdx];
+				AddSample(MorphologyMesh.GetVertex(VID));
+			}
+		}
+		else
+		{
+			// When we have fewer vertices than target samples, just try adding all of them
+			for (FVector3d Pos : MorphologyMesh.VerticesItr())
+			{
+				AddSample(Pos);
+			}
+		}
+	}
 
 	return bAddedPoints;
 }
diff --git a/Engine/Source/Runtime/GeometryCore/Public/CompGeom/ConvexDecomposition3.h b/Engine/Source/Runtime/GeometryCore/Public/CompGeom/ConvexDecomposition3.h
--- a/Engine/Source/Runtime/GeometryCore/Public/CompGeom/ConvexDecomposition3.h
+++ b/Engine/Source/Runtime/GeometryCore/Public/CompGeom/ConvexDecomposition3.h
@@ -42,8 +42,10 @@ struct FNegativeSpaceSampleSettings
 
 	enum class ESampleMethod
 	{
-		Uniform
-		// TODO: Consider e.g. sampling based on local search / floodfill
+		// Place sample spheres in a uniform grid pattern
+		Uniform,
+		// Use voxel-based subtraction and offsetting methods to specifically target concavities
+		VoxelSearch
 	};
 
 	// Method used to place samples
diff --git a/Engine/Source/Runtime/GeometryCore/Public/Implicit/Morphology.h b/Engine/Source/Runtime/GeometryCore/Public/Implicit/Morphology.h
--- a/Engine/Source/Runtime/GeometryCore/Public/Implicit/Morphology.h
+++ b/Engine/Source/Runtime/GeometryCore/Public/Implicit/Morphology.h
@@ -148,10 +148,10 @@ protected:
 		ComputedSDF.ComputeMode = MeshSDFType::EComputeModes::NarrowBand_SpatialFloodFill;
 
 
-		ComputedSDF.CellSize = GridCellSize;
+		ComputedSDF.CellSize = (float)GridCellSize;
 		NarrowBandMaxDistance = UnsignedOffset + ComputedSDF.CellSize;
 		ComputedSDF.NarrowBandMaxDistance = NarrowBandMaxDistance;
-		ComputedSDF.ExactBandWidth = FMath::CeilToInt(ComputedSDF.NarrowBandMaxDistance / ComputedSDF.CellSize);
+		ComputedSDF.ExactBandWidth = FMath::CeilToInt32(ComputedSDF.NarrowBandMaxDistance / ComputedSDF.CellSize);
 
 		// for meshes with long triangles relative to the width of the narrow band, don't use the AABB tree
 		double AvgEdgeLen = TMeshQueries<TriangleMeshType>::AverageEdgeLength(*Source);
@@ -218,14 +218,14 @@ protected:
 
 		MeshSDFType SecondSDF;
 		SecondSDF.Mesh = &MCAdapter;
-		SecondSDF.CellSize = GridCellSize;
+		SecondSDF.CellSize = (float)GridCellSize;
 		SecondSDF.Spatial = nullptr;
 
 		FAxisAlignedBox3d Bounds = MarchingCubes.Bounds;
 		Bounds.Expand(MeshCellSize); // (because mesh may spill one cell over bounds)
 
 		SecondSDF.NarrowBandMaxDistance = UnsignedOffset + SecondSDF.CellSize;
-		SecondSDF.ExactBandWidth = FMath::CeilToInt(SecondSDF.NarrowBandMaxDistance / SecondSDF.CellSize);
+		SecondSDF.ExactBandWidth = FMath::CeilToInt32(SecondSDF.NarrowBandMaxDistance / SecondSDF.CellSize);
 
 		if (SecondSDF.ExactBandWidth > 1) // for larger band width, prefer using the AABB tree to do one distance per cell.  TODO: tune?
 		{
diff --git a/Engine/Source/Runtime/GeometryCore/Public/Implicit/SparseNarrowBandMeshSDF.h b/Engine/Source/Runtime/GeometryCore/Public/Implicit/SparseNarrowBandMeshSDF.h
--- a/Engine/Source/Runtime/GeometryCore/Public/Implicit/SparseNarrowBandMeshSDF.h
+++ b/Engine/Source/Runtime/GeometryCore/Public/Implicit/SparseNarrowBandMeshSDF.h
@@ -320,7 +320,7 @@ public:
 			}
 		}
 
-		float fBufferWidth = ExactBandWidth * CellSize;
+		float fBufferWidth = float(ExactBandWidth) * CellSize;
 		if (ComputeMode == EComputeModes::NarrowBand_SpatialFloodFill)
 		{
 			fBufferWidth = (float)FMath::Max(fBufferWidth, float(NarrowBandMaxDistance));
