-GeometryScript: When creating a new SkeletaMesh asset, provide an option to use the bone attributes of the dynamic mesh to set the reference skeleton of the skeletal mesh.
-Added a convenience function for getting an array of bone names from a reference skeleton.
-Added MeshBones.h/cpp that contain utility functions to deal with dynamic mesh bone attributes

diff --git a/Engine/Plugins/Experimental/GeometryScripting/Source/GeometryScriptingCore/Private/MeshBoneWeightFunctions.cpp b/Engine/Plugins/Experimental/GeometryScripting/Source/GeometryScriptingCore/Private/MeshBoneWeightFunctions.cpp
--- a/Engine/Plugins/Experimental/GeometryScripting/Source/GeometryScriptingCore/Private/MeshBoneWeightFunctions.cpp
+++ b/Engine/Plugins/Experimental/GeometryScripting/Source/GeometryScriptingCore/Private/MeshBoneWeightFunctions.cpp
@@ -13,6 +13,7 @@
 #include "Operations/TransferBoneWeights.h"
 #include "MathUtil.h"
 #include "Containers/Queue.h"
+#include "DynamicMesh/MeshBones.h"
 
 #include UE_INLINE_GENERATED_CPP_BY_NAME(MeshBoneWeightFunctions)
 
@@ -382,20 +383,7 @@ UDynamicMesh* UGeometryScriptLibrary_MeshBoneWeightFunctions::TransferBoneWeight
 		FTransferBoneWeights TransferBoneWeights(&ReadMesh, Options.SourceProfile.GetProfileName());
 		TransferBoneWeights.TransferMethod = static_cast<FTransferBoneWeights::ETransferBoneWeightsMethod>(Options.TransferMethod);
 		TransferBoneWeights.bUseParallel = true;
-		if (Options.TransferMethod == ETransferBoneWeightsMethod::InpaintWeights)
-		{
-			TransferBoneWeights.NormalThreshold = FMathd::DegToRad * Options.NormalThreshold;
-			TransferBoneWeights.SearchRadius = Options.RadiusPercentage * ReadMesh.GetBounds().DiagonalLength();
-			TransferBoneWeights.NumSmoothingIterations = Options.NumSmoothingIterations;
-			TransferBoneWeights.SmoothingStrength = Options.SmoothingStrength;
-			TransferBoneWeights.LayeredMeshSupport = Options.LayeredMeshSupport;
-		}
 
-		if (TransferBoneWeights.Validate() != EOperationValidationResult::Ok)
-		{
-			UE::Geometry::AppendError(Debug, EGeometryScriptErrorType::OperationFailed, LOCTEXT("TransferBoneWeightsFromMesh_ValidationFailed", "TransferBoneWeightsFromMesh: Invalid parameters were set for the transfer weight operator"));
-			return;
-		}
 
 		TargetMesh->EditMesh([&](FDynamicMesh3& EditMesh)
 		{
@@ -426,6 +414,20 @@ UDynamicMesh* UGeometryScriptLibrary_MeshBoneWeightFunctions::TransferBoneWeight
 				}
 			}
 			
+			if (Options.TransferMethod == ETransferBoneWeightsMethod::InpaintWeights)
+			{
+				TransferBoneWeights.NormalThreshold = FMathd::DegToRad * Options.NormalThreshold;
+				TransferBoneWeights.SearchRadius = Options.RadiusPercentage * EditMesh.GetBounds().DiagonalLength();
+				TransferBoneWeights.NumSmoothingIterations = Options.NumSmoothingIterations;
+				TransferBoneWeights.SmoothingStrength = Options.SmoothingStrength;
+				TransferBoneWeights.LayeredMeshSupport = Options.LayeredMeshSupport;
+			}
+
+			if (TransferBoneWeights.Validate() != EOperationValidationResult::Ok)
+			{
+				UE::Geometry::AppendError(Debug, EGeometryScriptErrorType::OperationFailed, LOCTEXT("TransferBoneWeightsFromMesh_ValidationFailed", "TransferBoneWeightsFromMesh: Invalid parameters were set for the transfer weight operator"));
+				return;
+			}
 			if (!TransferBoneWeights.TransferWeightsToMesh(EditMesh, Options.TargetProfile.GetProfileName()))
 			{
 				UE::Geometry::AppendError(Debug, EGeometryScriptErrorType::OperationFailed, LOCTEXT("TransferBoneWeightsFromMesh_TransferFailed", "TransferBoneWeightsFromMesh: Failed to transfer the weights"));
@@ -615,56 +617,7 @@ UDynamicMesh* UGeometryScriptLibrary_MeshBoneWeightFunctions::GetBoneChildren(
 
 			TArray<int32> ChildrenIndices;
 
-			// add direct children only
-			for (int32 Idx = 0; Idx < NumBones; ++Idx)
-			{
-				if (ParentsAttrib[Idx] == BoneIndex)
-				{
-					ChildrenIndices.Add(Idx);
-				}
-			}
-			
-			// add all grandchildren as well
-			if (bRecursive) 
-			{
-				// Map each bone index to its direct children bone indices
-				TMap<int32, TArray<int32>> ChildrenMap;
-				for (int32 Idx = 0; Idx < NumBones; ++Idx)
-				{
-					const int32 ParentIdx = ParentsAttrib[Idx];
-					if (ChildrenMap.Contains(ParentsAttrib[Idx]))
-					{
-						ChildrenMap[ParentIdx].Add(Idx);
-					}
-					else 
-					{
-						ChildrenMap.Add(ParentIdx, {Idx});
-					}
-				}
-
-				// Do breadth first search
-				TQueue<int32> BFS;
-				for (const int32 ChildIdx : ChildrenIndices)
-				{
-					BFS.Enqueue(ChildIdx);
-				}
-
-				while (!BFS.IsEmpty())
-				{
-					int32 ChildIdx; 
-					BFS.Dequeue(ChildIdx);
-
-					if (ChildrenMap.Contains(ChildIdx))
-					{
-						for (const int32 GrandChildIdx : ChildrenMap[ChildIdx])
-						{
-							BFS.Enqueue(GrandChildIdx);
-							checkSlow(ChildrenIndices.Find(GrandChildIdx) == INDEX_NONE);
-							ChildrenIndices.Add(GrandChildIdx);
-						}
-					}
-				}
-			}
+			FMeshBones::GetBoneChildren(EditMesh, BoneIndex, ChildrenIndices, bRecursive);
 			
 			// Get all information about the children
 			ChildrenInfo.SetNum(ChildrenIndices.Num());
diff --git a/Engine/Plugins/Experimental/GeometryScripting/Source/GeometryScriptingEditor/Private/CreateNewAssetUtilityFunctions.cpp b/Engine/Plugins/Experimental/GeometryScripting/Source/GeometryScriptingEditor/Private/CreateNewAssetUtilityFunctions.cpp
--- a/Engine/Plugins/Experimental/GeometryScripting/Source/GeometryScriptingEditor/Private/CreateNewAssetUtilityFunctions.cpp
+++ b/Engine/Plugins/Experimental/GeometryScripting/Source/GeometryScriptingEditor/Private/CreateNewAssetUtilityFunctions.cpp
@@ -5,6 +5,7 @@
 #include "DynamicMesh/DynamicMesh3.h"
 #include "DynamicMesh/DynamicVertexSkinWeightsAttribute.h"
 #include "DynamicMesh/DynamicBoneAttribute.h"
+#include "DynamicMesh/MeshBones.h"
 #include "UDynamicMesh.h"
 
 #include "Engine/StaticMesh.h"
@@ -40,9 +41,26 @@ using namespace UE::Geometry;
 
 #define LOCTEXT_NAMESPACE "UGeometryScriptLibrary_CreateNewAssetUtilityFunctions"
 
+static bool CreateReferenceSkeletonFromMeshLods(const TArray<FDynamicMesh3>& MeshLods, FReferenceSkeleton& RefSkeleton, bool& bOrderChanged)
+{
+	TArray<FName> BoneNames;
+	TArray<int32> BoneParentIdx;
+	TArray<FTransform> BonePose;
+		
+	if (!FMeshBones::CombineLodBonesToReferenceSkeleton(MeshLods, BoneNames, BoneParentIdx, BonePose, bOrderChanged))
+	{
+		return false;
+	}
 
+	FReferenceSkeletonModifier Modifier(RefSkeleton, nullptr);
 
+	for (int32 BoneIdx = 0; BoneIdx < BoneNames.Num(); ++BoneIdx)
+	{
+		Modifier.Add(FMeshBoneInfo(BoneNames[BoneIdx], BoneNames[BoneIdx].ToString(), BoneParentIdx[BoneIdx]), BonePose[BoneIdx]);
+	}
 
+	return true;
+}
 
 
 void UGeometryScriptLibrary_CreateNewAssetFunctions::CreateUniqueNewAssetPathName(
@@ -269,7 +287,7 @@ USkeletalMesh* UGeometryScriptLibrary_CreateNewAssetFunctions::CreateNewSkeletal
 		return nullptr;
 	}
 
-	for (int LodIdx = 0; LodIdx < FromDynamicMeshLODs.Num(); ++LodIdx)
+	for (int32 LodIdx = 0; LodIdx < FromDynamicMeshLODs.Num(); ++LodIdx)
 	{	
 		const UDynamicMesh* FromDynamicMesh = FromDynamicMeshLODs[LodIdx];
 		if (FromDynamicMesh == nullptr)
@@ -331,7 +349,7 @@ USkeletalMesh* UGeometryScriptLibrary_CreateNewAssetFunctions::CreateNewSkeletal
 	TArray<FDynamicMesh3> CopyFromDynamicMeshLODs;
 	CopyFromDynamicMeshLODs.SetNum(FromDynamicMeshLODs.Num());
 
-	for (int LodIdx = 0; LodIdx < FromDynamicMeshLODs.Num(); ++LodIdx) 
+	for (int32 LodIdx = 0; LodIdx < FromDynamicMeshLODs.Num(); ++LodIdx) 
 	{
 		const UDynamicMesh* LODMesh = FromDynamicMeshLODs[LodIdx];
 		FDynamicMesh3* CopyLODMesh = &CopyFromDynamicMeshLODs[LodIdx];
@@ -343,47 +361,81 @@ USkeletalMesh* UGeometryScriptLibrary_CreateNewAssetFunctions::CreateNewSkeletal
 	}
 
 	// Check if all LODs have bone attributes.
-	bool bMeshHasBoneAttribs = true;
-	for (const FDynamicMesh3& FromDynamicMesh : CopyFromDynamicMeshLODs)
+	int32 LODWithoutBoneAttrib = -1;
+	for (int32 Idx = 0; Idx < CopyFromDynamicMeshLODs.Num(); ++Idx)
 	{
-		bMeshHasBoneAttribs = bMeshHasBoneAttribs && FromDynamicMesh.Attributes()->HasBones();
+		if (!CopyFromDynamicMeshLODs[Idx].Attributes()->HasBones())
+		{
+			LODWithoutBoneAttrib = Idx;
+			break;
+		}
 	}
 
-	if (!bMeshHasBoneAttribs) // If at least one LOD doesn't contain the bone attributes then add LOD mesh as is
+	TUniquePtr<FReferenceSkeleton> RefSkeleton = nullptr;
+	if (LODWithoutBoneAttrib >= 0) // If at least one LOD doesn't contain the bone attributes then add LOD meshes as is
 	{
 		for (const FDynamicMesh3& FromDynamicMesh : CopyFromDynamicMeshLODs)
 		{
 			AssetOptions.SourceMeshes.DynamicMeshes.Add(&FromDynamicMesh);
 		}
+		
+		if (Options.bUseMeshBoneProportions)
+		{
+			AppendWarning(Debug, EGeometryScriptErrorType::InvalidInputs, FText::Format(LOCTEXT("CreateNewSkeletalMeshAssetFromMeshLODs_MissingBoneAttributes", "CreateNewSkeletalMeshAssetFromMeshLODs: Mesh bone proportions were requested, but the LOD {0} has no bone attributes. Proportions will be ignored."), FText::AsNumber(LODWithoutBoneAttrib)));
+		}
 	}
-	else // If bone attributes are available then attempt to reindex the weights
+	else // If bone attributes are available then attempt to reindex the weights and if requested create a ReferenceSkeleton
 	{	
-		TArray<FName> ToSkeleton;
-		ToSkeleton.Reserve(InSkeleton->GetReferenceSkeleton().GetRawRefBoneInfo().Num());
-		for (const FMeshBoneInfo& BoneInfo : InSkeleton->GetReferenceSkeleton().GetRawRefBoneInfo())
+		TArray<FName> ToSkeleton; // array of bone names in the final reference skeleton
+		bool bNeedToReindex = true; // do we need to re-index the bone weights with respect to the reference skeleton
+		if (Options.bUseMeshBoneProportions)
 		{
-			ToSkeleton.Add(BoneInfo.Name);
-		}
+			// If mesh LODs have bone attributes and the user requested to use mesh bone proportions then we create a 
+			// new reference skeleton by finding the mesh with the largest number of bones and creating reference  
+			// skeleton out of its bone attributes
+			RefSkeleton = MakeUnique<FReferenceSkeleton>();
+			if (CreateReferenceSkeletonFromMeshLods(CopyFromDynamicMeshLODs, *RefSkeleton, bNeedToReindex))
+			{
+				ToSkeleton = RefSkeleton->GetRawRefBoneNames();
 
-		for (int LodIdx = 0; LodIdx < CopyFromDynamicMeshLODs.Num(); ++LodIdx)
+				// Asset will now use the custom reference skeleton instead of the InSkeleton reference skeleton
+				AssetOptions.RefSkeleton = RefSkeleton.Get(); 
+			}
+			else 
+			{
+				// if we failed to get reference skeleton from Lods, fall back to the skeleton asset
+				ToSkeleton = InSkeleton->GetReferenceSkeleton().GetRawRefBoneNames();
+				AppendWarning(Debug, EGeometryScriptErrorType::InvalidInputs, LOCTEXT("CreateNewSkeletalMeshAssetFromMeshLODs_FailedToCombineBoneAttributesForLods", "CreateNewSkeletalMeshAssetFromMeshLODs: Failed to combine the bone attributes for the Lods. Skeleton asset will be used instead."));
+			}
+		}
+		else
+		{
+			ToSkeleton = InSkeleton->GetReferenceSkeleton().GetRawRefBoneNames();
+		}
+	
+		for (int32 LodIdx = 0; LodIdx < CopyFromDynamicMeshLODs.Num(); ++LodIdx)
 		{
 			FDynamicMesh3& FromDynamicMesh = CopyFromDynamicMeshLODs[LodIdx];
-			FDynamicMeshAttributeSet* AttribSet = FromDynamicMesh.Attributes();
 
-			// Check if the skeleton we are trying to bind the mesh to is the same as the current mesh skeleton.
-			const TArray<FName>& FromSkeleton = AttribSet->GetBoneNames()->GetAttribValues();
-
-			if (FromSkeleton != ToSkeleton)
+			if (bNeedToReindex) // potentially need to re-index the weights
 			{
-				for (const TPair<FName, TUniquePtr<FDynamicMeshVertexSkinWeightsAttribute>>& Entry : AttribSet->GetSkinWeightsAttributes())
+				FDynamicMeshAttributeSet* AttribSet = FromDynamicMesh.Attributes();
+
+				// Check if the skeleton we are trying to bind the mesh to is the same as the current mesh skeleton.
+				const TArray<FName>& FromSkeleton = AttribSet->GetBoneNames()->GetAttribValues();
+
+				if (FromSkeleton != ToSkeleton)
 				{
-					FDynamicMeshVertexSkinWeightsAttribute* SkinWeightAttrib = Entry.Value.Get();
-					
-					// Reindex the bone indices
-					if (SkinWeightAttrib->ReindexBoneIndicesToSkeleton(FromSkeleton, ToSkeleton) == false)
+					for (const TPair<FName, TUniquePtr<FDynamicMeshVertexSkinWeightsAttribute>>& Entry : AttribSet->GetSkinWeightsAttributes())
 					{
-						AppendError(Debug, EGeometryScriptErrorType::OperationFailed, FText::Format(LOCTEXT("CreateNewSkeletalMeshAssetFromMeshLODs_FailedReindexing", "CreateNewSkeletalMeshAssetFromMeshLODs: LOD {0} has invalid skinning data or the bone data is not compatible with the specified skeleton."), FText::AsNumber(LodIdx)));
-						return nullptr;
+						FDynamicMeshVertexSkinWeightsAttribute* SkinWeightAttrib = Entry.Value.Get();
+						
+						// Reindex the bone indices
+						if (SkinWeightAttrib->ReindexBoneIndicesToSkeleton(FromSkeleton, ToSkeleton) == false)
+						{
+							AppendError(Debug, EGeometryScriptErrorType::OperationFailed, FText::Format(LOCTEXT("CreateNewSkeletalMeshAssetFromMeshLODs_FailedReindexing", "CreateNewSkeletalMeshAssetFromMeshLODs: LOD {0} has invalid skinning data or the bone data is not compatible with the specified skeleton."), FText::AsNumber(LodIdx)));
+							return nullptr;
+						}
 					}
 				}
 			}
diff --git a/Engine/Plugins/Experimental/GeometryScripting/Source/GeometryScriptingEditor/Public/GeometryScript/CreateNewAssetUtilityFunctions.h b/Engine/Plugins/Experimental/GeometryScripting/Source/GeometryScriptingEditor/Public/GeometryScript/CreateNewAssetUtilityFunctions.h
--- a/Engine/Plugins/Experimental/GeometryScripting/Source/GeometryScriptingEditor/Public/GeometryScript/CreateNewAssetUtilityFunctions.h
+++ b/Engine/Plugins/Experimental/GeometryScripting/Source/GeometryScriptingEditor/Public/GeometryScript/CreateNewAssetUtilityFunctions.h
@@ -85,6 +85,10 @@ public:
 
 	UPROPERTY(BlueprintReadWrite, Category = Options)
 	TMap<FName, TObjectPtr<UMaterialInterface>> Materials;
+
+	/** If true, will use the skeleton proportions (if availabale) stored in the dynamic mesh. */
+	UPROPERTY(BlueprintReadWrite, Category = Options)
+	bool bUseMeshBoneProportions = false;
 };
 
 
diff --git a/Engine/Plugins/Runtime/MeshModelingToolset/Source/ModelingComponentsEditorOnly/Private/AssetUtils/CreateSkeletalMeshUtil.cpp b/Engine/Plugins/Runtime/MeshModelingToolset/Source/ModelingComponentsEditorOnly/Private/AssetUtils/CreateSkeletalMeshUtil.cpp
--- a/Engine/Plugins/Runtime/MeshModelingToolset/Source/ModelingComponentsEditorOnly/Private/AssetUtils/CreateSkeletalMeshUtil.cpp
+++ b/Engine/Plugins/Runtime/MeshModelingToolset/Source/ModelingComponentsEditorOnly/Private/AssetUtils/CreateSkeletalMeshUtil.cpp
@@ -102,7 +102,8 @@ UE::AssetUtils::ECreateSkeletalMeshResult UE::AssetUtils::CreateSkeletalMeshAsse
 	}
 	
 	if (!FStaticToSkeletalMeshConverter::InitializeSkeletalMeshFromMeshDescriptions(
-		NewSkeletalMesh, MeshDescriptions, MaterialView, Options.Skeleton->GetReferenceSkeleton(),
+		NewSkeletalMesh, MeshDescriptions, MaterialView, 
+		Options.RefSkeleton ? *Options.RefSkeleton : Options.Skeleton->GetReferenceSkeleton(),
 		Options.bEnableRecomputeNormals, Options.bEnableRecomputeTangents))
 	{
 		return ECreateSkeletalMeshResult::UnknownError;
diff --git a/Engine/Plugins/Runtime/MeshModelingToolset/Source/ModelingComponentsEditorOnly/Public/AssetUtils/CreateSkeletalMeshUtil.h b/Engine/Plugins/Runtime/MeshModelingToolset/Source/ModelingComponentsEditorOnly/Public/AssetUtils/CreateSkeletalMeshUtil.h
--- a/Engine/Plugins/Runtime/MeshModelingToolset/Source/ModelingComponentsEditorOnly/Public/AssetUtils/CreateSkeletalMeshUtil.h
+++ b/Engine/Plugins/Runtime/MeshModelingToolset/Source/ModelingComponentsEditorOnly/Public/AssetUtils/CreateSkeletalMeshUtil.h
@@ -10,7 +10,7 @@ class USkeletalMesh;
 class USkeleton;
 struct FMeshDescription;
 struct FSkeletalMaterial;
-
+struct FReferenceSkeleton;
 
 namespace UE::AssetUtils
 {
@@ -52,6 +52,9 @@ namespace UE::AssetUtils
 		// The associated USkeleton object
 		USkeleton* Skeleton = nullptr;
 
+		// The associated FReferenceSkeleton object. If null, the reference skeleton of the Skeleton object is used instead.
+		FReferenceSkeleton* RefSkeleton = nullptr;
+
 		// Number of SourceModels (ie LODs)
 		int32 NumSourceModels = 1;
 		// Number of Material Slots desired on the Asset. At least one will always be created.
diff --git a/Engine/Source/Runtime/Engine/Public/ReferenceSkeleton.h b/Engine/Source/Runtime/Engine/Public/ReferenceSkeleton.h
--- a/Engine/Source/Runtime/Engine/Public/ReferenceSkeleton.h
+++ b/Engine/Source/Runtime/Engine/Public/ReferenceSkeleton.h
@@ -265,6 +265,20 @@ public:
 
 	const TMap<FName, int32>& GetRawNameToIndexMap() const { return RawNameToIndexMap; }
 	
+	/** Returns an array of raw bone names. */
+	TArray<FName> GetRawRefBoneNames() const
+	{
+		TArray<FName> BoneNames;
+		BoneNames.Reserve(RawRefBoneInfo.Num());
+
+		for (const FMeshBoneInfo& BoneInfo : RawRefBoneInfo)
+		{
+			BoneNames.Add(BoneInfo.Name);
+		}
+
+		return BoneNames;
+	}
+	
 	void Empty(int32 Size=0)
 	{
 		RawRefBoneInfo.Empty(Size);
diff --git a/Engine/Source/Runtime/GeometryCore/Private/DynamicMesh/MeshBones.cpp b/Engine/Source/Runtime/GeometryCore/Private/DynamicMesh/MeshBones.cpp
new file mode 100644
--- /dev/null
+++ b/Engine/Source/Runtime/GeometryCore/Private/DynamicMesh/MeshBones.cpp
@@ -0,0 +1,259 @@
+// Copyright Epic Games, Inc. All Rights Reserved.
+
+#include "DynamicMesh/MeshBones.h"
+#include "DynamicMesh/DynamicMesh3.h"
+#include "DynamicMesh/DynamicMeshAttributeSet.h"
+#include "DynamicMesh/DynamicBoneAttribute.h"
+
+#include "Containers/Queue.h"
+#include "Containers/Map.h"
+#include "Containers/Set.h"
+
+#include "UObject/NameTypes.h"
+#include "Math/Transform.h"
+#include "Misc/AssertionMacros.h"
+
+using namespace UE::Geometry;
+
+bool FMeshBones::GetBoneChildren(const FDynamicMesh3& Mesh, int32 BoneIndex, TArray<int32>& ChildrenIndices, bool bRecursive)
+{
+    if (!Mesh.HasAttributes() || !Mesh.Attributes()->HasBones())
+    {
+        return false;
+    }
+
+    const int32 NumBones = Mesh.Attributes()->GetNumBones();
+    const TArray<int32>& ParentsAttrib = Mesh.Attributes()->GetBoneParentIndices()->GetAttribValues();
+
+    // add direct children only
+	for (int32 Idx = 0; Idx < NumBones; ++Idx)
+	{
+		if (ParentsAttrib[Idx] == BoneIndex)
+		{
+			ChildrenIndices.Add(Idx);
+		}
+	}
+	
+	// add all grandchildren as well
+	if (bRecursive) 
+	{
+		// Map each bone index to its direct children bone indices
+		TMap<int32, TArray<int32>> ChildrenMap;
+		for (int32 Idx = 0; Idx < NumBones; ++Idx)
+		{
+			const int32 ParentIdx = ParentsAttrib[Idx];
+			if (ChildrenMap.Contains(ParentsAttrib[Idx]))
+			{
+				ChildrenMap[ParentIdx].Add(Idx);
+			}
+			else 
+			{
+				ChildrenMap.Add(ParentIdx, {Idx});
+			}
+		}
+
+		// Do breadth first search
+		TQueue<int32> BFS;
+		for (const int32 ChildIdx : ChildrenIndices)
+		{
+			BFS.Enqueue(ChildIdx);
+		}
+
+		while (!BFS.IsEmpty())
+		{
+			int32 ChildIdx; 
+			BFS.Dequeue(ChildIdx);
+
+			if (ChildrenMap.Contains(ChildIdx))
+			{
+				for (const int32 GrandChildIdx : ChildrenMap[ChildIdx])
+				{
+					BFS.Enqueue(GrandChildIdx);
+					checkSlow(ChildrenIndices.Find(GrandChildIdx) == INDEX_NONE);
+					ChildrenIndices.Add(GrandChildIdx);
+				}
+			}
+		}
+	}
+
+    return true;
+}
+
+bool FMeshBones::GetBonesInIncreasingOrder(const FDynamicMesh3& Mesh, 
+										   TArray<FName>& BoneNames, 
+										   TArray<int32>& BoneParentIdx, 
+										   TArray<FTransform>& BonePose, 
+										   bool& bOrderChanged)
+{
+    if (!Mesh.HasAttributes() || !Mesh.Attributes()->HasBones())
+    {
+        return false;
+    }
+
+    const TArray<FName>& NamesAttrib = Mesh.Attributes()->GetBoneNames()->GetAttribValues();
+    const TArray<int32>& ParentsAttrib = Mesh.Attributes()->GetBoneParentIndices()->GetAttribValues();
+    const TArray<FTransform>& TransformsAttrib = Mesh.Attributes()->GetBonePoses()->GetAttribValues();
+
+	const int32 NumBones = NamesAttrib.Num();
+	if (NumBones == 0)
+	{ 
+		return true; // empty skeleton, nothing to do
+	}
+
+	// Do a scan to see if the order of the bones in the mesh bone attribute is in strictly increasing order.
+	// So child must have an index greater than its parent (unless its a root bone)
+	bool bIsIncreasingOrder = true;
+	for (int32 BoneIdx = 0; BoneIdx < NumBones; ++BoneIdx)
+	{
+		const int32 ParentBoneIndex = ParentsAttrib[BoneIdx];
+		if (ParentBoneIndex == INDEX_NONE && BoneIdx != 0)
+		{
+			// root bone must be at index 0. Multi root is not supported in general Skeleton/SkeletalMesh and hence not supported in DynamicMesh.
+			bIsIncreasingOrder = false;
+			break;
+		}
+		else if (BoneIdx <= ParentBoneIndex)
+		{
+			bIsIncreasingOrder = false; // child index must be greater than it's parent
+			break;
+		}
+	}
+
+	if (!bIsIncreasingOrder)
+	{
+		// If bones in the bone attribute are out of order, traverse the hiearchy in breadth first search way to get the 
+		// correct order
+
+		// Root bone's parent will be INDEX_NONE
+		const int32 RootBoneIdx = ParentsAttrib.Find(INDEX_NONE);
+		if (!ensure(RootBoneIdx != INDEX_NONE))
+		{
+			return false; // couldn't find the root
+		}
+
+		// Get skeleton hiearchy by getting all the children and grandchildren of the root bone
+		TArray<int32> ChildrenIndices;
+		ChildrenIndices.Reserve(NumBones);
+		FMeshBones::GetBoneChildren(Mesh, RootBoneIdx, ChildrenIndices, true);
+		
+		// map bone name to its index in the array.
+		TMap<FName, int32> NameToIndex;
+
+		// preallocate space
+		NameToIndex.Reserve(NumBones);
+		BoneNames.Reserve(NumBones); 
+		BoneParentIdx.Reserve(NumBones);
+		BonePose.Reserve(NumBones); 
+
+		// add root first
+		NameToIndex.Add(NamesAttrib[RootBoneIdx], 0);
+		BoneNames.Add(NamesAttrib[RootBoneIdx]); 
+		BoneParentIdx.Add(INDEX_NONE);
+		BonePose.Add(TransformsAttrib[RootBoneIdx]);
+		
+		// add all children/grandchildren in the correct order
+		for (const int32 BoneIdx : ChildrenIndices)
+		{
+			NameToIndex.Add(NamesAttrib[BoneIdx], NameToIndex.Num());
+			BoneNames.Add(NamesAttrib[BoneIdx]); 
+			BonePose.Add(TransformsAttrib[BoneIdx]);
+		}
+
+		// reindex the parent indices
+		for (const int32 BoneIdx : ChildrenIndices)
+		{
+			const int32 OldParentIdx = ParentsAttrib[BoneIdx];
+			const FName OldParentName = NamesAttrib[OldParentIdx];
+			const int32 NewParentIdx = NameToIndex[OldParentName];
+			BoneParentIdx.Add(NewParentIdx);
+		}
+	}
+	else
+	{
+		BoneNames = NamesAttrib;
+		BoneParentIdx = ParentsAttrib;
+		BonePose = TransformsAttrib;
+	}
+
+	bOrderChanged = !bIsIncreasingOrder;
+
+	return true;
+}
+
+bool FMeshBones::CombineLodBonesToReferenceSkeleton(const TArray<FDynamicMesh3>& Meshes, 
+													TArray<FName>& BoneNames,
+													TArray<int32>& BoneParentIdx,
+													TArray<FTransform>& BonePose,  
+													bool& bOrderChanged)
+{
+    bOrderChanged = false;
+
+    if (Meshes.IsEmpty())
+    {
+        return true;
+    }
+
+    // all meshes must have bone attributes
+    for (const FDynamicMesh3& Mesh : Meshes)
+    {
+        if (Mesh.HasAttributes() == false || !Mesh.Attributes()->HasBones())
+        {
+            return false;
+        }
+    }
+
+    // check if all meshes have identical bone attributes which should be a common case when the array represents Lods
+    for (int32 LodIdx = 1; LodIdx < Meshes.Num(); ++LodIdx)
+    {
+        if (!Meshes[0].Attributes()->IsSameBoneAttributesAs(*Meshes[LodIdx].Attributes()))
+        {
+            bOrderChanged = true;
+            break;
+        }
+    }
+
+    if (!bOrderChanged)
+    {
+        // if all meshes have the same bone attributes then simply copy over the attributes of the first mesh
+        return GetBonesInIncreasingOrder(Meshes[0], BoneNames, BoneParentIdx, BonePose, bOrderChanged);
+    }
+    else 
+    {
+        // meshes don't have the same attributes so we find the mesh with the largest number of bones
+        int32 MaxBoneNum = -1;
+        int32 MaxBoneNumLod = -1;
+        for (int32 Idx = 0; Idx < Meshes.Num(); ++Idx)
+        {
+            const int32 BoneNum = Meshes[Idx].Attributes()->GetNumBones();
+            if (MaxBoneNum < BoneNum)
+            {
+                MaxBoneNum = BoneNum;
+                MaxBoneNumLod = Idx;
+            } 
+        }
+
+        // we check that it's a superset of the bone attributes in all other meshes
+        TSet<FName> HashSet;
+       	HashSet.Append(Meshes[MaxBoneNumLod].Attributes()->GetBoneNames()->GetAttribValues());
+        for (int32 Idx = 0; Idx < Meshes.Num(); ++Idx)
+        {
+            if (Idx == MaxBoneNumLod)
+            {
+                continue;
+            }
+
+            const TArray<FName>& NamesAttrib = Meshes[Idx].Attributes()->GetBoneNames()->GetAttribValues();
+            for (int32 BoneIdx = 0; BoneIdx < NamesAttrib.Num(); ++BoneIdx)
+            {
+                if (!ensure(HashSet.Contains(NamesAttrib[BoneIdx])))
+                {
+                    // The mesh has a bone that is not contained in the mesh with the largest number of bones.
+                    // Hence, the lods are not compatible.
+                    return false;
+                }
+            } 
+        }
+
+        return GetBonesInIncreasingOrder(Meshes[MaxBoneNumLod], BoneNames, BoneParentIdx, BonePose, bOrderChanged);
+    }
+}
diff --git a/Engine/Source/Runtime/GeometryCore/Public/DynamicMesh/MeshBones.h b/Engine/Source/Runtime/GeometryCore/Public/DynamicMesh/MeshBones.h
new file mode 100644
--- /dev/null
+++ b/Engine/Source/Runtime/GeometryCore/Public/DynamicMesh/MeshBones.h
@@ -0,0 +1,75 @@
+// Copyright Epic Games, Inc. All Rights Reserved.
+#pragma once
+
+#include "CoreTypes.h"
+#include "Containers/Array.h"
+#include "Math/MathFwd.h"
+
+// Forward declarations
+namespace UE::Geometry 
+{ 
+    class FDynamicMesh3; 
+}
+class FName;
+
+namespace UE
+{
+namespace Geometry
+{
+
+/**
+ * FMeshBones is a utility class for manipulating mesh bone attributes
+ */
+class FMeshBones
+{
+public:
+
+	/**
+     * Return an array of indices into the bone attribute arrays of all the children of the bone. Optionally, recursively
+     * add all grandchildren via a breadth-first search.
+     *
+     * @param BoneIndex the index of the bone whose children we are requesting
+     * @param bRecursive if true, recursively add all grandchildren
+     */
+	static GEOMETRYCORE_API bool GetBoneChildren(const FDynamicMesh3& Mesh, int32 BoneIndex, TArray<int32>& ChildrenIndices, bool bRecursive = false);
+
+	/**
+     * Create bone data arrays that can be used to create a FReferenceSkeleton.
+     *
+     * @param bOrderChanged true, if the order of the bones in the output data arrays is different than
+     *                      the order in the mesh bone attribute
+     *
+     * @return false if the dynamic mesh has no bone attributes
+     *
+     * @note the order of the bones in the output arrays might be different than the order of the bones in
+     * the mesh bone attribute. The reference skeleton expects the parent bone to always have a lower index than the
+     * children. This doesn't have to be true for dynamic mesh bone attributes. You can use
+     * TDynamicVertexSkinWeightsAttribute::ReindexBoneIndicesToSkeleton method to reindex the skinning weights if the
+     * order changed (bOutOrderChanged == true).
+     */
+	static GEOMETRYCORE_API bool GetBonesInIncreasingOrder(const FDynamicMesh3& Mesh, 
+                                                           TArray<FName>& BoneNames,
+                                                           TArray<int32>& BoneParentIdx,
+                                                           TArray<FTransform>& BonePose, 
+                                                           bool& bOrderChanged);
+
+
+	/** 
+     * Create a superset reference skeleton out of all the bones in the mesh Lods. 
+     * 
+     * @param bOrderChanged true, if all Lods have the same bone attributes and they are ordered correctly such that 
+     *                      the parent bone always has a lower index than the children. Hence the arrays can be directly
+     *                      added to an empty FReferenceSkeleton instance.
+     *
+     * @return false if one of the Lods has no bone attributes or Lods have incompatible skeletons
+     */
+	static GEOMETRYCORE_API bool CombineLodBonesToReferenceSkeleton(const TArray<FDynamicMesh3>& Meshes, 
+                                                                    TArray<FName>& BoneNames,
+                                                                    TArray<int32>& BoneParentIdx,
+                                                                    TArray<FTransform>& BonePose, 
+                                                                    bool& bOrderChanged);
+
+};
+
+} // end namespace UE::Geometry
+} // end namespace UE
