Fix crash when spline mesh tool has a spline with too few points, and generally make polygon triangulation not crash when passed an input with too few points

diff --git a/Engine/Plugins/Runtime/MeshModelingToolset/Source/ModelingOperators/Private/CurveOps/TriangulateCurvesOp.cpp b/Engine/Plugins/Runtime/MeshModelingToolset/Source/ModelingOperators/Private/CurveOps/TriangulateCurvesOp.cpp
--- a/Engine/Plugins/Runtime/MeshModelingToolset/Source/ModelingOperators/Private/CurveOps/TriangulateCurvesOp.cpp
+++ b/Engine/Plugins/Runtime/MeshModelingToolset/Source/ModelingOperators/Private/CurveOps/TriangulateCurvesOp.cpp
@@ -89,6 +89,11 @@ void FTriangulateCurvesOp::CalculateResult(FProgressCancel* Progress)
 	{
 		for (int32 PathIdx = 0; PathIdx < Paths.Num(); ++PathIdx)
 		{
+			if (Paths[PathIdx].Vertices.Num() < 3)
+			{
+				continue;
+			}
+			
 			FVector3d Normal, PlanePoint;
 			PolygonTriangulation::ComputePolygonPlane<double>(Paths[PathIdx].Vertices, Normal, PlanePoint);
 			if (bFlipResult)
diff --git a/Engine/Source/Runtime/GeometryCore/Private/CompGeom/PolygonTriangulation.cpp b/Engine/Source/Runtime/GeometryCore/Private/CompGeom/PolygonTriangulation.cpp
--- a/Engine/Source/Runtime/GeometryCore/Private/CompGeom/PolygonTriangulation.cpp
+++ b/Engine/Source/Runtime/GeometryCore/Private/CompGeom/PolygonTriangulation.cpp
@@ -36,10 +36,16 @@ void PolygonTriangulation::TriangulateSimplePolygon(const TArray<TVector2<T>>& V
 			return TriSignedArea * OrientationSign < 0;
 		}
 	};
+	
+	
+	OutTriangles.Reset();
 
-	// Polygon must have at least three vertices/edges
+	// Polygon must have at least three vertices/edges, or result is empty
 	int32 PolygonVertexCount = VertexPositions.Num();
-	check(PolygonVertexCount >= 3);
+	if (PolygonVertexCount < 3)
+	{
+		return;
+	}
 
 
 	// compute signed area of polygon
@@ -55,7 +61,6 @@ void PolygonTriangulation::TriangulateSimplePolygon(const TArray<TVector2<T>>& V
 	T OrientationSign = (bIsClockwise) ? -T(1) : T(1);
 
 
-	OutTriangles.Reset();
 
 
 	// If perimeter has 3 vertices, just copy content of perimeter out 
@@ -247,12 +252,15 @@ void PolygonTriangulation::TriangulateSimplePolygon(const TArray<TVector<T>>& Ve
 		}
 	};
 
-	// Polygon must have at least three vertices/edges
-	int32 PolygonVertexCount = VertexPositions.Num();
-	check(PolygonVertexCount >= 3);
-
 	OutTriangles.Reset();
 
+	// Polygon must have at least three vertices/edges, or result is empty
+	int32 PolygonVertexCount = VertexPositions.Num();
+	if (PolygonVertexCount < 3)
+	{
+		return;
+	}
+
 	// If perimeter has 3 vertices, just copy content of perimeter out 
 	if (PolygonVertexCount == 3)
 	{
