Add a new polygon-based simplification method to ConvexHull3, and a tolerance for the initial hull build, to support use cases where we do not need an exact hull and can prioritize using fewer plane + avoiding near-degenerate planes in the output. This simplification is especially effective for cases where a large number of near-coplanar faces would exist in the exact-hull output.

diff --git a/Engine/Source/Runtime/GeometryCore/Private/CompGeom/ConvexHull3.cpp b/Engine/Source/Runtime/GeometryCore/Private/CompGeom/ConvexHull3.cpp
--- a/Engine/Source/Runtime/GeometryCore/Private/CompGeom/ConvexHull3.cpp
+++ b/Engine/Source/Runtime/GeometryCore/Private/CompGeom/ConvexHull3.cpp
@@ -7,7 +7,9 @@
 
 #include "CompGeom/ConvexHull3.h"
 #include "CompGeom/ExactPredicates.h"
+#include "VertexConnectedComponents.h" // for FSizedDisjointSet
 
+#include "Algo/RemoveIf.h"
 #include "Async/ParallelFor.h"
 
 namespace UE
@@ -192,6 +194,38 @@ struct FHullConnectivity
 			Indices.Add(Idx);
 		}
 
+		// Remove a point from the visible point set; if it was the tracked MaxValue point, find a new MaxValue point
+		void RemovePt(int32 SourcePointIdx, const TVector<RealType> TriPts[3], TFunctionRef<TVector<RealType>(int32)> GetPointFunc)
+		{
+			if (MaxIdx == SourcePointIdx)
+			{
+				MaxValue = 0;
+				MaxIdx = -1;
+				for (int32 SubIdx = 0; SubIdx < Indices.Num(); ++SubIdx)
+				{
+					int32 PointIdx = Indices[SubIdx];
+					if (PointIdx == SourcePointIdx)
+					{
+						Indices.RemoveAtSwap(SubIdx, 1, false);
+						SubIdx--;
+					}
+					else
+					{
+						double Value = ExactPredicates::Orient3<RealType>(TriPts[0], TriPts[1], TriPts[2], GetPointFunc(PointIdx));
+						if (Value > MaxValue)
+						{
+							MaxValue = Value;
+							MaxIdx = PointIdx;
+						}
+					}
+				}
+			}
+			else
+			{
+				Indices.RemoveSingleSwap(SourcePointIdx, false);
+			}
+		}
+
 		int32 Num()
 		{
 			return Indices.Num();
@@ -457,6 +491,20 @@ struct FHullConnectivity
 		return FoundTriPointPair;
 	}
 
+	/**
+	 * @param TriPointPair		A tri point pair, as returned from ChooseVisiblePoint, that should be removed from consideration
+	 */
+	void RemoveVisiblePoint(FIndex2i TriPointPair, const TArray<FIndex3i>& Triangles, TFunctionRef<TVector<RealType>(int32)> GetPointFunc)
+	{
+		TVector<RealType> TriPts[3];
+		SetTriPts(Triangles[TriPointPair.A], GetPointFunc, TriPts);
+		VisiblePoints[TriPointPair.A].RemovePt(TriPointPair.B, TriPts, GetPointFunc);
+		if (VisiblePoints[TriPointPair.A].Indices.IsEmpty())
+		{
+			TrisWithPoints.Remove(TriPointPair.A);
+		}
+	}
+
 	/**
 	 * Helper struct carrying enough information to describe triangles that will be added when we UpdateHullWithNewPoint
 	 */
@@ -473,6 +521,16 @@ struct FHullConnectivity
 		int32 ConnectedTri;
 	};
 
+	inline void ClearVisiblePointsData(int32 TriIdx, TArray<int32>& NewlyUnclaimed)
+	{
+		if (VisiblePoints[TriIdx].Num() > 0)
+		{
+			TrisWithPoints.Remove(TriIdx);
+			NewlyUnclaimed.Append(VisiblePoints[TriIdx].Indices);
+			VisiblePoints[TriIdx].Reset();
+		}
+	}
+
 	/**
 	 * Recursive traversal strategy that starts from a 'visible' triangle and walks to the connected set of all visible triangles, s.t. it visits the boundary in order
 	 * See explanation here: http://algolist.ru/maths/geom/convhull/qhull3d.php
@@ -487,7 +545,7 @@ struct FHullConnectivity
 	 * @param CrossedEdgeFirstVertex The first vertex of the edge that was 'crossed over' to traverse to this TriIdx, or -1 for the initial call
 	 * @return false if triangle is beyond horizon (so we don't need to search through it / remove it), true otherwise
 	 */
-	bool HorizonHelper(const TArray<FIndex3i>& Triangles, TFunctionRef<TVector<RealType>(int32)> GetPointFunc, const TVector<RealType>& Pt, TArray<int32>& NewlyUnclaimed, TSet<int32>& ToDelete, TArray<FNewTriangle>& ToAdd, int32 TriIdx, int32 CrossedEdgeFirstVertex)
+	bool HorizonHelper(const TArray<FIndex3i>& Triangles, TFunctionRef<TVector<RealType>(int32)> GetPointFunc, const TVector<RealType>& Pt, TArray<int32>& NewlyUnclaimed, TSet<int32>& ToDelete, TArray<FNewTriangle>& ToAdd, int32 TriIdx, int32 CrossedEdgeFirstVertex, bool bCouldSkipPoint = false)
 	{
 		// if it's not the first triangle, crossed edge should be set and we should check if the triangle is visible / actually needs to be replaced
 		if (CrossedEdgeFirstVertex != -1)
@@ -502,16 +560,13 @@ struct FHullConnectivity
 		}
 
 		// track the triangle as needing deletion; for simplicity wait until after traversal to actually delete
-		ToDelete.Add(TriIdx); // TODO: could we delete as we go rather than keep this set?  seems tricky
-
-		// clean out the visible points right away
-		if (VisiblePoints[TriIdx].Num() > 0)
+		ToDelete.Add(TriIdx); // TODO: could we delete as we go rather than keep this set?  seems tricky (also, could not do so if 'bCouldSkipPoint == true')
+		
+		// if we know we can't ultimately skip adding this point, clean out the visible points right away; otherwise this must be done after
+		if (!bCouldSkipPoint)
 		{
-			TrisWithPoints.Remove(TriIdx);
-			NewlyUnclaimed.Append(VisiblePoints[TriIdx].Indices);
-			VisiblePoints[TriIdx].Reset();
+			ClearVisiblePointsData(TriIdx, NewlyUnclaimed);
 		}
-		
 
 		FIndex3i Tri = Triangles[TriIdx];
 		int32 FirstOff = 0, OffMax = 3; // Cross all three edges for first triangle
@@ -539,8 +594,9 @@ struct FHullConnectivity
 		return true;
 	}
 
-	void UpdateHullWithNewPoint(TArray<FIndex3i>& Triangles, TFunctionRef<TVector<RealType>(int32)> GetPointFunc, int32 StartTriIdx, int32 PtIdx)
+	bool UpdateHullWithNewPoint(TArray<FIndex3i>& Triangles, TFunctionRef<TVector<RealType>(int32)> GetPointFunc, int32 StartTriIdx, int32 PtIdx, RealType DegenerateEdgeToleranceSq = (RealType)0)
 	{
+		// Note: Commented-out ValidateConnectivity calls are very slow, but useful for debugging if the algorithm produces an invalid result
 		//ValidateConnectivity(Triangles);
 
 		TVector<RealType> Pt = GetPointFunc(PtIdx);
@@ -549,7 +605,27 @@ struct FHullConnectivity
 		TArray<int32> NewlyUnclaimed;
 		TSet<int32> ToDelete;
 		TArray<FNewTriangle> ToAdd;
-		HorizonHelper(Triangles, GetPointFunc, Pt, NewlyUnclaimed, ToDelete, ToAdd, StartTriIdx, -1);
+		bool bCouldSkipPoint = DegenerateEdgeToleranceSq > (RealType)0;
+		HorizonHelper(Triangles, GetPointFunc, Pt, NewlyUnclaimed, ToDelete, ToAdd, StartTriIdx, -1, bCouldSkipPoint);
+
+		// Optionally skip points if they are almost on top of an existing point on the current hull, as specified by the Degenerate Edge Tolerance (squared)
+		if (bCouldSkipPoint)
+		{
+			for (const FNewTriangle& NewTri : ToAdd)
+			{
+				RealType DistSq = TVector<RealType>::DistSquared(Pt, GetPointFunc(NewTri.EdgeVertices.A));
+				if (DistSq < DegenerateEdgeToleranceSq)
+				{
+					return false;
+				}
+			}
+
+			// Once we know the point will not be skipped, clear out the visible point data for all the tris that will be deleted
+			for (int32 TriIdx : ToDelete)
+			{
+				ClearVisiblePointsData(TriIdx, NewlyUnclaimed);
+			}
+		}
 
 		// Connect up all the horizon triangles
 		int32 NewTriStart = Triangles.Num();
@@ -593,6 +669,8 @@ struct FHullConnectivity
 		DeleteTriangles(Triangles, ToDelete);
 
 		//ValidateConnectivity(Triangles);
+
+		return true;
 	}
 };
 
@@ -635,6 +713,8 @@ bool TConvexHull3<RealType>::Solve(int32 NumPoints, TFunctionRef<TVector<RealTyp
 
 	NumHullPoints = 4;
 
+	RealType DegenerateEdgeToleranceSq = DegenerateEdgeTolerance * DegenerateEdgeTolerance;
+
 	FHullConnectivity<RealType> Connectivity;
 	Connectivity.BuildNeighbors(Hull);
 	Connectivity.InitVisibility(Hull, NumPoints, GetPointFunc, FilterFunc);
@@ -650,7 +730,11 @@ bool TConvexHull3<RealType>::Solve(int32 NumPoints, TFunctionRef<TVector<RealTyp
 			break;
 		}
 		NumHullPoints++;
-		Connectivity.UpdateHullWithNewPoint(Hull, GetPointFunc, Visible.A, Visible.B);
+		bool bAdded = Connectivity.UpdateHullWithNewPoint(Hull, GetPointFunc, Visible.A, Visible.B, DegenerateEdgeToleranceSq);
+		if (!bAdded)
+		{
+			Connectivity.RemoveVisiblePoint(Visible, Hull, GetPointFunc);
+		}
 	}
 
 	if (bSaveTriangleNeighbors)
@@ -753,6 +837,410 @@ void TConvexHull3<RealType>::GetFaces(TFunctionRef<void(TArray<int32>&, TVector<
 	}
 }
 
+template<class RealType>
+void TConvexHull3<RealType>::GetSimplifiedFaces(TFunctionRef<void(TArray<int32>&, TVector<RealType>)> PolygonFunc, TFunctionRef<TVector<RealType>(int32)> GetPointFunc,
+	RealType FaceAngleTolerance, RealType PlaneDistanceTolerance) const
+{
+	if (!ensureMsgf(bSaveTriangleNeighbors && HullNeighbors.Num() == Hull.Num(), TEXT("To extract faces, set bSaveTriangleNeighbors = true before calling Solve()")))
+	{
+		return;
+	}
+
+	TArray<FPolygonFace> Polygons;
+	TArray<TVector<RealType>> Normals;
+	GetSimplifiedFaces(Polygons, GetPointFunc, FaceAngleTolerance, PlaneDistanceTolerance, &Normals);
+	TArray<int32> PolygonVertices;
+	for (int32 Idx = 0; Idx < Polygons.Num(); ++Idx)
+	{
+		const FPolygonFace& Face = Polygons[Idx];
+		PolygonVertices.Reset(Face.Num());
+		PolygonVertices.Append(Face);
+		PolygonFunc(PolygonVertices, Normals[Idx]);
+	}
+}
+
+template<class RealType>
+void TConvexHull3<RealType>::GetSimplifiedFaces(TArray<FPolygonFace>& OutPolygons, TFunctionRef<TVector<RealType>(int32)> GetPointFunc,
+	RealType FaceAngleToleranceInDegrees, RealType PlaneDistanceTolerance, TArray<TVector<RealType>>* OutPolygonNormals) const
+{
+	if (!ensureMsgf(bSaveTriangleNeighbors && HullNeighbors.Num() == Hull.Num(), TEXT("To extract faces, set bSaveTriangleNeighbors = true before calling Solve()")))
+	{
+		return;
+	}
+
+	OutPolygons.Reset();
+	if (OutPolygonNormals)
+	{
+		OutPolygonNormals->Reset();
+	}
+
+	FaceAngleToleranceInDegrees = FMath::Min(FaceAngleToleranceInDegrees, 60); // Angle Tolerance should be a small value; if a very large angle tolerance is allowed, it could create a degenerate (flat or zero vertex) hull
+	const double FaceMergeThreshold = 1 - FMath::Cos(FMath::DegreesToRadians(FaceAngleToleranceInDegrees));
+
+	int32 NumTris = Hull.Num();
+	FSizedDisjointSet PlaneGroups;
+	PlaneGroups.Init(NumTris);
+	TArray<FVector3d> PlaneNormals; PlaneNormals.SetNumUninitialized(NumTris);
+	TArray<double> PlaneAreas; PlaneAreas.SetNumUninitialized(NumTris);
+	TArray<FVector3d> PlaneOrigins; PlaneOrigins.SetNumUninitialized(NumTris);
+	// Map of plane group ID -> vertex indices, only for faces that have been merged at least once (otherwise, we can use the triangle indices)
+	using FPlaneVertArray = TArray<int32, TInlineAllocator<16>>;
+	TMap<int32, FPlaneVertArray> PlaneGroupToVerticesMap;
+	int32 MaxVertexNum = 0;
+	for (int32 TriIdx = 0; TriIdx < NumTris; ++TriIdx)
+	{
+		const FIndex3i& Tri = Hull[TriIdx];
+		MaxVertexNum = FMath::Max(MaxVertexNum, 1 + FMath::Max(Tri.A, FMath::Max(Tri.B, Tri.C)));
+		FVector3d TriPosns[3]{ (FVector3d)GetPointFunc(Tri.A), (FVector3d)GetPointFunc(Tri.B), (FVector3d)GetPointFunc(Tri.C) };
+		double Area = 0;
+		FVector3d FaceNormal = VectorUtil::NormalArea(TriPosns[0], TriPosns[1], TriPosns[2], Area);
+		PlaneNormals[TriIdx] = FaceNormal;
+		PlaneAreas[TriIdx] = Area;
+		PlaneOrigins[TriIdx] = (TriPosns[0] + TriPosns[1] + TriPosns[2]) * (1.0/3.0);
+	}
+
+	struct FEdgeWeight
+	{
+		FEdgeWeight() = default;
+		FEdgeWeight(FIndex2i TriPair, FIndex2i VertPair, double Weight) : TriPair(TriPair), VertPair(VertPair), Weight(Weight) {}
+		FIndex2i TriPair;
+		FIndex2i VertPair;
+		double Weight;
+	};
+
+	constexpr double AreaThreshold = UE_DOUBLE_SMALL_NUMBER;
+	auto GetMergeWeight = [&PlaneGroups, &PlaneAreas, &PlaneNormals, &GetPointFunc, AreaThreshold](FIndex2i PlanePair, FIndex2i VertexPair, bool bHasGroups) -> double
+	{
+		if (bHasGroups)
+		{
+			PlanePair[0] = PlaneGroups.Find(PlanePair[0]);
+			PlanePair[1] = PlaneGroups.Find(PlanePair[1]);
+			if (PlanePair[0] == PlanePair[1])
+			{
+				// return a value higher than the max possible angle threshold if the plane groups are already merged
+				constexpr float CannotMergeValue = 4;
+				return CannotMergeValue;
+			}
+		}
+		// Planes with small area can be merged into any neighbor, because we don't trust their plane normal
+		int32 TooSmallAreas = int32(PlaneAreas[PlanePair[0]] < AreaThreshold) + int32(PlaneAreas[PlanePair[1]] < AreaThreshold);
+		if (TooSmallAreas == 1)
+		{
+			// favor merging too-small-area faces to the neighboring valid-area face with the longest shared edge
+			FVector3d V0 = (FVector3d)GetPointFunc(VertexPair.A);
+			FVector3d V1 = (FVector3d)GetPointFunc(VertexPair.B);
+			double EdgeLenSq = FVector3d::DistSquared(V0, V1);
+			return -EdgeLenSq;
+		}
+		else if (TooSmallAreas == 2)
+		{
+			// if both areas are too small, still allow a low-cost merge, but favor the above single-small-area merges
+			return 0;
+		}
+		double NormalAlignment = 1 - PlaneNormals[PlanePair[0]].Dot(PlaneNormals[PlanePair[1]]);
+		return NormalAlignment;
+	};
+
+	TArray<FEdgeWeight> EdgeWeights;
+	for (int32 TriIdx = 0; TriIdx < NumTris; ++TriIdx)
+	{
+		const FIndex3i& NbrInds = HullNeighbors[TriIdx];
+		const FIndex3i& Tri = Hull[TriIdx];
+		for (int32 PrevIdx = 2, SubIdx = 0; SubIdx < 3; PrevIdx = SubIdx++)
+		{
+			if (TriIdx < NbrInds[PrevIdx])
+			{
+				FIndex2i FacePair(TriIdx, NbrInds[PrevIdx]);
+				FIndex2i VertPair(PrevIdx, SubIdx);
+				double Weight = GetMergeWeight(FacePair, VertPair, false);
+				if (Weight < FaceMergeThreshold)
+				{
+					EdgeWeights.Emplace(FacePair, VertPair, Weight);
+				}
+			}
+		}
+	}
+	
+	// Consider merging across edges with lower edge weights first
+	EdgeWeights.Sort([&](const FEdgeWeight& A, const FEdgeWeight& B)
+		{
+			return A.Weight < B.Weight;
+		});
+	for (const FEdgeWeight& EdgeWeight : EdgeWeights)
+	{
+		// re-evaluate the edge weight with the current planes
+		double MergeWeight = GetMergeWeight(EdgeWeight.TriPair, EdgeWeight.VertPair, true);
+		if (MergeWeight < FaceMergeThreshold)
+		{
+			int32 Groups[2]{ PlaneGroups.Find(EdgeWeight.TriPair.A), PlaneGroups.Find(EdgeWeight.TriPair.B) };
+			int32 GroupSizes[2]{ PlaneGroups.GetSize(Groups[0]), PlaneGroups.GetSize(Groups[1]) };
+			double Areas[2]{ PlaneAreas[Groups[0]], PlaneAreas[Groups[1]] };
+			double AreaSum = Areas[0] + Areas[1];
+			FVector3d Centroid;
+			if (AreaSum > UE_DOUBLE_KINDA_SMALL_NUMBER)
+			{
+				Centroid = (PlaneOrigins[Groups[0]] * (Areas[0]/AreaSum) + PlaneOrigins[Groups[1]] * (Areas[1]/AreaSum));
+			}
+			else
+			{
+				Centroid = (PlaneOrigins[Groups[0]] + PlaneOrigins[Groups[1]]) * .5;
+			}
+
+			FVector3d Normal = PlaneNormals[Groups[0]] * Areas[0] + PlaneNormals[Groups[1]] * Areas[1];
+			Normal.Normalize();
+
+			// Test that the points on each plane are close enough to the merged plane
+			bool bPointsCloseEnoughToPlane = true;
+
+			auto TestPoint = [&bPointsCloseEnoughToPlane, &GetPointFunc, Centroid, Normal, PlaneDistanceTolerance](int32 PointIdx) -> bool
+			{
+				FVector3d PointPos = (FVector3d)GetPointFunc(PointIdx);
+				if (FMath::Abs((PointPos - Centroid).Dot(Normal)) > PlaneDistanceTolerance)
+				{
+					bPointsCloseEnoughToPlane = false;
+					return false;
+				}
+				return true;
+			};
+
+			for (int32 GroupSubIdx = 0; GroupSubIdx < 2 && bPointsCloseEnoughToPlane; ++GroupSubIdx)
+			{
+				int32 GroupIdx = Groups[GroupSubIdx];
+				if (GroupSizes[GroupSubIdx] == 1)
+				{
+					FIndex3i Tri = Hull[GroupIdx];
+					for (int32 Idx = 0; Idx < 3; ++Idx)
+					{
+						if (!TestPoint(Tri[Idx]))
+						{
+							break;
+						}
+					}
+				}
+				else
+				{
+					for (int32 PointIdx : PlaneGroupToVerticesMap[GroupIdx])
+					{
+						if (!TestPoint(PointIdx))
+						{
+							break;
+						}
+					}
+				}
+			}
+			if (!bPointsCloseEnoughToPlane)
+			{
+				continue;
+			}
+
+			PlaneGroups.Union(Groups[0], Groups[1]);
+			int32 NewParent = PlaneGroups.Find(Groups[0]);
+			int32 OldGroupIdx = Groups[0] == NewParent ? 1 : 0;
+			int32 OldGroup = Groups[OldGroupIdx];
+			PlaneNormals[NewParent] = Normal;
+			PlaneOrigins[NewParent] = Centroid;
+			PlaneAreas[NewParent] = Areas[0] + Areas[1];
+
+			// Update the PlaneGroupToVerticesMap structures, removing the old group if needed, creating the new group if needed, and adding the appropriate vertices
+			if (GroupSizes[1 - OldGroupIdx] <= 1)
+			{
+				FIndex3i Tri = Hull[Groups[1 - OldGroupIdx]];
+				FPlaneVertArray& NewGroupVerts = PlaneGroupToVerticesMap.Emplace(NewParent);
+				NewGroupVerts.Add(Tri.A);
+				NewGroupVerts.Add(Tri.B);
+				NewGroupVerts.Add(Tri.C);
+			}
+			if (GroupSizes[OldGroupIdx] > 1)
+			{
+				const FPlaneVertArray OldGroupVerts = PlaneGroupToVerticesMap.FindAndRemoveChecked(OldGroup);
+				FPlaneVertArray& NewGroupVerts = PlaneGroupToVerticesMap[NewParent];
+				NewGroupVerts.Reserve(NewGroupVerts.Num() + OldGroupVerts.Num() - 2);
+				for (int32 VIdx : OldGroupVerts)
+				{
+					if (!EdgeWeight.VertPair.Contains(VIdx))
+					{
+						NewGroupVerts.Add(VIdx);
+					}
+				}
+			}
+			else
+			{
+				checkSlow(!PlaneGroupToVerticesMap.Contains(OldGroup));
+				FPlaneVertArray& NewGroupVerts = PlaneGroupToVerticesMap[NewParent];
+				FIndex3i Tri = Hull[Groups[OldGroupIdx]];
+				for (int32 SubIdx = 0; SubIdx < 3; ++SubIdx)
+				{
+					int32 VIdx = Tri[SubIdx];
+					if (!EdgeWeight.VertPair.Contains(VIdx))
+					{
+						NewGroupVerts.Add(VIdx);
+					}
+				}
+			}
+		}
+	}
+
+	// Track which vertices are attached to at least 3 groups
+	TArray<FIndex3i> VertPlaneGroups;
+	VertPlaneGroups.Init(FIndex3i::Invalid(), MaxVertexNum);
+	for (int32 TriIdx = 0; TriIdx < Hull.Num(); ++TriIdx)
+	{
+		int32 Group = PlaneGroups.Find(TriIdx);
+		FIndex3i Tri = Hull[TriIdx];
+		for (int32 SubIdx = 0; SubIdx < 3; ++SubIdx)
+		{
+			int32 VIdx = Tri[SubIdx];
+			for (int32 GroupNbrIdx = 0; GroupNbrIdx < 3; ++GroupNbrIdx)
+			{
+				if (VertPlaneGroups[VIdx][GroupNbrIdx] == -1)
+				{
+					VertPlaneGroups[VIdx][GroupNbrIdx] = Group;
+					break;
+				}
+				else if (VertPlaneGroups[VIdx][GroupNbrIdx] == Group)
+				{
+					break;
+				}
+			}
+		}
+	}
+
+	auto VertexTouchesAtLeastThreeGroups = [&VertPlaneGroups](int32 VertIdx)
+	{
+		return VertPlaneGroups[VertIdx][2] != -1;
+	};
+
+	bool bFacesAreConvex = true;
+
+	TArray<int32> CurFaceVertIDs;
+	for (int32 TriIdx = 0; TriIdx < Hull.Num(); ++TriIdx)
+	{
+		int32 GroupIdx = PlaneGroups.Find(TriIdx);
+		if (GroupIdx != TriIdx)
+		{
+			continue;
+		}
+		int32 GroupSize = PlaneGroups.GetSize(TriIdx);
+		if (GroupSize == 1)
+		{
+			FIndex3i Tri = Hull[TriIdx];
+			if (VertexTouchesAtLeastThreeGroups(Tri.A) &&
+				VertexTouchesAtLeastThreeGroups(Tri.B) &&
+				VertexTouchesAtLeastThreeGroups(Tri.C))
+			{
+				FPolygonFace& Face = OutPolygons.Emplace_GetRef();
+				Face.Add(Tri.A);
+				Face.Add(Tri.B);
+				Face.Add(Tri.C);
+				if (OutPolygonNormals)
+				{
+					OutPolygonNormals->Add((TVector<RealType>)PlaneNormals[GroupIdx]);
+				}
+			}
+
+			continue;
+		}
+
+		CurFaceVertIDs.Reset();
+		WalkBorder(Hull, HullNeighbors, [&](int32 Idx) { return PlaneGroups.Find(Idx) == GroupIdx; }, TriIdx, CurFaceVertIDs);
+		int32 NumKept = Algo::StableRemoveIf(CurFaceVertIDs, [&](int32 VertIdx) { return !VertexTouchesAtLeastThreeGroups(VertIdx); });
+
+		if (NumKept < 3)
+		{
+			continue;
+		}
+
+		CurFaceVertIDs.SetNum(NumKept, false);
+		
+		// Validate convex-enough winding vs the plane normal, and trigger a fallback if this fails
+		FVector3d PlaneNormal = PlaneNormals[GroupIdx];
+		{
+			// Iterate over each polygon face corner ABC, testing direction the AB edge turns vs the BC edge (from the PoV of the plane normal)
+			int32 VertIdxA = NumKept - 2, VertIdxB = NumKept - 1, VertIdxC = 0;
+			FVector3d PtA = (FVector3d)GetPointFunc(CurFaceVertIDs[VertIdxA]);
+			FVector3d PtB = (FVector3d)GetPointFunc(CurFaceVertIDs[VertIdxB]);
+			FVector3d EdgeAB = PtB - PtA;
+			bool bABIsNormalized = EdgeAB.Normalize();
+			for (; VertIdxC < NumKept && bFacesAreConvex; VertIdxA = VertIdxB, VertIdxB = VertIdxC++)
+			{
+				FVector3d PtC = (FVector3d)GetPointFunc(CurFaceVertIDs[VertIdxC]);
+				FVector3d EdgeBC = PtC - PtB;
+				bool bBCIsNormalized = EdgeBC.Normalize();
+				if (bABIsNormalized && bBCIsNormalized)
+				{
+					FVector3d EdgeABPerp = EdgeAB.Cross(PlaneNormal);
+					if (EdgeBC.Dot(EdgeABPerp) < -UE_DOUBLE_KINDA_SMALL_NUMBER)
+					{
+						bFacesAreConvex = false;
+						break;
+					}
+				}
+
+				PtA = PtB;
+				PtB = PtC;
+				EdgeAB = EdgeBC;
+				bABIsNormalized = bBCIsNormalized;
+			}
+		}
+
+		if (!bFacesAreConvex)
+		{
+			break;
+		}
+
+		FPolygonFace& Face = OutPolygons.Emplace_GetRef();
+		Face.Append(CurFaceVertIDs);
+		if (OutPolygonNormals)
+		{
+			OutPolygonNormals->Add((TVector<RealType>)PlaneNormal);
+		}
+	}
+
+	// If the algorithm has failed to create convex faces, or failed to create at least a tetrahedron, e.g. due to unfortunate face merges, fall back to a more exact hull
+	// Note: This should be a rare case, and it still uses the reduced vertex set of the 'vertex touches at least 3 groups' criteria,
+	// so should still give a simpler hull than calling GetFaces() directly would have.
+	if (!bFacesAreConvex || OutPolygons.Num() < 4)
+	{
+		OutPolygons.Reset();
+		if (OutPolygonNormals)
+		{
+			OutPolygonNormals->Reset();
+		}
+		TConvexHull3<RealType> FallbackHull;
+		FallbackHull.bSaveTriangleNeighbors = true;
+		FallbackHull.DegenerateEdgeTolerance = DegenerateEdgeTolerance;
+
+		bool bFallbackSolveSuccess = FallbackHull.Solve(MaxVertexNum, GetPointFunc, VertexTouchesAtLeastThreeGroups);
+		if (bFallbackSolveSuccess)
+		{
+			FallbackHull.GetFaces([&](TArray<int32>& FaceIndices, TVector<RealType> Normal)
+			{
+				FPolygonFace& Face = OutPolygons.Emplace_GetRef();
+				Face.Append(FaceIndices);
+				if (OutPolygonNormals)
+				{
+					OutPolygonNormals->Add(Normal);
+				}
+			}, GetPointFunc);
+		}
+		else
+		{
+			// If we failed to solve for a new hull using just the group-corner vertices, then use the initial hull faces without simplification
+			// This could happen if the faces of the convex hull were all too small in area, so all faces were merged, leaving no vertices for the FallbackHull to find
+			GetFaces([&](TArray<int32>& FaceIndices, TVector<RealType> Normal)
+			{
+				FPolygonFace& Face = OutPolygons.Emplace_GetRef();
+				Face.Append(FaceIndices);
+				if (OutPolygonNormals)
+				{
+					OutPolygonNormals->Add(Normal);
+				}
+			}, GetPointFunc);
+		}
+
+	}
+}
+
 template<typename RealType>
 void TConvexHull3<RealType>::WalkBorder(const TArray<FIndex3i>& Triangles, const TArray<FIndex3i>& TriangleNeighbors, TFunctionRef<bool(int32)> InGroupFunc, int32 StartIdx, TArray<int32>& OutBorderVertexIndices)
 {
diff --git a/Engine/Source/Runtime/GeometryCore/Public/CompGeom/ConvexHull3.h b/Engine/Source/Runtime/GeometryCore/Public/CompGeom/ConvexHull3.h
--- a/Engine/Source/Runtime/GeometryCore/Public/CompGeom/ConvexHull3.h
+++ b/Engine/Source/Runtime/GeometryCore/Public/CompGeom/ConvexHull3.h
@@ -64,6 +64,10 @@ public:
 	/// If true, can call GetTriangleNeighbors() after Solve().
 	bool bSaveTriangleNeighbors = false;
 
+	/// Points will not be added to the hull if doing so would create an edge smaller than this distance tolerance.
+	/// If zero, no points will be skipped.
+	RealType DegenerateEdgeTolerance = (RealType)0;
+
 	/**
 	 * Generate convex hull as long as input is not degenerate
 	 * If input is degenerate, this will return false, and caller can call GetDimension()
@@ -136,13 +140,37 @@ public:
 	}
 
 	/**
-	 * Get faces of the convex hull as convex polygons
+	 * Get faces of the convex hull as convex polygons, merging faces that are exactly coplanar
 	 * @param PolygonFunc	Callback to be called for each polygon, with the array of vertex indices and the face normal
 	 * @param GetPointFunc	Function providing array-style access into points
-	 * TODO: Provide an optional Epsilon parameter to simplify the hull within some tolerance
 	 */
 	GEOMETRYCORE_API void GetFaces(TFunctionRef<void(TArray<int32>&, TVector<RealType>)> PolygonFunc, TFunctionRef<TVector<RealType>(int32)> GetPointFunc) const;
 
+	/**
+	 * Get faces of the convex hull as convex polygons, simplifying the hull by merging near-coplanar faces and only keeping vertices that are on the corner of at least three merged faces
+	 * 
+	 * @param PolygonFunc					Callback to be called for each polygon, with the array of vertex indices and the face normal
+	 * @param GetPointFunc					Function providing array-style access into points
+	 * @param FaceAngleToleranceInDegrees	The hull will be simplified by merging faces with less than this dihedral angle between them
+	 * @param PlaneDistanceTolerance		Faces will not merge unless all points on the face are within this distance of the combined (average) face plane
+	 */
+	GEOMETRYCORE_API void GetSimplifiedFaces(TFunctionRef<void(TArray<int32>&, TVector<RealType>)> PolygonFunc, TFunctionRef<TVector<RealType>(int32)> GetPointFunc,
+		RealType FaceAngleToleranceInDegrees = (RealType)1.0, RealType PlaneDistanceTolerance = (RealType)1.0) const;
+
+	// Polygon face array type: A nested array with inline allocator per face, optimizing for the case where most faces have less than 8 vertices
+	using FPolygonFace = TArray<int32, TInlineAllocator<8>>;
+	/**
+	 * Get faces of the convex hull as convex polygons, simplifying the hull by merging near-coplanar faces and only keeping vertices that are on the corner of at least three merged faces
+	 * 
+	 * @param OutPolygons					Polygons of the convex hull faces, as arrays of indices into the original points
+	 * @param GetPointFunc					Function providing array-style access into points
+	 * @param FaceAngleToleranceInDegrees	The hull will be simplified by merging faces with less than this dihedral angle between them
+	 * @param PlaneDistanceTolerance		Faces will not merge unless all points on the face are within this distance of the combined (average) face plane
+	 * @param OutPolygonNormals				Optional array of normals for each polygon
+	 */
+	GEOMETRYCORE_API void GetSimplifiedFaces(TArray<FPolygonFace>& OutPolygons, TFunctionRef<TVector<RealType>(int32)> GetPointFunc,
+		RealType FaceAngleToleranceInDegrees = (RealType)1.0, RealType PlaneDistanceTolerance = (RealType)1.0, TArray<TVector<RealType>>* OutPolygonNormals = nullptr) const;
+
 	/**
 	 * Only valid if bSaveTriangleNeighbors was true when Solve() was called
 	 * @return Neighbors of each hull triangle, in edge order -- i.e., Nbr.A is the triangle across edge(Tri.A, Tri.B)
