Refine ConvexHull3's simplification algorithm to remove vertices that looked like corners vertices only due to face groups that were deleted

diff --git a/Engine/Source/Runtime/GeometryCore/Private/CompGeom/ConvexHull3.cpp b/Engine/Source/Runtime/GeometryCore/Private/CompGeom/ConvexHull3.cpp
--- a/Engine/Source/Runtime/GeometryCore/Private/CompGeom/ConvexHull3.cpp
+++ b/Engine/Source/Runtime/GeometryCore/Private/CompGeom/ConvexHull3.cpp
@@ -1110,9 +1110,10 @@ void TConvexHull3<RealType>::GetSimplifiedFaces(TArray<FPolygonFace>& OutPolygon
 		return VertPlaneGroups[VertIdx][2] != -1;
 	};
 
-	bool bFacesAreConvex = true;
-
 	TArray<int32> CurFaceVertIDs;
+	bool bHasDeletedFaces = false;
+	// Track the mapping from polygons back to groups, only if not using OutPolygonNormals, to recover the polygon normals for convexity tests
+	TArray<int32> PolygonToGroup;
 	for (int32 TriIdx = 0; TriIdx < Hull.Num(); ++TriIdx)
 	{
 		int32 GroupIdx = PlaneGroups.Find(TriIdx);
@@ -1137,6 +1138,10 @@ void TConvexHull3<RealType>::GetSimplifiedFaces(TArray<FPolygonFace>& OutPolygon
 					OutPolygonNormals->Add((TVector<RealType>)PlaneNormals[GroupIdx]);
 				}
 			}
+			else
+			{
+				bHasDeletedFaces = true;
+			}
 
 			continue;
 		}
@@ -1147,55 +1152,128 @@ void TConvexHull3<RealType>::GetSimplifiedFaces(TArray<FPolygonFace>& OutPolygon
 
 		if (NumKept < 3)
 		{
+			bHasDeletedFaces = true;
 			continue;
 		}
 
 		CurFaceVertIDs.SetNum(NumKept, false);
-		
-		// Validate convex-enough winding vs the plane normal, and trigger a fallback if this fails
-		FVector3d PlaneNormal = PlaneNormals[GroupIdx];
+
+		FPolygonFace& Face = OutPolygons.Emplace_GetRef();
+		Face.Append(CurFaceVertIDs);
+		if (OutPolygonNormals)
+		{
+			OutPolygonNormals->Add((TVector<RealType>)PlaneNormals[GroupIdx]);
+		}
+		else
+		{
+			PolygonToGroup.Add(GroupIdx);
+		}
+	}
+
+	// If we deleted faces in the initial pass, keep doing passes to identify vertices that touch fewer than 3 polygons & deleting them + deleting polygons w/ < 3 remaining vertices,
+	// until we either stop deleting polygons or we have too few polygons to form a solid volume
+	while (bHasDeletedFaces && OutPolygons.Num() > 3)
+	{
+		bHasDeletedFaces = false; // reset for the next pass
+
+		// recompute vert plane groups using the output polygons
+		VertPlaneGroups.Init(FIndex3i::Invalid(), MaxVertexNum);
+		for (int32 PolyIdx = 0; PolyIdx < OutPolygons.Num(); ++PolyIdx)
 		{
-			// Iterate over each polygon face corner ABC, testing direction the AB edge turns vs the BC edge (from the PoV of the plane normal)
-			int32 VertIdxA = NumKept - 2, VertIdxB = NumKept - 1, VertIdxC = 0;
-			FVector3d PtA = (FVector3d)GetPointFunc(CurFaceVertIDs[VertIdxA]);
-			FVector3d PtB = (FVector3d)GetPointFunc(CurFaceVertIDs[VertIdxB]);
-			FVector3d EdgeAB = PtB - PtA;
-			bool bABIsNormalized = EdgeAB.Normalize();
-			for (; VertIdxC < NumKept && bFacesAreConvex; VertIdxA = VertIdxB, VertIdxB = VertIdxC++)
+			FPolygonFace& Face = OutPolygons[PolyIdx];
+			for (int32 SubIdx = 0; SubIdx < Face.Num(); ++SubIdx)
 			{
-				FVector3d PtC = (FVector3d)GetPointFunc(CurFaceVertIDs[VertIdxC]);
-				FVector3d EdgeBC = PtC - PtB;
-				bool bBCIsNormalized = EdgeBC.Normalize();
-				if (bABIsNormalized && bBCIsNormalized)
+				int32 VIdx = Face[SubIdx];
+				for (int32 GroupNbrIdx = 0; GroupNbrIdx < 3; ++GroupNbrIdx)
 				{
-					FVector3d EdgeABPerp = EdgeAB.Cross(PlaneNormal);
-					if (EdgeBC.Dot(EdgeABPerp) < -UE_DOUBLE_KINDA_SMALL_NUMBER)
+					if (VertPlaneGroups[VIdx][GroupNbrIdx] == -1)
+					{
+						VertPlaneGroups[VIdx][GroupNbrIdx] = PolyIdx;
+						break;
+					}
+					else if (VertPlaneGroups[VIdx][GroupNbrIdx] == PolyIdx)
 					{
-						bFacesAreConvex = false;
 						break;
 					}
 				}
-
-				PtA = PtB;
-				PtB = PtC;
-				EdgeAB = EdgeBC;
-				bABIsNormalized = bBCIsNormalized;
 			}
 		}
 
-		if (!bFacesAreConvex)
+
+		for (int32 PolyIdx = 0; PolyIdx < OutPolygons.Num(); ++PolyIdx)
 		{
-			break;
+			FPolygonFace& Face = OutPolygons[PolyIdx];
+			int32 NumKept = Algo::StableRemoveIf(Face, [&](int32 VertIdx) { return !VertexTouchesAtLeastThreeGroups(VertIdx); });
+			if (NumKept != Face.Num())
+			{
+				if (NumKept > 2)
+				{
+					Face.SetNum(NumKept);
+				}
+				else
+				{
+					OutPolygons.RemoveAtSwap(PolyIdx, 1, false);
+					if (OutPolygonNormals)
+					{
+						OutPolygonNormals->RemoveAtSwap(PolyIdx, 1, false);
+					}
+					else
+					{
+						PolygonToGroup.RemoveAtSwap(PolyIdx, 1, false);
+					}
+					bHasDeletedFaces = true;
+					PolyIdx--;
+				}
+			}
 		}
+	}
 
-		FPolygonFace& Face = OutPolygons.Emplace_GetRef();
-		Face.Append(CurFaceVertIDs);
+	// Validate convex-enough winding vs the plane normal, and trigger a fallback if this fails
+	bool bFacesAreConvex = true;
+	for (int32 PolyIdx = 0; PolyIdx < OutPolygons.Num() && bFacesAreConvex; ++PolyIdx)
+	{
+
+		FVector3d PlaneNormal;
 		if (OutPolygonNormals)
 		{
-			OutPolygonNormals->Add((TVector<RealType>)PlaneNormal);
+			PlaneNormal = (FVector3d)(*OutPolygonNormals)[PolyIdx];
+		}
+		else
+		{
+			PlaneNormal = PlaneNormals[PolygonToGroup[PolyIdx]];
+		}
+
+		// Iterate over each polygon face corner ABC, testing direction the AB edge turns vs the BC edge (from the PoV of the plane normal)
+		const FPolygonFace& Polygon = OutPolygons[PolyIdx];
+		int32 PolygonVertNum = Polygon.Num();
+		int32 VertIdxA = PolygonVertNum - 2, VertIdxB = PolygonVertNum - 1, VertIdxC = 0;
+		FVector3d PtA = (FVector3d)GetPointFunc(Polygon[VertIdxA]);
+		FVector3d PtB = (FVector3d)GetPointFunc(Polygon[VertIdxB]);
+		FVector3d EdgeAB = PtB - PtA;
+		bool bABIsNormalized = EdgeAB.Normalize();
+		for (; VertIdxC < PolygonVertNum && bFacesAreConvex; VertIdxA = VertIdxB, VertIdxB = VertIdxC++)
+		{
+			FVector3d PtC = (FVector3d)GetPointFunc(Polygon[VertIdxC]);
+			FVector3d EdgeBC = PtC - PtB;
+			bool bBCIsNormalized = EdgeBC.Normalize();
+			if (bABIsNormalized && bBCIsNormalized)
+			{
+				FVector3d EdgeABPerp = EdgeAB.Cross(PlaneNormal);
+				if (EdgeBC.Dot(EdgeABPerp) < -UE_DOUBLE_KINDA_SMALL_NUMBER)
+				{
+					bFacesAreConvex = false;
+					break;
+				}
+			}
+
+			PtA = PtB;
+			PtB = PtC;
+			EdgeAB = EdgeBC;
+			bABIsNormalized = bBCIsNormalized;
 		}
 	}
 
+
 	// If the algorithm has failed to create convex faces, or failed to create at least a tetrahedron, e.g. due to unfortunate face merges, fall back to a more exact hull
 	// Note: This should be a rare case, and it still uses the reduced vertex set of the 'vertex touches at least 3 groups' criteria,
 	// so should still give a simpler hull than calling GetFaces() directly would have.
