Mesh Modeling Tools,  For the "normal aware" simplification method in the modeling mode simplification tool.  Makes an internal computation more robust for the specific case when triangles being collapsed are on a plane that intersects the origin.
Specifically the computations of attribute gradients now uses a local coordinate system defined in terms of the triangle faces.

diff --git a/Engine/Source/Runtime/GeometryCore/Public/QuadricError.h b/Engine/Source/Runtime/GeometryCore/Public/QuadricError.h
--- a/Engine/Source/Runtime/GeometryCore/Public/QuadricError.h
+++ b/Engine/Source/Runtime/GeometryCore/Public/QuadricError.h
@@ -433,7 +433,7 @@ public:
 		attrweight(AttrWeight)
 	{
 		/**
-		* Given scalar attribute 'a' defined at the vertices e.g. a0, a1, a2
+		* Given a scalar attribute 'a' defined at the vertices e.g. a0, a1, a2 (in this case the attribute is a single component of the vertex normal)
 		*
 		* solve 4x4 system:
 		*       (p0^t   1) (g[0])  = (a0)
@@ -441,58 +441,64 @@ public:
 		*       (p2^t   1) (g[2])    (a2)
 		*       (n^t    0) (  d )    (0 )
 		*
-		* for the interpolating gradient 'g'   -- note this is really p0.dot(g) + d = a0  and n.dot(g) = 0
+		* for the interpolating gradient 'g'   -- note this is really p.dot(g) + d = a0  and n.dot(g) = 0
 		*
-		* in practice we actually use a Schur complement approach that only requires inverting a 3x3 matrix.
+		* this system can be re-written as 
+		*       < e_20 | g > = (a2-a0)
+		*       < e_10 | g > = (a1-a0)
+		*       <   n  | g > = 0
+		* and   
+		*       < p0 | g > + d = a0.
+		* where e_20 is the edge vector p2 - p0
+		*       e_10 is the edge vector p1 - p0.
+		* 
+		* the first 3 equations can be solved for the vector 'g' either as a 3x3 matrix system 
+		* or one could solve a 2x2 system  since 'g' must be in the plane spanned by {e20, e10}  
+		* ( give 'n' is orthogonal to the triangle face).
 		*
 		* For this quadric, each component of the normal is treated as separate attribute.
 		*/
-		TMatrix3<RealType> Matrix(P0, P1, P2, true); // row-based matrix constructor.
-
-		bool bSolved = false;
-		double det = Matrix.Determinant();
-		if (FMath::Abs(det) > 1.e-5)
+		 
+		// the basis matrix is composed of two edges of the triangle and the face normal.
+		// unless the triangle is degenerate (or too small), this should be invertible. 
+		//@todo replace the 3x3 system with a 2x2 system that results from g = a1 e_10 + a2 e_20
+		// and solving for a1, a2
+		const TMatrix3<RealType> BasisMatrix(P2-P0, P1-P0, NFace, true); // row-based matrix constructor.
+
+		const double det = BasisMatrix.Determinant(); // geometrically the det is related to the triangle area squared
+		if (FMath::Abs(det) > 1.e-8)
 		{
-			FMatrix3d InvMatrix = Matrix.Inverse();
-			FVector3d NFaceT_InvMatrix;
-			for (int i = 0; i < 3; ++i)
-			{
-				NFaceT_InvMatrix[i] = NFace[0] * InvMatrix.Row0[i] + NFace[1] * InvMatrix.Row1[i] + NFace[2] * InvMatrix.Row2[i];
-			}
-
-			double NFaceDotInvMOne = NFaceT_InvMatrix.Dot(FVector3d(1, 1, 1));
+			// invert the matrix
+			const FMatrix3d CoBasisMatrix = BasisMatrix.Inverse();
 
-			if (FMath::Abs(NFaceDotInvMOne) > 1.e-5)
+			// for each component of the attribute vector, 
+			for (int i = 0; i < 3; ++i)
 			{
-				bSolved = true;
-				for (int i = 0; i < 3; ++i)
-				{
-					FVector3d AttrVec(N0[i], N1[i], N2[i]);
-					AttrVec *= AttrWeight;
-
-					d[i] = NFaceT_InvMatrix.Dot(AttrVec) / NFaceDotInvMOne;
-					grad[i] = InvMatrix * (AttrVec - FVector3d(d[i], d[i], d[i]));
+				const FVector3d ScaledAttr = AttrWeight * FVector3d(N0[i], N1[i], N2[i]);
+				const FVector3d DataVec(ScaledAttr[2] - ScaledAttr[0], ScaledAttr[1] - ScaledAttr[0], 0.);
+				
+				// compute vector 'g' and scalar 'd'
+				grad[i] = CoBasisMatrix * DataVec;
+				d[i] = ScaledAttr[0] - P0.Dot(grad[i]);
 
-					BaseStruct::Axx += grad[i].X * grad[i].X;
+				// add the values to the quadric 
+				BaseStruct::Axx += grad[i].X * grad[i].X;
 
-					BaseStruct::Axy += grad[i].X * grad[i].Y;
-					BaseStruct::Axz += grad[i].X * grad[i].Z;
+				BaseStruct::Axy += grad[i].X * grad[i].Y;
+				BaseStruct::Axz += grad[i].X * grad[i].Z;
 
-					BaseStruct::Ayy += grad[i].Y * grad[i].Y;
-					BaseStruct::Ayz += grad[i].Y * grad[i].Z;
-					BaseStruct::Azz += grad[i].Z * grad[i].Z;
+				BaseStruct::Ayy += grad[i].Y * grad[i].Y;
+				BaseStruct::Ayz += grad[i].Y * grad[i].Z;
+				BaseStruct::Azz += grad[i].Z * grad[i].Z;
 
-					BaseStruct::bx += d[i] * grad[i].X;
-					BaseStruct::by += d[i] * grad[i].Y;
-					BaseStruct::bz += d[i] * grad[i].Z;
+				BaseStruct::bx += d[i] * grad[i].X;
+				BaseStruct::by += d[i] * grad[i].Y;
+				BaseStruct::bz += d[i] * grad[i].Z;
 
-					BaseStruct::c += d[i] * d[i];
-				}
+				BaseStruct::c += d[i] * d[i];
 			}
-
 		}
-
-		if (!bSolved)
+		else // the triangle was too small or degenerate in some other way
 		{
 			for (int i = 0; i < 3; ++i)
 			{
