Add K-Means clustering to GeometryCore
Add initial PointSetFunctions to Geometry Script, including: K-Means Clustering, Flatten/Unflatten, Make Bounding Box, Transforms -> Points
Add a custom For Each-style macro to provide a more convenient interface for working with the clustering result

diff --git a/Engine/Plugins/Experimental/GeometryScripting/Source/GeometryScriptingCore/Private/PointSetFunctions.cpp b/Engine/Plugins/Experimental/GeometryScripting/Source/GeometryScriptingCore/Private/PointSetFunctions.cpp
new file mode 100644
--- /dev/null
+++ b/Engine/Plugins/Experimental/GeometryScripting/Source/GeometryScriptingCore/Private/PointSetFunctions.cpp
@@ -0,0 +1,151 @@
+// Copyright Epic Games, Inc. All Rights Reserved.
+
+#include "GeometryScript/PointSetFunctions.h"
+
+#include "Clustering/KMeans.h"
+
+#include UE_INLINE_GENERATED_CPP_BY_NAME(PointSetFunctions)
+
+using namespace UE::Geometry;
+
+#define LOCTEXT_NAMESPACE "UGeometryScriptLibrary_PointSetFunctions"
+
+namespace UE::Private::PointSetFunctionLocals
+{
+	static FClusterKMeans GetClustering(const TArray<FVector>& Points, const FGeometryScriptPointClusteringOptions& Options)
+	{
+		FClusterKMeans Clustering;
+		Clustering.RandomSeed = Options.RandomSeed;
+		Clustering.MaxIterations = Options.MaxIterations;
+		const TArray<FVector>* UseInitialCenters = &Options.InitialClusterCenters;
+		TArray<FVector> LocalInitialCenters;
+		if (Options.InitialClusterCenters.IsEmpty())
+		{
+			UseInitialCenters = &LocalInitialCenters;
+			if (Options.InitializeMethod == EGeometryScriptInitKMeansMethod::UniformSpacing)
+			{
+				Clustering.GetUniformSpacedInitialCenters<FVector>(Points, Options.TargetNumClusters, LocalInitialCenters);
+			}
+		}
+		Clustering.ComputeClusters<FVector>(Points, Options.TargetNumClusters, *UseInitialCenters);
+		return Clustering;
+	}
+
+	// Remap axes for 3D -> 2D conversion, dropping the chosen axis
+	inline static FIndex2i FlattenedAxesMap(EGeometryScriptAxis DropAxis)
+	{
+		return FIndex2i(
+			int32(DropAxis == EGeometryScriptAxis::X),
+			1 + int32(DropAxis != EGeometryScriptAxis::Z));
+	}
+}
+
+void UGeometryScriptLibrary_PointSetSamplingFunctions::KMeansClusterToIDs(
+	const TArray<FVector>& Points,
+	const FGeometryScriptPointClusteringOptions& Options,
+	TArray<int32>& PointClusterIndices)
+{
+	FClusterKMeans Clustering = UE::Private::PointSetFunctionLocals::GetClustering(Points, Options);
+	PointClusterIndices = MoveTemp(Clustering.ClusterIDs);
+}
+
+void UGeometryScriptLibrary_PointSetSamplingFunctions::KMeansClusterToArrays(
+	const TArray<FVector>& Points,
+	const FGeometryScriptPointClusteringOptions& Options,
+	TArray<FGeometryScriptIndexList>& ClusterIDToLists)
+{
+	FClusterKMeans Clustering = UE::Private::PointSetFunctionLocals::GetClustering(Points, Options);
+	ClusterIDToLists.SetNum(Clustering.ClusterSizes.Num());
+	for (int32 ClusterIdx = 0; ClusterIdx < ClusterIDToLists.Num(); ++ClusterIdx)
+	{
+		ClusterIDToLists[ClusterIdx].Reset(EGeometryScriptIndexType::Any);
+		ClusterIDToLists[ClusterIdx].List->Reserve(Clustering.ClusterSizes[ClusterIdx]);
+	}
+	if (ClusterIDToLists.IsEmpty())
+	{
+		return;
+	}
+
+	for (int32 PointIdx = 0; PointIdx < Clustering.ClusterIDs.Num(); ++PointIdx)
+	{
+		ClusterIDToLists[Clustering.ClusterIDs[PointIdx]].List->Add(PointIdx);
+	}
+}
+
+void UGeometryScriptLibrary_PointSetSamplingFunctions::TransformsToPoints(
+	const TArray<FTransform>& Transforms,
+	TArray<FVector>& Points
+)
+{
+	Points.SetNum(Transforms.Num());
+	for (int32 Idx = 0; Idx < Transforms.Num(); ++Idx)
+	{
+		Points[Idx] = Transforms[Idx].GetLocation();
+	}
+}
+
+void UGeometryScriptLibrary_PointSetSamplingFunctions::GetPointsFromIndexList(
+	const TArray<FVector>& AllPoints,
+	const FGeometryScriptIndexList& Indices,
+	TArray<FVector>& SelectedPoints
+)
+{
+	if (!Indices.List.IsValid())
+	{
+		SelectedPoints.Reset();
+		return;
+	}
+	SelectedPoints.Reset(Indices.List->Num());
+	for (int32 Idx : *Indices.List)
+	{
+		SelectedPoints.Add(AllPoints[Idx]);
+	}
+}
+
+void UGeometryScriptLibrary_PointSetSamplingFunctions::FlattenPoints(
+	const TArray<FVector>& PointsIn3D,
+	TArray<FVector2D>& PointsIn2D,
+	const FGeometryScriptPointFlatteningOptions& Options
+)
+{
+	PointsIn2D.SetNumUninitialized(PointsIn3D.Num());
+
+	FIndex2i AxesMap = UE::Private::PointSetFunctionLocals::FlattenedAxesMap(Options.DropAxis);
+	for (int32 Idx = 0; Idx < PointsIn3D.Num(); ++Idx)
+	{
+		FVector RelPoint = Options.Frame.InverseTransformPosition(PointsIn3D[Idx]);
+		PointsIn2D[Idx].X = RelPoint[AxesMap[0]];
+		PointsIn2D[Idx].Y = RelPoint[AxesMap[1]];
+	}
+}
+
+void UGeometryScriptLibrary_PointSetSamplingFunctions::UnflattenPoints(
+	const TArray<FVector2D>& PointsIn2D,
+	TArray<FVector>& PointsIn3D,
+	const FGeometryScriptPointFlatteningOptions& Options,
+	double Height
+)
+{
+	PointsIn3D.SetNumUninitialized(PointsIn2D.Num());
+
+	FIndex2i AxesMap = UE::Private::PointSetFunctionLocals::FlattenedAxesMap(Options.DropAxis);
+	for (int32 Idx = 0; Idx < PointsIn2D.Num(); ++Idx)
+	{
+		FVector RelPoint;
+		RelPoint[AxesMap[0]] = PointsIn2D[Idx].X;
+		RelPoint[AxesMap[1]] = PointsIn2D[Idx].Y;
+		RelPoint[(int32)Options.DropAxis] = Height;
+		PointsIn3D[Idx] = Options.Frame.TransformPosition(RelPoint);
+	}
+}
+
+
+FBox UGeometryScriptLibrary_PointSetSamplingFunctions::MakeBoundingBoxFromPoints(const TArray<FVector>& Points, double ExpandBy)
+{
+	FBox Bounds(Points);
+	return Bounds.ExpandBy(ExpandBy);
+}
+
+
+
+#undef LOCTEXT_NAMESPACE
\ No newline at end of file
diff --git a/Engine/Plugins/Experimental/GeometryScripting/Source/GeometryScriptingCore/Public/GeometryScript/PointSetFunctions.h b/Engine/Plugins/Experimental/GeometryScripting/Source/GeometryScriptingCore/Public/GeometryScript/PointSetFunctions.h
new file mode 100644
--- /dev/null
+++ b/Engine/Plugins/Experimental/GeometryScripting/Source/GeometryScriptingCore/Public/GeometryScript/PointSetFunctions.h
@@ -0,0 +1,136 @@
+// Copyright Epic Games, Inc. All Rights Reserved.
+
+#pragma once
+
+#include "CoreMinimal.h"
+#include "Kismet/BlueprintFunctionLibrary.h"
+#include "GeometryScript/GeometryScriptTypes.h"
+#include "PointSetFunctions.generated.h"
+
+
+UENUM(BlueprintType)
+enum class EGeometryScriptInitKMeansMethod : uint8
+{
+	Random,
+	UniformSpacing
+};
+
+USTRUCT(BlueprintType)
+struct GEOMETRYSCRIPTINGCORE_API FGeometryScriptPointClusteringOptions
+{
+	GENERATED_BODY()
+public:
+
+	/** If not empty, will be used instead of Target Num Clusters. Specifies the initial cluster centers to use. */
+	UPROPERTY(EditAnywhere, BlueprintReadWrite, Category = Options)
+	TArray<FVector> InitialClusterCenters;
+
+	/** Number of clusters requested, if Initial Cluster Centers is empty. Actual clusters generated may be smaller, e.g. if there are fewer points */
+	UPROPERTY(EditAnywhere, BlueprintReadWrite, Category = Options)
+	int TargetNumClusters = 3;
+
+	/** Method to initialize the cluster centers, if Initial Cluster Centers is empty */
+	UPROPERTY(EditAnywhere, BlueprintReadWrite, Category = Options)
+	EGeometryScriptInitKMeansMethod InitializeMethod = EGeometryScriptInitKMeansMethod::UniformSpacing;
+
+	/** Random Seed used to initialize clustering, if the Random cluster initialization is chosen */
+	UPROPERTY(EditAnywhere, BlueprintReadWrite, Category = Options, meta = (EditCondition="InitializeMethod == EGeometryScriptInitKMeansMethod::Random"))
+	int32 RandomSeed = 0;
+
+	/** Maximum iterations to run the clustering process. Will stop earlier if/when clustering converges. */
+	UPROPERTY(EditAnywhere, BlueprintReadWrite, Category = Options, meta = (ClampMin = "1"))
+	int32 MaxIterations = 500;
+};
+
+
+USTRUCT(BlueprintType)
+struct GEOMETRYSCRIPTINGCORE_API FGeometryScriptPointFlatteningOptions
+{
+	GENERATED_BODY()
+public:
+
+	/** Relative transform to use as a frame of reference. When flattening, the inverse transform will be applied to bring the points into the local space of the frame. */
+	UPROPERTY(EditAnywhere, BlueprintReadWrite, Category = Options)
+	FTransform Frame = FTransform::Identity;
+
+	/** Which axis to drop when flattening */
+	UPROPERTY(EditAnywhere, BlueprintReadWrite, Category = Options)
+	EGeometryScriptAxis DropAxis = EGeometryScriptAxis::Z;
+};
+
+
+
+UCLASS(meta = (ScriptName = "GeometryScript_PointSetSampling"))
+class GEOMETRYSCRIPTINGCORE_API UGeometryScriptLibrary_PointSetSamplingFunctions : public UBlueprintFunctionLibrary
+{
+	GENERATED_BODY()
+public:
+	
+	/**
+	 * Use K-Means clustering to cluster the given points into a target number of clusters,
+	 * and return an array with a cluster index per point.
+	 */
+	UFUNCTION(BlueprintCallable, Category = "GeometryScript|PointSet")
+	static void KMeansClusterToIDs(
+		const TArray<FVector>& Points,
+		const FGeometryScriptPointClusteringOptions& Options,
+		TArray<int32>& PointClusterIndices);
+
+	/**
+	 * Use K-Means clustering to cluster the given points into a target number of clusters,
+	 * and return the clusters as an array of lists of point indices.
+	 */
+	UFUNCTION(BlueprintCallable, Category = "GeometryScript|PointSet")
+	static void KMeansClusterToArrays(
+		const TArray<FVector>& Points,
+		const FGeometryScriptPointClusteringOptions& Options,
+		TArray<FGeometryScriptIndexList>& ClusterIDToLists);
+
+	/**
+	 * Create an array of the positions of the input Transforms
+	 */
+	UFUNCTION(BlueprintCallable, Category = "GeometryScript|PointSet")
+	static void TransformsToPoints(
+		const TArray<FTransform>& Transforms,
+		TArray<FVector>& Points
+	);
+
+	/**
+	 * Convert an array of points from 3D to 2D, by transforming into the given ReferenceFrame and taking the X,Y coordinates
+	 * Note that to transform into the ReferenceFrame, we apply the inverse of the ReferenceFrame's transform.
+	 */
+	UFUNCTION(BlueprintCallable, Category = "GeometryScript|PointSet")
+	static void FlattenPoints(
+		const TArray<FVector>& PointsIn3D,
+		TArray<FVector2D>& PointsIn2D,
+		const FGeometryScriptPointFlatteningOptions& Options
+	);
+
+	/**
+	 * Convert an array of points from 2D to 3D, by transforming out of the given ReferenceFrame, with the given Height for the non-flat axis (default Z).
+	 */
+	UFUNCTION(BlueprintCallable, Category = "GeometryScript|PointSet")
+	static void UnflattenPoints(
+		const TArray<FVector2D>& PointsIn2D,
+		TArray<FVector>& PointsIn3D,
+		const FGeometryScriptPointFlatteningOptions& Options,
+		double Height = 0.0
+	);
+
+	/**
+	 * Make a Axis Aligned Bounding Box that bounds the given Points, optionally expanded by some additional amount on each side
+	 */
+	UFUNCTION(BlueprintCallable, Category = "GeometryScript|PointSet")
+	static UPARAM(DisplayName = "Bounding Box") FBox MakeBoundingBoxFromPoints(const TArray<FVector>& Points, double ExpandBy = 0.0);
+
+	/**
+	 * Create an array of the subset of AllPoints indicated by the Indices list
+	 */
+	UFUNCTION(BlueprintCallable, Category = "GeometryScript|PointSet")
+	static void GetPointsFromIndexList(
+		const TArray<FVector>& AllPoints,
+		const FGeometryScriptIndexList& Indices,
+		TArray<FVector>& SelectedPoints
+	);
+
+};
\ No newline at end of file
diff --git a/Engine/Source/Runtime/GeometryCore/Private/Clustering/KMeans.cpp b/Engine/Source/Runtime/GeometryCore/Private/Clustering/KMeans.cpp
new file mode 100644
--- /dev/null
+++ b/Engine/Source/Runtime/GeometryCore/Private/Clustering/KMeans.cpp
@@ -0,0 +1,27 @@
+// Copyright Epic Games, Inc. All Rights Reserved.
+
+#include "Clustering/KMeans.h"
+
+
+namespace UE::Geometry
+{
+
+void FClusterKMeans::GetClusters(TArray<TArray<int32>>& OutClusters)
+{
+	OutClusters.SetNum(ClusterSizes.Num());
+	for (int32 ClusterIdx = 0; ClusterIdx < ClusterSizes.Num(); ++ClusterIdx)
+	{
+		OutClusters[ClusterIdx].Reset(ClusterSizes[ClusterIdx]);
+	}
+	if (OutClusters.IsEmpty())
+	{
+		return;
+	}
+
+	for (int32 PointIdx = 0; PointIdx < ClusterIDs.Num(); ++PointIdx)
+	{
+		OutClusters[ClusterIDs[PointIdx]].Add(PointIdx);
+	}
+}
+
+}
\ No newline at end of file
diff --git a/Engine/Source/Runtime/GeometryCore/Public/Clustering/KMeans.h b/Engine/Source/Runtime/GeometryCore/Public/Clustering/KMeans.h
new file mode 100644
--- /dev/null
+++ b/Engine/Source/Runtime/GeometryCore/Public/Clustering/KMeans.h
@@ -0,0 +1,217 @@
+// Copyright Epic Games, Inc. All Rights Reserved.
+
+#pragma once
+
+#include "Math/RandomStream.h"
+#include "MathUtil.h"
+#include "VectorTypes.h"
+#include "MatrixTypes.h"
+#include "Spatial/PriorityOrderPoints.h"
+
+namespace UE
+{
+namespace Geometry
+{
+
+using namespace UE::Math;
+
+
+struct FClusterKMeans
+{
+	/// Parameters
+
+	// Max iterations of K-Means clustering. Will use fewer iterations if the clustering method converges.
+	int32 MaxIterations = 500;
+
+	// Random Seed used to initialize clustering (if InitialCenters are not provided)
+	int32 RandomSeed = 0;
+
+	/// Outputs
+
+	// Mapping from input points to cluster IDs
+	TArray<int32> ClusterIDs;
+
+	// Number of points in each cluster
+	TArray<int32> ClusterSizes;
+
+	/**
+	 * Compute the K-Means clustering of FVector points
+	 * 
+	 * @param PointsToCluster	Points to partition into clusters
+	 * @param NumClusters		Target number of clusters to create, if InitialCenters is not provided
+	 * @param InitialCenters	If non-empty, these positions will be used to initialize the cluster locations
+	 * @param OutClusterCenters	If non-null, will be filled with the cluster centers
+	 * @return number of clusters found
+	 */
+	template<typename TVectorType = FVector>
+	int32 ComputeClusters(
+		TArrayView<const TVectorType> PointsToCluster, int32 NumClusters,
+		TArrayView<const TVectorType> InitialCenters = TArrayView<const TVectorType>(), 
+		TArray<TVectorType>* OutClusterCenters = nullptr)
+	{
+		if (PointsToCluster.IsEmpty())
+		{
+			// nothing to cluster
+			return 0;
+		}
+
+		// Initialize cluster centers
+		TArray<TVectorType> LocalCenters;
+		TArray<TVectorType>* UseCenters = &LocalCenters;
+		if (OutClusterCenters)
+		{
+			UseCenters = OutClusterCenters;
+			UseCenters->Reset();
+		}
+		if (InitialCenters.IsEmpty())
+		{
+			TArray<int32> Ordering;
+			const int32 OrderingNum = PointsToCluster.Num();
+			Ordering.SetNumUninitialized(OrderingNum);
+			for (int32 Idx = 0; Idx < OrderingNum; ++Idx)
+			{
+				Ordering[Idx] = Idx;
+			}
+			// Shuffle the first NumClusters indices and use them as the initial centers
+			NumClusters = FMath::Min(PointsToCluster.Num(), NumClusters);
+			UseCenters->Reserve(NumClusters);
+			FRandomStream RandomStream(RandomSeed);
+			for (int32 Idx = 0; Idx < NumClusters; ++Idx)
+			{
+				Swap(Ordering[Idx], Ordering[Idx + RandomStream.RandHelper(OrderingNum - Idx)]);
+				UseCenters->Add(PointsToCluster[Ordering[Idx]]);
+			}
+		}
+		else
+		{
+			// Note: We intentionally do not check if more centers were provided than points here.
+			// The excess centers will instead be removed below when no points are assigned to them.
+			UseCenters->Append(InitialCenters);
+		}
+
+		NumClusters = UseCenters->Num();
+		ClusterIDs.Init(-1, PointsToCluster.Num());
+		ClusterSizes.SetNumZeroed(NumClusters);
+
+		if (NumClusters == 0)
+		{
+			return NumClusters;
+		}
+		else if (NumClusters == 1)
+		{
+			for (int32 PointIdx = 0; PointIdx < PointsToCluster.Num(); ++PointIdx)
+			{
+				ClusterIDs[PointIdx] = 0;
+			}
+			ClusterSizes[0] = PointsToCluster.Num();
+			return NumClusters;
+		}
+
+		
+		TArray<TVectorType> NextCenters;
+		NextCenters.SetNumZeroed(NumClusters);
+
+		TArray<int32> ClusterIDRemap; // array to use if we need to remap cluster IDs due to an empty cluster
+
+		int32 UseMaxIterations = FMath::Max(1, MaxIterations); // always use at least one iteration to make sure the initial assignment happens
+		for (int32 Iterations = 0; Iterations < UseMaxIterations; ++Iterations)
+		{
+			bool bClustersChanged = Iterations == 0; // clusters always change on first iteration; otherwise we must check
+			for (int32 PointIdx = 0; PointIdx < PointsToCluster.Num(); ++PointIdx)
+			{
+				const TVectorType& Point = PointsToCluster[PointIdx];
+				double ClosestDistSq = (double)DistanceSquared(Point, (*UseCenters)[0]);
+				int32 ClosestCenter = 0;
+				for (int32 CenterIdx = 1; CenterIdx < UseCenters->Num(); ++CenterIdx)
+				{
+					double DistSq = DistanceSquared(PointsToCluster[PointIdx], (*UseCenters)[CenterIdx]);
+					if (DistSq < ClosestDistSq)
+					{
+						ClosestDistSq = DistSq;
+						ClosestCenter = CenterIdx;
+					}
+				}
+
+				ClusterSizes[ClosestCenter]++;
+				if (Iterations > 0)
+				{
+					int32 OldClusterID = ClusterIDs[PointIdx];
+					ClusterSizes[OldClusterID]--;
+					if (OldClusterID != ClosestCenter)
+					{
+						bClustersChanged = true;
+					}
+				}
+				
+				ClusterIDs[PointIdx] = ClosestCenter;
+				NextCenters[ClosestCenter] += Point;
+			}
+
+			// Stop iterating if clusters are unchanged
+			if (!bClustersChanged)
+			{
+				break;
+			}
+
+			// Update cluster centers and detect/delete any empty clusters
+			bool bDeletedClusters = false;
+			for (int32 ClusterIdx = 0; ClusterIdx < UseCenters->Num(); ++ClusterIdx)
+			{
+				if (ClusterSizes[ClusterIdx] > 0)
+				{
+					(*UseCenters)[ClusterIdx] = NextCenters[ClusterIdx] / ClusterSizes[ClusterIdx];
+					NextCenters[ClusterIdx] = TVectorType::ZeroVector;
+				}
+				else
+				{
+					if (!bDeletedClusters)
+					{
+						ClusterIDRemap.SetNumUninitialized(UseCenters->Num(), false);
+						for (int32 Idx = 0; Idx < UseCenters->Num(); ++Idx)
+						{
+							ClusterIDRemap[Idx] = Idx;
+						}
+					}
+					bDeletedClusters = true;
+					ClusterIDRemap[ClusterSizes.Num() - 1] = ClusterIdx;
+					ClusterSizes.RemoveAtSwap(ClusterIdx, 1, false);
+					UseCenters->RemoveAtSwap(ClusterIdx, 1, false);
+					NextCenters.RemoveAtSwap(ClusterIdx, 1, false);
+				}
+			}
+			if (bDeletedClusters)
+			{
+				checkSlow(!ClusterSizes.IsEmpty());
+				checkSlow(ClusterSizes.Num() == UseCenters->Num());
+				for (int32& ClusterID : ClusterIDs)
+				{
+					ClusterID = ClusterIDRemap[ClusterID];
+				}
+			}
+		}
+
+		return ClusterSizes.Num();
+	}
+
+	/**
+	 * Helper function to generate (approximately) uniform-spaced initial clusters centers, which can be passed to ComputeClusters.
+	 */
+	template<typename TVectorType = FVector>
+	void GetUniformSpacedInitialCenters(TArrayView<const TVectorType> PointsToCluster, int32 NumClusters, TArray<TVectorType>& OutCenters)
+	{
+		FPriorityOrderPoints OrderPoints;
+		OrderPoints.ComputeUniformSpaced(PointsToCluster, TArrayView<const float>(), NumClusters);
+		int32 NumOut = FMath::Min(NumClusters, OrderPoints.Order.Num());
+		OutCenters.Reset(NumOut);
+		for (int32 OrderIdx = 0; OrderIdx < NumOut; ++OrderIdx)
+		{
+			OutCenters.Add(PointsToCluster[OrderPoints.Order[OrderIdx]]);
+		}
+	}
+
+
+	void GEOMETRYCORE_API GetClusters(TArray<TArray<int32>>& OutClusters);
+};
+
+} // end namespace UE::Geometry
+} // end namespace UE
\ No newline at end of file
