Add missing feature in 2d classes based on 3d implementation :
- == and != in AxisAlignedBox2i
- default constructor in ScaleGridIndexer2

diff --git a/Engine/Source/Runtime/GeometryCore/Public/IntBoxTypes.h b/Engine/Source/Runtime/GeometryCore/Public/IntBoxTypes.h
--- a/Engine/Source/Runtime/GeometryCore/Public/IntBoxTypes.h
+++ b/Engine/Source/Runtime/GeometryCore/Public/IntBoxTypes.h
@@ -102,6 +102,15 @@ struct FAxisAlignedBox2i
 			FVector2i(-TNumericLimits<int32>::Max(), -TNumericLimits<int32>::Max()));
 	}
 
+	bool operator==(const FAxisAlignedBox2i& Other) const
+	{
+		return Max == Other.Max && Min == Other.Min;
+	}
+	bool operator!=(const FAxisAlignedBox2i& Other) const
+	{
+		return Max != Other.Max || Min != Other.Min;
+	}
+
 	/**
 	 * Corners are ordered to follow the perimeter of the bounding rectangle, starting from the (Min.X, Min.Y) corner and ending at (Min.X, Max.Y)
 	 * @param Index which corner to return, must be in range [0,3]
diff --git a/Engine/Source/Runtime/GeometryCore/Public/Util/GridIndexing2.h b/Engine/Source/Runtime/GeometryCore/Public/Util/GridIndexing2.h
--- a/Engine/Source/Runtime/GeometryCore/Public/Util/GridIndexing2.h
+++ b/Engine/Source/Runtime/GeometryCore/Public/Util/GridIndexing2.h
@@ -23,6 +23,11 @@ struct TScaleGridIndexer2
 	/** Real-valued size of an integer grid cell */
 	RealType CellSize;
 
+	TScaleGridIndexer2() : CellSize((RealType)1) 
+	{
+		ensure(CellSize >= TMathUtil<RealType>::ZeroTolerance);
+	}
+
 	TScaleGridIndexer2(RealType CellSize) : CellSize(CellSize)
 	{
 		ensure(CellSize >= TMathUtil<RealType>::ZeroTolerance);
