DynamicMesh: trivial conversion of Async/TFuture usage to new UE::Tasks
#rnx

diff --git a/Engine/Source/Runtime/GeometryCore/Private/DynamicMesh/DynamicMeshAttributeSet.cpp b/Engine/Source/Runtime/GeometryCore/Private/DynamicMesh/DynamicMeshAttributeSet.cpp
--- a/Engine/Source/Runtime/GeometryCore/Private/DynamicMesh/DynamicMeshAttributeSet.cpp
+++ b/Engine/Source/Runtime/GeometryCore/Private/DynamicMesh/DynamicMeshAttributeSet.cpp
@@ -5,7 +5,7 @@
 #include "DynamicMesh/DynamicVertexSkinWeightsAttribute.h"
 #include "DynamicMesh/DynamicBoneAttribute.h"
 #include "IndexTypes.h"
-#include "Async/Async.h"
+#include "Tasks/Task.h"
 #include "Serialization/NameAsStringProxyArchive.h"
 
 
@@ -238,12 +238,12 @@ void FDynamicMeshAttributeSet::SplitAllBowties(bool bParallel)
 	int32 UVLayerCount = NumUVLayers();
 	int32 NormalLayerCount = NumNormalLayers();
 	
-	TArray<TFuture<void>> Pending;
+	TArray<UE::Tasks::FTask> Pending;
 	auto ASyncOrRunSplit = [&Pending, bParallel](auto Overlay)->void
 	{
 		if (bParallel)
 		{	
-			auto AsyncTask = Async(EAsyncExecution::ThreadPool, [Overlay]() 
+			UE::Tasks::FTask AsyncTask = UE::Tasks::Launch(UE_SOURCE_LOCATION, [Overlay]() 
 			{
 				Overlay->SplitBowties();
 			});
@@ -272,10 +272,7 @@ void FDynamicMeshAttributeSet::SplitAllBowties(bool bParallel)
 	}
 
 	// this array will be empty if bParallel == false
-	for (TFuture<void>& Future : Pending)
-	{
-		Future.Wait();
-	}
+	UE::Tasks::Wait(Pending);
 }
 
 
