DynamicMesh: Fix overlay elements not being freed and prevent changing some overlay values arbitrarily when the collapsed triangles were unset.

diff --git a/Engine/Source/Runtime/GeometryCore/Private/DynamicMesh/DynamicMeshOverlay.cpp b/Engine/Source/Runtime/GeometryCore/Private/DynamicMesh/DynamicMeshOverlay.cpp
--- a/Engine/Source/Runtime/GeometryCore/Private/DynamicMesh/DynamicMeshOverlay.cpp
+++ b/Engine/Source/Runtime/GeometryCore/Private/DynamicMesh/DynamicMeshOverlay.cpp
@@ -1226,95 +1226,76 @@ void TDynamicMeshOverlay<RealType, ElementSize>::OnCollapseEdge(const FDynamicMe
 		checkSlow(removed_elemid[1] != FDynamicMesh3::InvalidID);
 		
 	}
-	if (!bT0Set && !bT1Set)
+
+	// update value of kept elements
+	for (int i = 0; i < 2; ++i)
 	{
-		// if neither triangle was set, still need to check surrounding triangles; one or both elements may still be used
-		checkSlow(!bFoundRemovedElement[0] && !bFoundRemovedElement[1]);
-		
-		// note this may do the wrong thing in the case that at least one of the verts has split elements.
-		// higher-level code should make sure this doesn't happen.
-		
-		for (int onering_tid : ParentMesh->VtxTrianglesItr(vid_base_kept))
+		if (kept_elemid[i] == FDynamicMesh3::InvalidID || removed_elemid[i] == FDynamicMesh3::InvalidID)
 		{
-			if (!IsSetTriangle(onering_tid))
-			{
-				continue;
-			}
-			FIndex3i elem_tri = GetTriangle(onering_tid);
-			for (int j = 0; j < 3; ++j) 
-			{
-				int elem_id = elem_tri[j];
-				int vid_base = ParentVertices[elem_id];
-				if (vid_base == vid_base_removed)
-				{
-					bFoundRemovedElement[0] = true;
-					removed_elemid[0] = elem_id;
-				}
-				if (vid_base == vid_base_kept)
-				{
-					bFoundKeptElement[0] = true;
-					kept_elemid[0] = elem_id;
-				}
-			}
+			continue;
 		}
+
+		SetElementFromLerp(kept_elemid[i], kept_elemid[i], removed_elemid[i], (double)collapseInfo.CollapseT);
 	}
 
-	// look for still-existing triangles that have elements linked to the removed vertex.
-	// in that case, replace with the element we found
+	// Helper for detaching from elements further below. Technically, the freeing gets done for us if the 
+	// triangle unset call is the last detachment (which it should be as long as only elements on a removed 
+	// overlay edge are ones that are removed), but it is saner to have it.
+	auto DecrementAndFreeIfLast = [this](int32 elem_id)
+	{
+		ElementsRefCounts.Decrement(elem_id);
+		if (ElementsRefCounts.GetRefCount(elem_id) == 1)
+		{
+			ElementsRefCounts.Decrement(elem_id);
+			ParentVertices[elem_id] = FDynamicMesh3::InvalidID;
+		}
+	};
+
 
-	if (bFoundRemovedElement[0] || bFoundRemovedElement[1])
+	// Look for still-existing triangles that have elements linked to the removed vertex and update them.
+	// Note that this has to happen even if both triangles were unset, as the removed vertex may have had
+	// other elements associated with it, so we need to look at its triangles (which are now attached to
+	// vid_base_kept).
+	for (int onering_tid : ParentMesh->VtxTrianglesItr(vid_base_kept))
 	{
-		for (int onering_tid : ParentMesh->VtxTrianglesItr(vid_base_kept))
+		if (!IsSetTriangle(onering_tid))
 		{
-			if (!IsSetTriangle(onering_tid))
-			{
-				continue;
-			}
-			FIndex3i elem_tri = GetTriangle(onering_tid);
-			for (int j = 0; j < 3; ++j)
+			continue;
+		}
+		FIndex3i elem_tri = GetTriangle(onering_tid);
+		for (int j = 0; j < 3; ++j)
+		{
+			int elem_id = elem_tri[j];
+			if (ParentVertices[elem_id] == vid_base_removed)
 			{
-				int elem_id = elem_tri[j];
-				if (ParentVertices[elem_id] == vid_base_removed)
+				if (elem_id == removed_elemid[0])
 				{
-					if (elem_id == removed_elemid[0])
-					{
-						ElementsRefCounts.Decrement(elem_id);
-						ElementTriangles[3 * onering_tid + j] = kept_elemid[0];
-						if (bFoundKeptElement[0])
-						{
-							ElementsRefCounts.Increment(kept_elemid[0]);
-						}
-					}
-					else if (elem_id == removed_elemid[1])
+					ElementTriangles[3 * onering_tid + j] = kept_elemid[0];
+					if (bFoundKeptElement[0])
 					{
-						ElementsRefCounts.Decrement(elem_id);
-						ElementTriangles[3 * onering_tid + j] = kept_elemid[1];
-						if (bFoundKeptElement[1])
-						{
-							ElementsRefCounts.Increment(kept_elemid[1]);
-						}
+						ElementsRefCounts.Increment(kept_elemid[0]);
 					}
-					else
+					DecrementAndFreeIfLast(elem_id);
+				}
+				else if (elem_id == removed_elemid[1])
+				{
+					ElementTriangles[3 * onering_tid + j] = kept_elemid[1];
+					if (bFoundKeptElement[1])
 					{
-						// this could happen if a split edge is adjacent to the edge we collapse
-						ParentVertices[elem_id] = vid_base_kept;
+						ElementsRefCounts.Increment(kept_elemid[1]);
 					}
+					DecrementAndFreeIfLast(elem_id);
+				}
+				else
+				{
+					// this could happen if a split edge is adjacent to the edge we collapse
+					ParentVertices[elem_id] = vid_base_kept;
 				}
 			}
 		}
-
-		for (int i = 0; i < 2; ++i)
-		{
-			if (kept_elemid[i] == FDynamicMesh3::InvalidID || removed_elemid[i] == FDynamicMesh3::InvalidID)
-			{
-				continue;
-			}
-
-			// update value of kept element
-			SetElementFromLerp(kept_elemid[i], kept_elemid[i], removed_elemid[i], (double)collapseInfo.CollapseT);
-		}
 	}
 
+
 	// clear the two triangles we removed
 	if (bT0Set)
 	{
@@ -1331,12 +1312,12 @@ void TDynamicMeshOverlay<RealType, ElementSize>::OnCollapseEdge(const FDynamicMe
 		removed_elemid[1] = FDynamicMesh3::InvalidID;
 	}
 
-	// Note: the elements associated with the removed vertex should already have been removed when the triangles were un-set above
+	// Note: the elements associated with the removed vertex should have been removed in the iteration or triangle unsetting above
 #if UE_BUILD_DEBUG
 	for (int k = 0; k < 2; ++k)
 	{
 		if (removed_elemid[k] != FDynamicMesh3::InvalidID)
-	{
+		{
 			int rc = ElementsRefCounts.GetRefCount(removed_elemid[k]);
 			checkSlow(rc == 0);
 		}
