[Geometry] Added profiler traces to mesh generators

diff --git a/Engine/Source/Runtime/GeometryCore/Private/Generators/DiscMeshGenerator.cpp b/Engine/Source/Runtime/GeometryCore/Private/Generators/DiscMeshGenerator.cpp
--- a/Engine/Source/Runtime/GeometryCore/Private/Generators/DiscMeshGenerator.cpp
+++ b/Engine/Source/Runtime/GeometryCore/Private/Generators/DiscMeshGenerator.cpp
@@ -27,6 +27,8 @@ FPuncturedDiscMeshGenerator::FPuncturedDiscMeshGenerator()
 
 FMeshShapeGenerator& FDiscMeshGenerator::Generate()
 {
+	TRACE_CPUPROFILER_EVENT_SCOPE(DiscMeshGenerator_Generate);
+	
 	int AngleNV = AngleSamples > 3 ? AngleSamples : 3;
 	int RadialNV = RadialSamples > 1 ? RadialSamples : 1;
 	int NumVertices = AngleNV * RadialNV + 1;
diff --git a/Engine/Source/Runtime/GeometryCore/Private/Generators/FlatTriangulationMeshGenerator.cpp b/Engine/Source/Runtime/GeometryCore/Private/Generators/FlatTriangulationMeshGenerator.cpp
--- a/Engine/Source/Runtime/GeometryCore/Private/Generators/FlatTriangulationMeshGenerator.cpp
+++ b/Engine/Source/Runtime/GeometryCore/Private/Generators/FlatTriangulationMeshGenerator.cpp
@@ -22,6 +22,8 @@ FMeshShapeGenerator& FFlatTriangulationMeshGenerator::Generate()
 		return *this;
 	}
 
+	TRACE_CPUPROFILER_EVENT_SCOPE(FlatTriangulationMeshGenerator_Generate);
+
 	int NumTriangles = Triangles2D.Num();
 	SetBufferSizes(NumVertices, NumTriangles, NumVertices, NumVertices);
 
diff --git a/Engine/Source/Runtime/GeometryCore/Private/Generators/PlanarPolygonMeshGenerator.cpp b/Engine/Source/Runtime/GeometryCore/Private/Generators/PlanarPolygonMeshGenerator.cpp
--- a/Engine/Source/Runtime/GeometryCore/Private/Generators/PlanarPolygonMeshGenerator.cpp
+++ b/Engine/Source/Runtime/GeometryCore/Private/Generators/PlanarPolygonMeshGenerator.cpp
@@ -35,6 +35,8 @@ FMeshShapeGenerator& FPlanarPolygonMeshGenerator::Generate()
 		return *this;
 	}
 
+	TRACE_CPUPROFILER_EVENT_SCOPE(PlanarPolygonMeshGenerator_Generate);
+
 	TArray<FIndex3i> TriangleList;
 	PolygonTriangulation::TriangulateSimplePolygon(Polygon.GetVertices(), TriangleList, false);
 
diff --git a/Engine/Source/Runtime/GeometryCore/Private/Generators/PolygonEdgeMeshGenerator.cpp b/Engine/Source/Runtime/GeometryCore/Private/Generators/PolygonEdgeMeshGenerator.cpp
--- a/Engine/Source/Runtime/GeometryCore/Private/Generators/PolygonEdgeMeshGenerator.cpp
+++ b/Engine/Source/Runtime/GeometryCore/Private/Generators/PolygonEdgeMeshGenerator.cpp
@@ -299,6 +299,8 @@ void FPolygonEdgeMeshGenerator::CurvePath(const TArray<FVector3d>& InPath,
 FMeshShapeGenerator& FPolygonEdgeMeshGenerator::Generate()
 {
 	using namespace PolygonEdgeMeshGeneratorLocal;
+	
+	TRACE_CPUPROFILER_EVENT_SCOPE(PolygonEdgeMeshGenerator_Generate);
 
 	const int32 NumInputVertices = Polygon.Num();
 
diff --git a/Engine/Source/Runtime/GeometryCore/Private/Generators/RectangleMeshGenerator.cpp b/Engine/Source/Runtime/GeometryCore/Private/Generators/RectangleMeshGenerator.cpp
--- a/Engine/Source/Runtime/GeometryCore/Private/Generators/RectangleMeshGenerator.cpp
+++ b/Engine/Source/Runtime/GeometryCore/Private/Generators/RectangleMeshGenerator.cpp
@@ -23,6 +23,8 @@ FMeshShapeGenerator& FRectangleMeshGenerator::Generate()
 	check(IndicesMap.A >= 0 && IndicesMap.A <= 2);
 	check(IndicesMap.B >= 0 && IndicesMap.B <= 2);
 
+	TRACE_CPUPROFILER_EVENT_SCOPE(RectangleMeshGenerator_Generate);
+
 	int WidthNV = (WidthVertexCount > 1) ? WidthVertexCount : 2;
 	int HeightNV = (HeightVertexCount > 1) ? HeightVertexCount : 2;
 
diff --git a/Engine/Source/Runtime/GeometryCore/Private/Generators/RevolveGenerator.cpp b/Engine/Source/Runtime/GeometryCore/Private/Generators/RevolveGenerator.cpp
--- a/Engine/Source/Runtime/GeometryCore/Private/Generators/RevolveGenerator.cpp
+++ b/Engine/Source/Runtime/GeometryCore/Private/Generators/RevolveGenerator.cpp
@@ -119,6 +119,8 @@ namespace UELocal
 
 FDynamicMesh3 FRevolvePlanarPathGenerator::GenerateMesh()
 {
+	TRACE_CPUPROFILER_EVENT_SCOPE(RevolvePlanarPathGenerator_GenerateMesh);
+	
 	// revolve around +Z axis, around origin
 	FVector3d AxisDirection(0, 0, 1);
 	FVector3d AxisOrigin(0, 0, 0);
@@ -327,6 +329,8 @@ FDynamicMesh3 FRevolvePlanarPathGenerator::GenerateMesh()
 
 FDynamicMesh3 FRevolvePlanarPolygonGenerator::GenerateMesh()
 {
+	TRACE_CPUPROFILER_EVENT_SCOPE(RevolvePlanarPolygonGenerator_GenerateMesh);
+	
 	// revolve around +Z axis, around origin
 	FVector3d AxisDirection(0, 0, 1);
 	FVector3d AxisOrigin(0, 0, 0);
@@ -516,6 +520,8 @@ FDynamicMesh3 FRevolvePlanarPolygonGenerator::GenerateMesh()
 
 FDynamicMesh3 FSpiralRevolvePlanarPolygonGenerator::GenerateMesh()
 {
+	TRACE_CPUPROFILER_EVENT_SCOPE(SpiralRevolvePlanarPolygonGenerator_GenerateMesh);
+	
 	// revolve around +Z axis, around origin
 	FVector3d AxisDirection(0, 0, 1);
 	FVector3d AxisOrigin(0, 0, 0);
diff --git a/Engine/Source/Runtime/GeometryCore/Private/Generators/StairGenerator.cpp b/Engine/Source/Runtime/GeometryCore/Private/Generators/StairGenerator.cpp
--- a/Engine/Source/Runtime/GeometryCore/Private/Generators/StairGenerator.cpp
+++ b/Engine/Source/Runtime/GeometryCore/Private/Generators/StairGenerator.cpp
@@ -11,6 +11,8 @@ using namespace UE::Geometry;
 
 FMeshShapeGenerator& FStairGenerator::Generate()
 {
+	TRACE_CPUPROFILER_EVENT_SCOPE(StairGenerator_Generate);
+	
 	ResetData();
 
 	switch (StairStyle)
diff --git a/Engine/Source/Runtime/GeometryCore/Private/Generators/SweepGenerator.cpp b/Engine/Source/Runtime/GeometryCore/Private/Generators/SweepGenerator.cpp
--- a/Engine/Source/Runtime/GeometryCore/Private/Generators/SweepGenerator.cpp
+++ b/Engine/Source/Runtime/GeometryCore/Private/Generators/SweepGenerator.cpp
@@ -109,6 +109,8 @@ FMeshShapeGenerator& FProfileSweepGenerator::Generate()
 		Reset();
 		return *this;
 	}
+	
+	TRACE_CPUPROFILER_EVENT_SCOPE(ProfileSweepGenerator_Generate);
 
 	/*
 	The generated vertices are organized and connected as follows:
