remove unused local variable in Delaunay2.cpp

diff --git a/Engine/Source/Runtime/GeometryCore/Private/CompGeom/Delaunay2.cpp b/Engine/Source/Runtime/GeometryCore/Private/CompGeom/Delaunay2.cpp
--- a/Engine/Source/Runtime/GeometryCore/Private/CompGeom/Delaunay2.cpp
+++ b/Engine/Source/Runtime/GeometryCore/Private/CompGeom/Delaunay2.cpp
@@ -792,7 +792,6 @@ namespace DelaunayInternal
 	{
 		IsDuplicateOfOut = FDelaunay2Connectivity::InvalidIndex;
 
-		FIndex2i FirstEdge(FDelaunay2Connectivity::InvalidIndex, FDelaunay2Connectivity::InvalidIndex);
 		constexpr int32 GhostV = FDelaunay2Connectivity::GhostIndex; // shorter name
 
 		auto ChooseCross = [&Random, bAssumeDelaunay, &Vertices, Vertex, GhostV](const FIndex3i& Tri, bool bSkipFirst) -> int32
