Add support for direct construction of a simplified convex hull in GeometryCore's ConvexHull3 via:

1. A MaxHullVertices setting. If positive, the convex hull algorithm will greedily add vertices to best expand the hull, generating a simplified hull if the exact hull would exceed the specified budget.
2. SkipAtHullDistanceAbsolute/SkipAtHullDistanceAsFraction settings. If positive, the convex hull algorithm will ignore points that are closer than the threshold value, resulting in a simpler hull.

Also stop relying on the Orient3 exact predicate to find the farthest hull points (regardless of whether PointsBudget is specified) both so that the PointsBudget and SkipAtHullDistance settings can accurately compare distances from different triangles and also to fix an issue where the Orient3 magnitude was discarded when requesting a single-precision Orient3 value.

Update convex hull unit test to also test these new settings.

Also adds a small optimization the reduces the number of duplicate points considered.

diff --git a/Engine/Source/Runtime/GeometryCore/Private/CompGeom/ConvexHull3.cpp b/Engine/Source/Runtime/GeometryCore/Private/CompGeom/ConvexHull3.cpp
--- a/Engine/Source/Runtime/GeometryCore/Private/CompGeom/ConvexHull3.cpp
+++ b/Engine/Source/Runtime/GeometryCore/Private/CompGeom/ConvexHull3.cpp
@@ -182,9 +182,31 @@ struct FHullConnectivity
 		TArray<int32> Indices;
 		int32 MaxIdx = -1; // Note: an index into the source point array, not into the above Indices array
 		// Max value of the Orient3D predicate, indicating the farthest point outside; values less than zero are inside the hull so not considered
-		double MaxValue = 0;
+		double MaxValue = -FMathd::MaxReal;
+		FPlane3d Plane;
 
-		void AddPt(int32 Idx, double Value)
+		void SetPlane(const TVector<RealType> TriPts[3])
+		{
+			Plane = FPlane3d(FVector3d(TriPts[0]), FVector3d(TriPts[1]), FVector3d(TriPts[2]));
+		}
+
+		double GetPlaneDistance(const TVector<RealType>& Pt)
+		{
+			return Plane.DistanceTo((FVector3d)Pt);
+		}
+
+		void AddPt(int32 Idx, const TVector<RealType>& Pt)
+		{
+			double Value = GetPlaneDistance(Pt);
+			if (Value > MaxValue)
+			{
+				MaxValue = Value;
+				MaxIdx = Idx;
+			}
+			Indices.Add(Idx);
+		}
+
+		void AddPtByValue(int32 Idx, double Value)
 		{
 			if (Value > MaxValue)
 			{
@@ -195,7 +217,7 @@ struct FHullConnectivity
 		}
 
 		// Remove a point from the visible point set; if it was the tracked MaxValue point, find a new MaxValue point
-		void RemovePt(int32 SourcePointIdx, const TVector<RealType> TriPts[3], TFunctionRef<TVector<RealType>(int32)> GetPointFunc)
+		void RemovePt(int32 SourcePointIdx, TFunctionRef<TVector<RealType>(int32)> GetPointFunc)
 		{
 			if (MaxIdx == SourcePointIdx)
 			{
@@ -211,7 +233,7 @@ struct FHullConnectivity
 					}
 					else
 					{
-						double Value = ExactPredicates::Orient3<RealType>(TriPts[0], TriPts[1], TriPts[2], GetPointFunc(PointIdx));
+						double Value = GetPlaneDistance(GetPointFunc(PointIdx));
 						if (Value > MaxValue)
 						{
 							MaxValue = Value;
@@ -240,13 +262,19 @@ struct FHullConnectivity
 		{
 			Indices.Reset();
 			MaxIdx = -1;
-			MaxValue = 0;
+			MaxValue = -FMathd::MaxReal;
 		}
 	};
 
 	TArray<FIndex3i> TriNeighbors;
 	TArray<FVisiblePoints> VisiblePoints;
 	TSet<int32> TrisWithPoints;
+	
+	TArray<uint16> PointMemberships; // Used for tracking set membership for point indices
+	uint16 MembershipNumber = 0;
+
+	// If positive, this threshold additionally filters which points are considered 'visible' as only points at least this far from the plane
+	double VisibleDistanceThreshold = -FMathd::MaxReal;
 
 	/**
 	 * Fully build neighbor connectivity; only on the initial tet, 
@@ -277,13 +305,12 @@ struct FHullConnectivity
 	}
 
 	/**
-	 * @param ValueOut The volume of the tetrahedron formed by TriPts and Pt
 	 * @return true if Pt is on the 'positive' side of the triangle
 	 */
-	bool IsVisible(const TVector<RealType> TriPts[3], const TVector<RealType>& Pt, double& ValueOut)
+	bool IsVisible(const TVector<RealType> TriPts[3], const TVector<RealType>& Pt)
 	{
-		ValueOut = ExactPredicates::Orient3<RealType>(TriPts[0], TriPts[1], TriPts[2], Pt);
-		return ValueOut > 0;
+		double PredicateValue = ExactPredicates::Orient3<RealType>(TriPts[0], TriPts[1], TriPts[2], Pt);
+		return PredicateValue > 0;
 	}
 
 	/**
@@ -305,24 +332,27 @@ struct FHullConnectivity
 
 		TVector<RealType> TriPts[3];
 		TVector<RealType> Pt;
-		for (int32 PtIdx = 0; PtIdx < NumPoints; PtIdx++)
+		for (int32 TriIdx = 0; TriIdx < Triangles.Num(); TriIdx++)
 		{
-			if (!FilterFunc(PtIdx))
-			{
-				continue;
-			}
-			Pt = GetPointFunc(PtIdx);
-			for (int32 TriIdx = 0; TriIdx < Triangles.Num(); TriIdx++)
+			SetTriPts(Triangles[TriIdx], GetPointFunc, TriPts);
+			VisiblePoints[TriIdx].SetPlane(TriPts);
+
+			for (int32 PtIdx = 0; PtIdx < NumPoints; PtIdx++)
 			{
-				SetTriPts(Triangles[TriIdx], GetPointFunc, TriPts);
-				double Value;
-				if (IsVisible(TriPts, Pt, Value))
+				if (!FilterFunc(PtIdx))
+				{
+					continue;
+				}
+				Pt = GetPointFunc(PtIdx);
+				double Distance = VisiblePoints[TriIdx].GetPlaneDistance(Pt);
+				bool bDistanceOk = VisibleDistanceThreshold < 0 || Distance > VisibleDistanceThreshold;
+				if (bDistanceOk && IsVisible(TriPts, Pt))
 				{
 					if (!VisiblePoints[TriIdx].Num())
 					{
 						TrisWithPoints.Add(TriIdx);
 					}
-					VisiblePoints[TriIdx].AddPt(PtIdx, Value);
+					VisiblePoints[TriIdx].AddPtByValue(PtIdx, Distance);
 				}
 			}
 		}
@@ -474,19 +504,42 @@ struct FHullConnectivity
 	}
 
 	/**
+	 * @param bChooseBestPoint	Whether to choose the point with the highest value, rather than the first found point
 	 * @return A triangle index and visible-from-that-triangle point index that can be added to the hull next
 	 */
-	FIndex2i ChooseVisiblePoint()
+	FIndex2i ChooseVisiblePoint(bool bChooseBestPoint)
 	{
 		FIndex2i FoundTriPointPair(-1, -1);
-		for (int32 TriIdx : TrisWithPoints)
+		if (bChooseBestPoint)
 		{
-			ensure(VisiblePoints[TriIdx].Num() > 0);
+			// Note: If MaxHullVertices is large enough, it could make sense to use a priority queue to efficiently track the triangles with farthest points.
+			// Though if MaxHullVertices is small, this simple linear pass will be faster.
+			// TODO: Try using FIndexPriorityQueue to track the faces with the best visible points, for the large vertex count case.
+			// (Note the priority queue would need the max triangle index; this should be (2*FMath::Min(NumPoints, MaxHullVertices)-4))
+			double BestValue = -1;
+			for (int32 TriIdx : TrisWithPoints)
+			{
+				checkSlow(VisiblePoints[TriIdx].Num() > 0);
 
-			FoundTriPointPair[0] = TriIdx;
-			// choose the "max point" -- the point with the largest volume when it makes a tetrahedron w/ the triangle
-			FoundTriPointPair[1] = VisiblePoints[TriIdx].MaxPt();
-			break;
+				if (VisiblePoints[TriIdx].MaxValue > BestValue)
+				{
+					FoundTriPointPair[0] = TriIdx;
+					FoundTriPointPair[1] = VisiblePoints[TriIdx].MaxPt();
+					BestValue = VisiblePoints[TriIdx].MaxValue;
+				}
+			}
+		}
+		else
+		{
+			for (int32 TriIdx : TrisWithPoints)
+			{
+				checkSlow(VisiblePoints[TriIdx].Num() > 0);
+
+				FoundTriPointPair[0] = TriIdx;
+				// choose the "max point" -- the point with the largest volume when it makes a tetrahedron w/ the triangle
+				FoundTriPointPair[1] = VisiblePoints[TriIdx].MaxPt();
+				break;
+			}
 		}
 		return FoundTriPointPair;
 	}
@@ -494,11 +547,9 @@ struct FHullConnectivity
 	/**
 	 * @param TriPointPair		A tri point pair, as returned from ChooseVisiblePoint, that should be removed from consideration
 	 */
-	void RemoveVisiblePoint(FIndex2i TriPointPair, const TArray<FIndex3i>& Triangles, TFunctionRef<TVector<RealType>(int32)> GetPointFunc)
+	void RemoveVisiblePoint(FIndex2i TriPointPair, TFunctionRef<TVector<RealType>(int32)> GetPointFunc)
 	{
-		TVector<RealType> TriPts[3];
-		SetTriPts(Triangles[TriPointPair.A], GetPointFunc, TriPts);
-		VisiblePoints[TriPointPair.A].RemovePt(TriPointPair.B, TriPts, GetPointFunc);
+		VisiblePoints[TriPointPair.A].RemovePt(TriPointPair.B, GetPointFunc);
 		if (VisiblePoints[TriPointPair.A].Indices.IsEmpty())
 		{
 			TrisWithPoints.Remove(TriPointPair.A);
@@ -552,8 +603,7 @@ struct FHullConnectivity
 		{
 			TVector<RealType> TriPts[3];
 			SetTriPts(Triangles[TriIdx], GetPointFunc, TriPts);
-			double UnusedValue;
-			if (!IsVisible(TriPts, Pt, UnusedValue))
+			if (!IsVisible(TriPts, Pt))
 			{
 				return false;
 			}
@@ -594,7 +644,7 @@ struct FHullConnectivity
 		return true;
 	}
 
-	bool UpdateHullWithNewPoint(TArray<FIndex3i>& Triangles, TFunctionRef<TVector<RealType>(int32)> GetPointFunc, int32 StartTriIdx, int32 PtIdx, RealType DegenerateEdgeToleranceSq = (RealType)0)
+	bool UpdateHullWithNewPoint(TArray<FIndex3i>& Triangles, int32 NumPoints, TFunctionRef<TVector<RealType>(int32)> GetPointFunc, int32 StartTriIdx, int32 PtIdx, RealType DegenerateEdgeToleranceSq = (RealType)0)
 	{
 		// Note: Commented-out ValidateConnectivity calls are very slow, but useful for debugging if the algorithm produces an invalid result
 		//ValidateConnectivity(Triangles);
@@ -631,6 +681,35 @@ struct FHullConnectivity
 		int32 NewTriStart = Triangles.Num();
 		int32 NumAdd = ToAdd.Num();
 		TVector<RealType> TriPts[3];
+
+		// Remove duplicates from the unclaimed list (unless the list is small)
+		if (NewlyUnclaimed.Num() > 10)
+		{
+			// Use  PointMemberships to track if we've already seen the point
+			if (PointMemberships.Num() != NumPoints || MembershipNumber == MAX_uint16)
+			{
+				MembershipNumber = 1;
+				PointMemberships.Reset();
+				PointMemberships.SetNumZeroed(NumPoints);
+			}
+			else
+			{
+				MembershipNumber++;
+			}
+			for (int32 Idx = 0; Idx < NewlyUnclaimed.Num(); ++Idx)
+			{
+				int32 UnclaimedIdx = NewlyUnclaimed[Idx];
+				if (PointMemberships[UnclaimedIdx] == MembershipNumber)
+				{
+					NewlyUnclaimed.RemoveAtSwap(Idx, 1, false);
+				}
+				else
+				{
+					PointMemberships[UnclaimedIdx] = MembershipNumber;
+				}
+			}
+		}
+
 		for (int32 AddIdx = 0; AddIdx < NumAdd; AddIdx++)
 		{
 			const FNewTriangle& TriData = ToAdd[AddIdx];
@@ -643,19 +722,40 @@ struct FHullConnectivity
 			UpdateNeighbor(Triangles, AcrossTriIdx, TriData.EdgeVertices.B, NewTriIdx);
 			FVisiblePoints& Visible = VisiblePoints.Emplace_GetRef();
 			SetTriPts(Triangles[NewTriIdx], GetPointFunc, TriPts);
+			Visible.SetPlane(TriPts);
 			// claim any claim-able points
-			for (int32 UnclaimedIdx = 0; UnclaimedIdx < NewlyUnclaimed.Num(); UnclaimedIdx++)
+			if (VisibleDistanceThreshold > 0) // If using VisibleDistanceThreshold, compute PlaneDist first so we can use it to threshold
 			{
-				int32 UnPtIdx = NewlyUnclaimed[UnclaimedIdx];
-				double Value;
-				if (IsVisible(TriPts, GetPointFunc(UnPtIdx), Value))
+				for (int32 UnclaimedIdx = 0; UnclaimedIdx < NewlyUnclaimed.Num(); UnclaimedIdx++)
 				{
-					Visible.AddPt(UnPtIdx, Value);
-					NewlyUnclaimed.RemoveAtSwap(UnclaimedIdx, 1, false);
-					UnclaimedIdx--;
-					continue;
+					int32 UnPtIdx = NewlyUnclaimed[UnclaimedIdx];
+					TVector<RealType> UnPt = GetPointFunc(UnPtIdx);
+					double PlaneDist = Visible.GetPlaneDistance(UnPt);
+					if (PlaneDist > VisibleDistanceThreshold && IsVisible(TriPts, UnPt))
+					{
+						Visible.AddPtByValue(UnPtIdx, PlaneDist);
+						NewlyUnclaimed.RemoveAtSwap(UnclaimedIdx, 1, false);
+						UnclaimedIdx--;
+						continue;
+					}
 				}
 			}
+			else // otherwise skip the comparison and only compute PlaneDist when adding the point
+			{
+				for (int32 UnclaimedIdx = 0; UnclaimedIdx < NewlyUnclaimed.Num(); UnclaimedIdx++)
+				{
+					int32 UnPtIdx = NewlyUnclaimed[UnclaimedIdx];
+					TVector<RealType> UnPt = GetPointFunc(UnPtIdx);
+					if (IsVisible(TriPts, UnPt))
+					{
+						Visible.AddPt(UnPtIdx, UnPt);
+						NewlyUnclaimed.RemoveAtSwap(UnclaimedIdx, 1, false);
+						UnclaimedIdx--;
+						continue;
+					}
+				}
+			}
+
 			if (Visible.Num() > 0)
 			{
 				TrisWithPoints.Add(NewTriIdx);
@@ -712,28 +812,40 @@ bool TConvexHull3<RealType>::Solve(int32 NumPoints, TFunctionRef<TVector<RealTyp
 	Hull.Add(FIndex3i(InitialTet.Extreme[0], InitialTet.Extreme[2], InitialTet.Extreme[1]));
 
 	NumHullPoints = 4;
+	bool bHasPointBudget = MaxHullVertices > 0;
+	if (!bHasPointBudget)
+	{
+		MaxHullVertices = NumPoints;
+	}
 
 	RealType DegenerateEdgeToleranceSq = DegenerateEdgeTolerance * DegenerateEdgeTolerance;
 
 	FHullConnectivity<RealType> Connectivity;
+	double UseSkipAtHullDistance = SkipAtHullDistanceAbsolute;
+	if (SkipAtHullDistanceAsFraction > 0)
+	{
+		double Extent = (GetPointFunc(InitialTet.Extreme[0]) - GetPointFunc(InitialTet.Extreme[1])).Length();
+		UseSkipAtHullDistance = FMathd::Max(UseSkipAtHullDistance, Extent * SkipAtHullDistanceAsFraction);
+	}
+	Connectivity.VisibleDistanceThreshold = UseSkipAtHullDistance;
 	Connectivity.BuildNeighbors(Hull);
 	Connectivity.InitVisibility(Hull, NumPoints, GetPointFunc, FilterFunc);
-	while (true)
+	while (NumHullPoints < MaxHullVertices)
 	{
 		if (Progress && (NumHullPoints % 100) == 0 && Progress->Cancelled())
 		{
 			return false;
 		}
-		FIndex2i Visible = Connectivity.ChooseVisiblePoint();
+		FIndex2i Visible = Connectivity.ChooseVisiblePoint(bHasPointBudget);
 		if (Visible.A == -1)
 		{
 			break;
 		}
 		NumHullPoints++;
-		bool bAdded = Connectivity.UpdateHullWithNewPoint(Hull, GetPointFunc, Visible.A, Visible.B, DegenerateEdgeToleranceSq);
+		bool bAdded = Connectivity.UpdateHullWithNewPoint(Hull, NumPoints, GetPointFunc, Visible.A, Visible.B, DegenerateEdgeToleranceSq);
 		if (!bAdded)
 		{
-			Connectivity.RemoveVisiblePoint(Visible, Hull, GetPointFunc);
+			Connectivity.RemoveVisiblePoint(Visible, GetPointFunc);
 		}
 	}
 
diff --git a/Engine/Source/Runtime/GeometryCore/Public/CompGeom/ConvexHull3.h b/Engine/Source/Runtime/GeometryCore/Public/CompGeom/ConvexHull3.h
--- a/Engine/Source/Runtime/GeometryCore/Public/CompGeom/ConvexHull3.h
+++ b/Engine/Source/Runtime/GeometryCore/Public/CompGeom/ConvexHull3.h
@@ -68,6 +68,15 @@ public:
 	/// If zero, no points will be skipped.
 	RealType DegenerateEdgeTolerance = (RealType)0;
 
+	/// If positive, hulls generated will only have at most this many points.
+	int32 MaxHullVertices = -1;
+
+	/// If positive, skip adding points that are closer than this threshold to the in-progress hull.
+	RealType SkipAtHullDistanceAbsolute = -TMathUtil<RealType>::MaxReal;
+	/// If positive, skip adding points that are closer than this threshold to the in-progress hull -- expressed as a fraction of the overall hull extent.
+	/// Note if both the Absolute and Fraction MinPlaneDistance thresholds are set, they will both apply (i.e., the larger of the two will be used)
+	RealType SkipAtHullDistanceAsFraction = -TMathUtil<RealType>::MaxReal;
+
 	/**
 	 * Generate convex hull as long as input is not degenerate
 	 * If input is degenerate, this will return false, and caller can call GetDimension()
@@ -75,10 +84,14 @@ public:
 	 *
 	 * @param NumPoints Number of points to consider
 	 * @param GetPointFunc Function providing array-style access into points
-	 * @param Filter Optional filter to include only a subset of the points in the output hull
+	 * @param FilterFunc Optional filter to include only a subset of the points in the output hull
 	 * @return true if hull was generated, false if points span < 2 dimensions
 	 */
-	GEOMETRYCORE_API bool Solve(int32 NumPoints, TFunctionRef<TVector<RealType>(int32)> GetPointFunc, TFunctionRef<bool(int32)> FilterFunc = [](int32 Idx) {return true;});
+	GEOMETRYCORE_API bool Solve(int32 NumPoints, TFunctionRef<TVector<RealType>(int32)> GetPointFunc, TFunctionRef<bool(int32)> FilterFunc);
+	GEOMETRYCORE_API bool Solve(int32 NumPoints, TFunctionRef<TVector<RealType>(int32)> GetPointFunc)
+	{
+		return Solve(NumPoints, GetPointFunc, [](int32 Idx) {return true;});
+	}
 
 	/**
 	 * Generate convex hull as long as input is not degenerate
