PR #10364: Migrate to new Distance method

diff --git a/Engine/Source/Runtime/GeometryCore/Public/Polyline3.h b/Engine/Source/Runtime/GeometryCore/Public/Polyline3.h
--- a/Engine/Source/Runtime/GeometryCore/Public/Polyline3.h
+++ b/Engine/Source/Runtime/GeometryCore/Public/Polyline3.h
@@ -264,7 +264,7 @@ public:
 		int N = SegmentCount();
 		for (int i = 0; i < N; ++i)
 		{
-			length += Vertices[i].Distance(Vertices[i+1]);
+			length += Distance(Vertices[i], Vertices[i+1]);
 		}
 		return length;
 	}
@@ -402,7 +402,7 @@ public:
 	{
 		T avg = 0; int N = Vertices.Num();
 		for (int i = 1; i < N; ++i) {
-			avg += Vertices[i].Distance(Vertices[i - 1]);
+			avg += Distance(Vertices[i], Vertices[i - 1]);
 		}
 		return avg / (T)(N-1);
 	}
@@ -437,4 +437,4 @@ typedef TPolyline3<double> FPolyline3d;
 typedef TPolyline3<float> FPolyline3f;
 
 } // end namespace UE::Geometry
-} // end namespace UE
\ No newline at end of file
+} // end namespace UE
