Move TConvexHull3 simplification-related settings to a separate struct, to make it easier to manage these settings (inside the convex hull code and also for other classes that call into it, like convex decomposition)

Note: These settings were all introduced post-5.3, so moving them should not require deprecation

diff --git a/Engine/Source/Runtime/Experimental/Chaos/Private/Chaos/CollisionConvexMesh.cpp b/Engine/Source/Runtime/Experimental/Chaos/Private/Chaos/CollisionConvexMesh.cpp
--- a/Engine/Source/Runtime/Experimental/Chaos/Private/Chaos/CollisionConvexMesh.cpp
+++ b/Engine/Source/Runtime/Experimental/Chaos/Private/Chaos/CollisionConvexMesh.cpp
@@ -50,7 +50,7 @@ namespace Chaos
 			if (bUseSimplifier)
 			{
 				HullCompute.bSaveTriangleNeighbors = true;
-				HullCompute.DegenerateEdgeTolerance = (FRealType)UE_DOUBLE_SMALL_NUMBER;
+				HullCompute.SimplificationSettings.DegenerateEdgeTolerance = (FRealType)UE_DOUBLE_SMALL_NUMBER;
 			}
 			if (HullCompute.Solve<FVec3Type>(InVertices))
 			{
@@ -176,7 +176,7 @@ namespace Chaos
 				HullCompute.bSaveTriangleNeighbors = true;
 				if (bUseSimplifier)
 				{
-					HullCompute.DegenerateEdgeTolerance = (FRealType)UE_DOUBLE_SMALL_NUMBER;
+					HullCompute.SimplificationSettings.DegenerateEdgeTolerance = (FRealType)UE_DOUBLE_SMALL_NUMBER;
 				}
 				if (HullCompute.Solve<FVec3Type>(*VerticesToUse))
 				{
diff --git a/Engine/Source/Runtime/GeometryCore/Private/CompGeom/ConvexHull3.cpp b/Engine/Source/Runtime/GeometryCore/Private/CompGeom/ConvexHull3.cpp
--- a/Engine/Source/Runtime/GeometryCore/Private/CompGeom/ConvexHull3.cpp
+++ b/Engine/Source/Runtime/GeometryCore/Private/CompGeom/ConvexHull3.cpp
@@ -812,25 +812,26 @@ bool TConvexHull3<RealType>::Solve(int32 NumPoints, TFunctionRef<TVector<RealTyp
 	Hull.Add(FIndex3i(InitialTet.Extreme[0], InitialTet.Extreme[2], InitialTet.Extreme[1]));
 
 	NumHullPoints = 4;
-	bool bHasPointBudget = MaxHullVertices > 0;
+	bool bHasPointBudget = SimplificationSettings.MaxHullVertices > 0;
+	int32 UseMaxHullVertices = SimplificationSettings.MaxHullVertices;
 	if (!bHasPointBudget)
 	{
-		MaxHullVertices = NumPoints;
+		UseMaxHullVertices = NumPoints;
 	}
 
-	RealType DegenerateEdgeToleranceSq = DegenerateEdgeTolerance * DegenerateEdgeTolerance;
+	RealType DegenerateEdgeToleranceSq = SimplificationSettings.DegenerateEdgeTolerance * SimplificationSettings.DegenerateEdgeTolerance;
 
 	FHullConnectivity<RealType> Connectivity;
-	double UseSkipAtHullDistance = SkipAtHullDistanceAbsolute;
-	if (SkipAtHullDistanceAsFraction > 0)
+	double UseSkipAtHullDistance = SimplificationSettings.SkipAtHullDistanceAbsolute;
+	if (SimplificationSettings.SkipAtHullDistanceAsFraction > 0)
 	{
 		double Extent = (GetPointFunc(InitialTet.Extreme[0]) - GetPointFunc(InitialTet.Extreme[1])).Length();
-		UseSkipAtHullDistance = FMathd::Max(UseSkipAtHullDistance, Extent * SkipAtHullDistanceAsFraction);
+		UseSkipAtHullDistance = FMathd::Max(UseSkipAtHullDistance, Extent * SimplificationSettings.SkipAtHullDistanceAsFraction);
 	}
 	Connectivity.VisibleDistanceThreshold = UseSkipAtHullDistance;
 	Connectivity.BuildNeighbors(Hull);
 	Connectivity.InitVisibility(Hull, NumPoints, GetPointFunc, FilterFunc);
-	while (NumHullPoints < MaxHullVertices)
+	while (NumHullPoints < UseMaxHullVertices)
 	{
 		if (Progress && (NumHullPoints % 100) == 0 && Progress->Cancelled())
 		{
@@ -1398,7 +1399,7 @@ void TConvexHull3<RealType>::GetSimplifiedFaces(TArray<FPolygonFace>& OutPolygon
 		}
 		TConvexHull3<RealType> FallbackHull;
 		FallbackHull.bSaveTriangleNeighbors = true;
-		FallbackHull.DegenerateEdgeTolerance = DegenerateEdgeTolerance;
+		// Note: We do not need the FallbackHull to use the original simplification settings, because the filter will guarantee we only use a subset of on-hull vertices in this pass
 
 		bool bFallbackSolveSuccess = FallbackHull.Solve(MaxVertexNum, GetPointFunc, VertexTouchesAtLeastThreeGroups);
 		if (bFallbackSolveSuccess)
diff --git a/Engine/Source/Runtime/GeometryCore/Public/CompGeom/ConvexHull3.h b/Engine/Source/Runtime/GeometryCore/Public/CompGeom/ConvexHull3.h
--- a/Engine/Source/Runtime/GeometryCore/Public/CompGeom/ConvexHull3.h
+++ b/Engine/Source/Runtime/GeometryCore/Public/CompGeom/ConvexHull3.h
@@ -52,18 +52,9 @@ private:
 	GEOMETRYCORE_API void Init(int32 NumPoints, TFunctionRef<TVector<RealType>(int32)> GetPointFunc, TFunctionRef<bool(int32)> FilterFunc, RealType Epsilon = TMathUtil<RealType>::Epsilon);
 };
 
-/**
- * Calculate the Convex Hull of a 3D point set as a Triangle Mesh
- */
 template<typename RealType>
-class TConvexHull3
+struct TConvexHullSimplificationSettings
 {
-public:
-
-	/// Whether neighbors for the hull triangles should be computed/saved.
-	/// If true, can call GetTriangleNeighbors() after Solve().
-	bool bSaveTriangleNeighbors = false;
-
 	/// Points will not be added to the hull if doing so would create an edge smaller than this distance tolerance.
 	/// If zero, no points will be skipped.
 	RealType DegenerateEdgeTolerance = (RealType)0;
@@ -76,6 +67,22 @@ public:
 	/// If positive, skip adding points that are closer than this threshold to the in-progress hull -- expressed as a fraction of the overall hull extent.
 	/// Note if both the Absolute and Fraction MinPlaneDistance thresholds are set, they will both apply (i.e., the larger of the two will be used)
 	RealType SkipAtHullDistanceAsFraction = -TMathUtil<RealType>::MaxReal;
+};
+
+/**
+ * Calculate the Convex Hull of a 3D point set as a Triangle Mesh
+ */
+template<typename RealType>
+class TConvexHull3
+{
+public:
+
+	/// Whether neighbors for the hull triangles should be computed/saved.
+	/// If true, can call GetTriangleNeighbors() after Solve().
+	bool bSaveTriangleNeighbors = false;
+
+	// Settings controlling whether and how to generate a simpler hull
+	TConvexHullSimplificationSettings<RealType> SimplificationSettings;
 
 	/**
 	 * Generate convex hull as long as input is not degenerate
