Add a new k-DOP fitting algorithm to GeometryCore and swap it in to GeomFitUtils so it will be used by the Static Mesh Editor.

diff --git a/Engine/Source/Editor/UnrealEd/Private/GeomFitUtils.cpp b/Engine/Source/Editor/UnrealEd/Private/GeomFitUtils.cpp
--- a/Engine/Source/Editor/UnrealEd/Private/GeomFitUtils.cpp
+++ b/Engine/Source/Editor/UnrealEd/Private/GeomFitUtils.cpp
@@ -23,6 +23,7 @@
 #include "MeshDescription.h"
 #include "StaticMeshAttributes.h"
 #include "Settings/EditorExperimentalSettings.h"
+#include "CompGeom/FitKDOP3.h"
 
 static bool PromptToRemoveExistingCollision(UStaticMesh* StaticMesh)
 {
@@ -55,9 +56,6 @@ static bool PromptToRemoveExistingCollision(UStaticMesh* StaticMesh)
 // This function takes the current collision model, and fits a k-DOP around it.
 // It uses the array of k unit-length direction vectors to define the k bounding planes.
 
-// THIS FUNCTION REPLACES EXISTING SIMPLE COLLISION MODEL WITH KDOP
-#define MY_FLTMAX (3.402823466e+38F)
-
 int32 GenerateKDopAsSimpleCollision(UStaticMesh* StaticMesh, const TArray<FVector> &Dirs)
 {
 	// Make sure rendering is done - so we are not changing data being used by collision drawing.
@@ -70,97 +68,22 @@ int32 GenerateKDopAsSimpleCollision(UStaticMesh* StaticMesh, const TArray<FVecto
 
 	UBodySetup* bs = StaticMesh->GetBodySetup();
 
-	// Do k- specific stuff.
-	int32 kCount = Dirs.Num();
-	TArray<float> maxDist;
-	for(int32 i=0; i<kCount; i++)
-		maxDist.Add(-MY_FLTMAX);
-
-	// Construct temporary UModel for kdop creation. We keep no refs to it, so it can be GC'd.
-	auto TempModel = NewObject<UModel>();
-	TempModel->Initialize(nullptr, 1);
-
-	// For each vertex, project along each kdop direction, to find the max in that direction.
 	const FStaticMeshLODResources& RenderData = StaticMesh->GetRenderData()->LODResources[0];
-	for(int32 i=0; i<RenderData.GetNumVertices(); i++)
-	{
-		for(int32 j=0; j<kCount; j++)
-		{
-			float dist = RenderData.VertexBuffers.PositionVertexBuffer.VertexPosition(i) | (FVector3f)Dirs[j];
-			maxDist[j] = FMath::Max(dist, maxDist[j]);
-		}
-	}
-
-	// Inflate kdop to ensure it is no degenerate
-	const float MinSize = 0.1f;
-	for(int32 i=0; i<kCount; i++)
-	{
-		maxDist[i] += MinSize;
-	}
+	TArray<FVector> HullVertices;
+	UE::Geometry::FitKDOPVertices3<double>(Dirs, RenderData.GetNumVertices(),
+		[&](int32 VertIdx) { return (FVector)RenderData.VertexBuffers.PositionVertexBuffer.VertexPosition(VertIdx); }, HullVertices);
 
-	// Now we have the planes of the kdop, we work out the face polygons.
-	TArray<FPlane4f> planes;
-	for(int32 i=0; i<kCount; i++)
-		planes.Add( FPlane4f((FVector3f)Dirs[i], maxDist[i]) );
-
-	for(int32 i=0; i<planes.Num(); i++)
-	{
-		FPoly*	Polygon = new(TempModel->Polys->Element) FPoly();
-		FVector3f Base, AxisX, AxisY;
-
-		Polygon->Init();
-		Polygon->Normal = planes[i];
-		Polygon->Normal.FindBestAxisVectors(AxisX,AxisY);
-
-		Base = planes[i] * planes[i].W;
-
-		new(Polygon->Vertices) FVector3f(Base + AxisX * UE_OLD_HALF_WORLD_MAX + AxisY * UE_OLD_HALF_WORLD_MAX);
-		new(Polygon->Vertices) FVector3f(Base + AxisX * UE_OLD_HALF_WORLD_MAX - AxisY * UE_OLD_HALF_WORLD_MAX);
-		new(Polygon->Vertices) FVector3f(Base - AxisX * UE_OLD_HALF_WORLD_MAX - AxisY * UE_OLD_HALF_WORLD_MAX);
-		new(Polygon->Vertices) FVector3f(Base - AxisX * UE_OLD_HALF_WORLD_MAX + AxisY * UE_OLD_HALF_WORLD_MAX);
-
-		for(int32 j=0; j<planes.Num(); j++)
-		{
-			if(i != j)
-			{
-				if(!Polygon->Split(-FVector3f(planes[j]), planes[j] * planes[j].W))
-				{
-					Polygon->Vertices.Empty();
-					break;
-				}
-			}
-		}
-
-		if(Polygon->Vertices.Num() < 3)
-		{
-			// If poly resulted in no verts, remove from array
-			TempModel->Polys->Element.RemoveAt(TempModel->Polys->Element.Num()-1);
-		}
-		else
-		{
-			// Other stuff...
-			Polygon->iLink = i;
-			Polygon->CalcNormal(1);
-		}
-	}
-
-	if(TempModel->Polys->Element.Num() < 4)
-	{
-		TempModel = NULL;
-		return INDEX_NONE;
-	}
-
-	// Build bounding box.
-	TempModel->BuildBound();
+	bs->Modify();
 
-	// Build BSP for the brush.
-	FBSPOps::bspBuild(TempModel,FBSPOps::BSP_Good,15,70,1,0);
-	FBSPOps::bspRefresh(TempModel,1);
-	FBSPOps::bspBuildBounds(TempModel);
+	// Create new GUID
+	bs->InvalidatePhysicsData();
 
-	bs->Modify();
+	FKConvexElem ConvexElem;
+	ConvexElem.VertexData = HullVertices;
+	// Note: UpdateElemBox also computes the convex hull indices
+	ConvexElem.UpdateElemBox();
 
-	bs->CreateFromModel(TempModel, false);
+	bs->AggGeom.ConvexElems.Add(ConvexElem);
 	
 	// create all body instances
 	RefreshCollisionChange(*StaticMesh);
@@ -168,6 +91,8 @@ int32 GenerateKDopAsSimpleCollision(UStaticMesh* StaticMesh, const TArray<FVecto
 	// Mark staticmesh as dirty, to help make sure it gets saved.
 	StaticMesh->MarkPackageDirty();
 
+	StaticMesh->bCustomizedCollision = true;	//mark the static mesh for collision customization
+
 	return bs->AggGeom.ConvexElems.Num() - 1;
 }
 
diff --git a/Engine/Source/Editor/UnrealEd/UnrealEd.Build.cs b/Engine/Source/Editor/UnrealEd/UnrealEd.Build.cs
--- a/Engine/Source/Editor/UnrealEd/UnrealEd.Build.cs
+++ b/Engine/Source/Editor/UnrealEd/UnrealEd.Build.cs
@@ -182,6 +182,7 @@ public class UnrealEd : ModuleRules
 				"HeadMountedDisplay",
 				"FieldNotification",
 				"VirtualizationEditor",
+				"GeometryCore",
 			}
 		);
 
diff --git a/Engine/Source/Runtime/GeometryCore/Private/CompGeom/FitKDOP3.cpp b/Engine/Source/Runtime/GeometryCore/Private/CompGeom/FitKDOP3.cpp
new file mode 100644
--- /dev/null
+++ b/Engine/Source/Runtime/GeometryCore/Private/CompGeom/FitKDOP3.cpp
@@ -0,0 +1,150 @@
+// Copyright Epic Games, Inc. All Rights Reserved.
+
+#include "CompGeom/FitKDOP3.h"
+#include "BoxTypes.h"
+#include "Curve/CurveUtil.h"
+#include "VectorTypes.h"
+#include "VectorUtil.h"
+#include "Templates/Function.h"
+
+#include "CompGeom/ConvexHull3.h"
+
+namespace UE::Private::FitKDop3Helpers
+{
+	using namespace UE::Geometry;
+
+	template<typename RealType>
+	bool FitKDopImpl3(TArrayView<const UE::Math::TVector<RealType>> PlaneDirections,
+		const int32 NumPoints, TFunctionRef<UE::Math::TVector<RealType>(int32)> GetPointFunc, TFunctionRef<bool(int32)> FilterFunc,
+		TArray<UE::Math::TVector<RealType>>& OutVertices, TArray<UE::Math::TPlane<RealType>>* OptionalOutPlanes, RealType Epsilon, RealType VertexSnapDistance)
+	{
+		using namespace UE::Math;
+
+		OutVertices.Reset();
+
+		if (PlaneDirections.Num() < 4)
+		{
+			return false;
+		}
+
+		const int32 NumPlanes = PlaneDirections.Num();
+		TAxisAlignedBox3<RealType> Bounds;
+
+		TArray<RealType> PlaneDists;
+		PlaneDists.Init(-TMathUtil<RealType>::MaxReal, NumPlanes);
+		for (int32 PointIdx = 0; PointIdx < NumPoints; ++PointIdx)
+		{
+			TVector<RealType> Point = GetPointFunc(PointIdx);
+			Bounds.Contain(Point);
+			for (int32 DirIdx = 0; DirIdx < NumPlanes; ++DirIdx)
+			{
+				RealType Dist = Point.Dot(PlaneDirections[DirIdx]);
+				PlaneDists[DirIdx] = FMath::Max(PlaneDists[DirIdx], Dist);
+			}
+		}
+
+		TArray<TPlane<RealType>> LocalPlanes;
+		TArray<TPlane<RealType>>* UsePlanes = OptionalOutPlanes ? OptionalOutPlanes : &LocalPlanes;
+		UsePlanes->SetNumUninitialized(PlaneDists.Num());
+		for (int32 DirIdx = 0; DirIdx < NumPlanes; ++DirIdx)
+		{
+			(*UsePlanes)[DirIdx] = TPlane<RealType>(PlaneDirections[DirIdx], PlaneDists[DirIdx]);
+			if (!ensureMsgf((*UsePlanes)[DirIdx].Normalize(), TEXT("k-DOP fitting attempted with invalid (near-zero) search direction(s)")))
+			{
+				return false;
+			}
+		}
+
+		TVector<RealType> Center = Bounds.Center();
+		RealType BoundingSize = Bounds.DiagonalLength();
+		RealType VertexSnapDistanceSq = VertexSnapDistance * VertexSnapDistance;
+		TArray<TVector<RealType>> PlanePoly, ClippedPlanePoly; // Note: Re-used per clip plane
+		for (int32 WorkPlaneIdx = 0; WorkPlaneIdx < NumPlanes; ++WorkPlaneIdx)
+		{
+			const TPlane<RealType>& WorkPlane = (*UsePlanes)[WorkPlaneIdx];
+			TVector<RealType> AxisX, AxisY;
+			VectorUtil::MakePerpVectors(WorkPlane.GetNormal(), AxisX, AxisY);
+
+			TVector<RealType> CenterOnPlane = TVector<RealType>::PointPlaneProject(Center, WorkPlane);
+
+			// create an initial polygon bounding the possible hull face
+			PlanePoly.Reset();
+			PlanePoly.Add(CenterOnPlane + AxisX * BoundingSize + AxisY * BoundingSize);
+			PlanePoly.Add(CenterOnPlane - AxisX * BoundingSize + AxisY * BoundingSize);
+			PlanePoly.Add(CenterOnPlane - AxisX * BoundingSize - AxisY * BoundingSize);
+			PlanePoly.Add(CenterOnPlane + AxisX * BoundingSize - AxisY * BoundingSize);
+
+			// clip the polygon vs all other planes
+			for (int32 CutPlaneIdx = 0; CutPlaneIdx < NumPlanes; ++CutPlaneIdx)
+			{
+				if (CutPlaneIdx == WorkPlaneIdx)
+				{
+					continue;
+				}
+
+				const TPlane<RealType>& CutPlane = (*UsePlanes)[CutPlaneIdx];
+				if (CurveUtil::ClipConvexToPlane<RealType, TVector<RealType>, true>(PlanePoly, CutPlane, ClippedPlanePoly))
+				{
+					Swap(PlanePoly, ClippedPlanePoly);
+					ClippedPlanePoly.Reset();
+				}
+				if (PlanePoly.Num() < 3)
+				{
+					break;
+				}
+			}
+
+			for (int32 PolyIdx = 0; PolyIdx < PlanePoly.Num(); ++PolyIdx)
+			{
+				bool bSnapped = false;
+				for (int32 OutIdx = 0; !bSnapped && OutIdx < OutVertices.Num(); ++OutIdx)
+				{
+					bSnapped = TVector<RealType>::DistSquared(PlanePoly[PolyIdx], OutVertices[OutIdx]) < VertexSnapDistanceSq;
+				}
+				if (!bSnapped)
+				{
+					OutVertices.Add(PlanePoly[PolyIdx]);
+				}
+			}
+		}
+
+		if (OutVertices.Num() < 4)
+		{
+			return false;
+		}
+
+		TExtremePoints3<RealType> ExtremePoints(OutVertices.Num(), [&OutVertices](int32 Idx) {return OutVertices[Idx];}, [](int32) {return true;}, Epsilon);
+		if (ExtremePoints.Dimension < 3)
+		{
+			return false;
+		}
+
+		return true;
+	}
+} // namespace UE::Private::FitKDop3Helpers
+
+namespace UE::Geometry
+{
+	template<typename RealType>
+	bool FitKDOPVertices3(
+		TArrayView<const UE::Math::TVector<RealType>> PlaneDirections,
+		const int32 NumPoints, TFunctionRef<UE::Math::TVector<RealType>(int32)> GetPointFunc, TFunctionRef<bool(int32)> FilterFunc,
+		TArray<UE::Math::TVector<RealType>>& OutVertices, TArray<UE::Math::TPlane<RealType>>* OptionalOutPlanes, RealType Epsilon, RealType VertexSnapDistance)
+	{
+		return UE::Private::FitKDop3Helpers::FitKDopImpl3(PlaneDirections,
+			NumPoints, GetPointFunc, FilterFunc,
+			OutVertices, OptionalOutPlanes,
+			Epsilon, VertexSnapDistance);
+	}
+
+
+	// explicit instantiations
+	template bool GEOMETRYCORE_API FitKDOPVertices3<float>(
+		TArrayView<const UE::Math::TVector<float>> PlaneDirections, const int32 NumPoints, TFunctionRef<UE::Math::TVector<float>(int32)> GetPointFunc, TFunctionRef<bool(int32)> FilterFunc,
+		TArray<UE::Math::TVector<float>>& OutVertices, TArray<UE::Math::TPlane<float>>* OptionalOutPlanes, float Epsilon, float VertexSnapDistance);
+	template bool GEOMETRYCORE_API FitKDOPVertices3<double>(
+		TArrayView<const UE::Math::TVector<double>> PlaneDirections, const int32 NumPoints, TFunctionRef<UE::Math::TVector<double>(int32)> GetPointFunc, TFunctionRef<bool(int32)> FilterFunc,
+		TArray<UE::Math::TVector<double>>& OutVertices, TArray<UE::Math::TPlane<double>>* OptionalOutPlanes, double Epsilon, double VertexSnapDistance);
+	
+
+} /// namespace UE::Geometry
diff --git a/Engine/Source/Runtime/GeometryCore/Public/CompGeom/FitKDOP3.h b/Engine/Source/Runtime/GeometryCore/Public/CompGeom/FitKDOP3.h
new file mode 100644
--- /dev/null
+++ b/Engine/Source/Runtime/GeometryCore/Public/CompGeom/FitKDOP3.h
@@ -0,0 +1,60 @@
+// Copyright Epic Games, Inc. All Rights Reserved.
+
+#pragma once
+
+#include "Containers/Array.h"
+#include "Containers/ArrayView.h"
+#include "MathUtil.h"
+#include "Math/Vector.h"
+#include "Math/Plane.h"
+#include "VectorTypes.h"
+
+template <typename FuncType> class TFunctionRef;
+
+namespace UE
+{
+namespace Math { template <typename T> struct TVector; }
+
+namespace Geometry
+{
+
+
+	/**
+	 * Compute the vertices (and optionally planes) of a k-DOP bounding convex hull containing the points returned by GetPointFunc where FilterFunc returns true
+	 */
+	template <typename RealType>
+	bool GEOMETRYCORE_API FitKDOPVertices3(
+		TArrayView<const UE::Math::TVector<RealType>> PlaneDirections,
+		const int32 NumPoints, TFunctionRef<UE::Math::TVector<RealType>(int32)> GetPointFunc, TFunctionRef<bool(int32)> FilterFunc,
+		TArray<UE::Math::TVector<RealType>>& OutVertices, TArray<UE::Math::TPlane<RealType>>* OptionalOutPlanes = nullptr,
+		RealType Epsilon = TMathUtil<RealType>::Epsilon, RealType VertexSnapDistance = (RealType)0.1);
+
+	/**
+	 * Compute the vertices (and optionally planes) of a k-DOP bounding convex hull containing the points returned by GetPointFunc
+	 */
+	template <typename RealType>
+	bool FitKDOPVertices3(
+		TArrayView<const UE::Math::TVector<RealType>> PlaneDirections,
+		const int32 NumPoints, TFunctionRef<UE::Math::TVector<RealType>(int32)> GetPointFunc,
+		TArray<UE::Math::TVector<RealType>>& OutVertices, TArray<UE::Math::TPlane<RealType>>* OptionalOutPlanes = nullptr,
+		RealType Epsilon = TMathUtil<RealType>::Epsilon, RealType VertexSnapDistance = (RealType)0.1)
+	{
+		return FitKDOPVertices3<RealType>(PlaneDirections, NumPoints, GetPointFunc, [](int32)->bool {return true;}, OutVertices, OptionalOutPlanes, Epsilon);
+	}
+
+	/**
+	 * Compute the vertices (and optionally planes) of a k-DOP bounding convex hull containing Points
+	 */
+	template <typename RealType>
+	bool FitKDOPVertices3(
+		TArrayView<const UE::Math::TVector<RealType>> PlaneDirections,
+		const TArray<UE::Math::TVector<RealType>>& Points,
+		TArray<UE::Math::TVector<RealType>>& OutVertices, TArray<UE::Math::TPlane<RealType>>* OptionalOutPlanes = nullptr,
+		RealType Epsilon = TMathUtil<RealType>::Epsilon, RealType VertexSnapDistance = (RealType)0.1)
+	{
+		return FitKDOPVertices3<RealType>(PlaneDirections, Points.Num(), [&Points](int32 Idx) {return Points[Idx];}, [](int32)->bool {return true;}, OutVertices, OptionalOutPlanes, Epsilon);
+	}
+
+} // end namespace Geometry
+}// end namespace UE
+
diff --git a/Engine/Source/Runtime/GeometryCore/Public/Curve/CurveUtil.h b/Engine/Source/Runtime/GeometryCore/Public/Curve/CurveUtil.h
--- a/Engine/Source/Runtime/GeometryCore/Public/Curve/CurveUtil.h
+++ b/Engine/Source/Runtime/GeometryCore/Public/Curve/CurveUtil.h
@@ -298,6 +298,75 @@ using namespace UE::Math;
 		}
 	}
 
+	/**
+	 * Use the Sutherland–Hodgman algorithm to clip the vertices to the given plane
+	 * Note if the path/polygon is concave, this may leave overlapping edges at the boundary.
+	 * @tparam bLoop				Whether the path is a loop
+	 * @param Vertices				Vertices to clip vs the plane
+	 * @param Plane					Plane to use for clipping
+	 * @param OutClipped			Vertices clipped to the plane
+	 * @param bKeepPositiveSide		Whether to keep the portion of the polygon that is on the positive side of the plane. Otherwise, keeps the negative side.
+	 * @return Whether vertices were clipped
+	 */
+	template<typename RealType, typename VectorType, bool bLoop = true>
+	static bool ClipConvexToPlane(TArray<VectorType>& Vertices, const TPlane<RealType>& Plane, TArray<VectorType>& OutClipped, bool bKeepPositiveSide = false)
+	{
+		OutClipped.Reset(Vertices.Num());
+		bool bWasClipped = false;
+		int32 VertNum = Vertices.Num();
+		int32 StartCur, StartPrev;
+		if constexpr (bLoop)
+		{
+			StartCur = 0;
+			StartPrev = VertNum - 1;
+		}
+		else
+		{
+			StartCur = 1;
+			StartPrev = 0;
+		}
+		RealType SideSign = bKeepPositiveSide ? (RealType)1 : (RealType)-1;
+		RealType PrevDist = Plane.PlaneDot(Vertices[StartPrev]) * RealType(SideSign);
+		if constexpr (!bLoop)
+		{
+			if (PrevDist >= 0)
+			{
+				OutClipped.Add(Vertices[0]);
+			}
+			else
+			{
+				bWasClipped = true;
+			}
+		}
+		for (int32 CurIdx = StartCur, PrevIdx = StartPrev; CurIdx < VertNum; PrevIdx = CurIdx++)
+		{
+			RealType CurDist = Plane.PlaneDot(Vertices[CurIdx]) * RealType(SideSign);
+			if (CurDist >= 0)
+			{
+				if (PrevDist < 0 && CurDist > 0)
+				{
+					RealType T = CurDist / (CurDist - PrevDist);
+					VectorType LerpVec = Lerp(Vertices[CurIdx], Vertices[PrevIdx], T);
+					OutClipped.Add(LerpVec);
+					bWasClipped = true;
+				}
+				OutClipped.Add(Vertices[CurIdx]);
+			}
+			else
+			{
+				bWasClipped = true;
+				if (PrevDist > 0)
+				{
+					RealType T = CurDist / (CurDist - PrevDist);
+					VectorType LerpVec = Lerp(Vertices[CurIdx], Vertices[PrevIdx], T);
+					OutClipped.Add(LerpVec);
+				}
+			}
+			PrevDist = CurDist;
+		}
+		return bWasClipped;
+	}
+
 	/**
 	 * Tests closed, 2D curve for convexity, with an optional tolerance allowing for approximately-collinear points
 	 * Note that tolerance is per vertex angle, so a very-finely-sampled smooth concavity could be reported as convex
