Updated Clip and DoClipping to be forceinlined. This saved roughly 100ms from calling functions on average. Overall it saved 8secs on average when generating from **Model assets.

diff --git a/Engine/Source/Runtime/GeometryCore/Public/Intersection/IntrRay3AxisAlignedBox3.h b/Engine/Source/Runtime/GeometryCore/Public/Intersection/IntrRay3AxisAlignedBox3.h
--- a/Engine/Source/Runtime/GeometryCore/Public/Intersection/IntrRay3AxisAlignedBox3.h
+++ b/Engine/Source/Runtime/GeometryCore/Public/Intersection/IntrRay3AxisAlignedBox3.h
@@ -137,7 +137,7 @@ protected:
 
 	// internal functions
 
-	static bool DoClipping(RealType& t0, RealType& t1,
+	static FORCEINLINE bool DoClipping(RealType& t0, RealType& t1,
 		const TVector<RealType>& RayOrigin, const TVector<RealType>& RayDirection,
 		const TAxisAlignedBox3<RealType>& Box, bool solid, 
 		int& quantity, TVector<RealType>& Point0, TVector<RealType>& Point1, EIntersectionType& intrType)
@@ -179,10 +179,7 @@ protected:
 		return intrType != EIntersectionType::Empty;
 	}
 
-
-
-
-	static bool Clip(RealType denom, RealType numer, RealType& t0, RealType& t1)
+	static FORCEINLINE bool Clip(RealType denom, RealType numer, RealType& t0, RealType& t1)
 	{
 		// Return value is 'true' if line segment intersects the current test
 		// plane.  Otherwise 'false' is returned in which case the line segment
