add cbrt (cube root) function to mathutil

diff --git a/Engine/Source/Runtime/GeometryCore/Public/MathUtil.h b/Engine/Source/Runtime/GeometryCore/Public/MathUtil.h
--- a/Engine/Source/Runtime/GeometryCore/Public/MathUtil.h
+++ b/Engine/Source/Runtime/GeometryCore/Public/MathUtil.h
@@ -156,6 +156,8 @@ public:
 	/** compute min and max of a,b,c with max 3 comparisons (sometimes 2) */
 	static inline void MinMax(RealType A, RealType B, RealType C, RealType& MinOut, RealType& MaxOut);
 	static inline RealType Sqrt(const RealType Value);
+	/** Compute cube root */
+	static inline RealType Cbrt(const RealType Value);
 	static inline RealType Tan(const RealType Value);
 	static inline RealType Atan2(const RealType ValueY, const RealType ValueX);
 	static inline RealType Sin(const RealType Value);
@@ -305,6 +307,12 @@ RealType TMathUtil<RealType>::Sqrt(const RealType Value)
 	return sqrt(Value);
 }
 
+template<typename RealType>
+RealType TMathUtil<RealType>::Cbrt(const RealType Value)
+{
+	return cbrt(Value);
+}
+
 template<typename RealType>
 RealType TMathUtil<RealType>::Tan(const RealType Value)
 {
