Add a normal-based clustering algorithm which groups connected triangle groups based on the current average normal of the group.

Use this algorithm in the polygroups generator tool, so we get more useful groups on smooth surfaces (like a cylinder)

Also add a helper to the FSizedDisjointSet class to easily get the number of groups in the set and a mapping between compact group indices and sparse group IDs

diff --git a/Engine/Plugins/Experimental/MeshModelingToolsetExp/Source/MeshModelingToolsExp/Private/ConvertToPolygonsTool.cpp b/Engine/Plugins/Experimental/MeshModelingToolsetExp/Source/MeshModelingToolsExp/Private/ConvertToPolygonsTool.cpp
--- a/Engine/Plugins/Experimental/MeshModelingToolsetExp/Source/MeshModelingToolsExp/Private/ConvertToPolygonsTool.cpp
+++ b/Engine/Plugins/Experimental/MeshModelingToolsetExp/Source/MeshModelingToolsExp/Private/ConvertToPolygonsTool.cpp
@@ -56,6 +56,7 @@ public:
 	// parameters set by the tool
 	EConvertToPolygonsMode ConversionMode = EConvertToPolygonsMode::FaceNormalDeviation;
 	double AngleTolerance = 0.1f;
+	bool bUseAverageGroupNormal = true;
 	int32 NumPoints = 10;
 	bool bSubdivideExisting = false;
 	FPolygroupsGenerator::EWeightingType WeightingType = FPolygroupsGenerator::EWeightingType::None;
@@ -135,7 +136,8 @@ public:
 			Generator.FindPolygroupsFromFaceNormals(
 				DotTolerance,
 				bRespectUVSeams, 
-				bRespectHardNormals );
+				bRespectHardNormals,
+				bUseAverageGroupNormal);
 			break;
 		}
 		case EConvertToPolygonsMode::FindPolygons:
@@ -324,6 +326,7 @@ void UConvertToPolygonsTool::Setup()
 	Settings->WatchProperty(Settings->ConversionMode, [this](EConvertToPolygonsMode) { OnSettingsModified(); });
 	Settings->WatchProperty(Settings->bShowGroupColors, [this](bool) { UpdateVisualization(); });
 	Settings->WatchProperty(Settings->AngleTolerance, [this](float) { OnSettingsModified(); });
+	Settings->WatchProperty(Settings->bUseAverageGroupNormal, [this](float) { OnSettingsModified(); });
 	Settings->WatchProperty(Settings->NumPoints, [this](int32) { OnSettingsModified(); });
 	Settings->WatchProperty(Settings->bSplitExisting, [this](bool) { OnSettingsModified(); });
 	Settings->WatchProperty(Settings->bNormalWeighted, [this](bool) { OnSettingsModified(); });
@@ -397,6 +400,7 @@ void UConvertToPolygonsTool::UpdateOpParameters(FConvertToPolygonsOp& ConvertToP
 	ConvertToPolygonsOp.bCalculateNormals = Settings->bCalculateNormals;
 	ConvertToPolygonsOp.ConversionMode    = Settings->ConversionMode;
 	ConvertToPolygonsOp.AngleTolerance    = Settings->AngleTolerance;
+	ConvertToPolygonsOp.bUseAverageGroupNormal = Settings->bUseAverageGroupNormal;
 	ConvertToPolygonsOp.NumPoints = Settings->NumPoints;
 	ConvertToPolygonsOp.bSubdivideExisting = Settings->bSplitExisting;
 	ConvertToPolygonsOp.WeightingType = (Settings->bNormalWeighted) ? FPolygroupsGenerator::EWeightingType::NormalDeviation : FPolygroupsGenerator::EWeightingType::None;
diff --git a/Engine/Plugins/Experimental/MeshModelingToolsetExp/Source/MeshModelingToolsExp/Public/ConvertToPolygonsTool.h b/Engine/Plugins/Experimental/MeshModelingToolsetExp/Source/MeshModelingToolsExp/Public/ConvertToPolygonsTool.h
--- a/Engine/Plugins/Experimental/MeshModelingToolsetExp/Source/MeshModelingToolsExp/Public/ConvertToPolygonsTool.h
+++ b/Engine/Plugins/Experimental/MeshModelingToolsetExp/Source/MeshModelingToolsExp/Public/ConvertToPolygonsTool.h
@@ -72,6 +72,10 @@ public:
 	UPROPERTY(EditAnywhere, Category = NormalDeviation, meta = (UIMin = "0.001", UIMax = "60.0", ClampMin = "0.0", ClampMax = "90.0", EditCondition = "ConversionMode == EConvertToPolygonsMode::FaceNormalDeviation", EditConditionHides))
 	float AngleTolerance = 0.1f;
 
+	/** Whether to compute Face Normal Deviation based on the overall PolyGroup's average normal, or to only consider the normals of the individual triangles */
+	UPROPERTY(EditAnywhere, Category = NormalDeviation, meta = (EditCondition = "ConversionMode == EConvertToPolygonsMode::FaceNormalDeviation", EditConditionHides))
+	bool bUseAverageGroupNormal = true;
+
 	/** Furthest-Point Sample count, approximately this number of polygroups will be generated */
 	UPROPERTY(EditAnywhere, Category = FurthestPoint, meta = (UIMin = "1", UIMax = "100", ClampMin = "1", ClampMax = "10000", EditCondition = "ConversionMode == EConvertToPolygonsMode::FromFurthestPointSampling", EditConditionHides))
 	int32 NumPoints = 100;
diff --git a/Engine/Plugins/Runtime/GeometryProcessing/Source/DynamicMesh/Private/Polygroups/PolygroupsGenerator.cpp b/Engine/Plugins/Runtime/GeometryProcessing/Source/DynamicMesh/Private/Polygroups/PolygroupsGenerator.cpp
--- a/Engine/Plugins/Runtime/GeometryProcessing/Source/DynamicMesh/Private/Polygroups/PolygroupsGenerator.cpp
+++ b/Engine/Plugins/Runtime/GeometryProcessing/Source/DynamicMesh/Private/Polygroups/PolygroupsGenerator.cpp
@@ -1,6 +1,8 @@
 // Copyright Epic Games, Inc. All Rights Reserved.
 
 #include "Polygroups/PolygroupsGenerator.h"
+
+#include "Clustering/FaceNormalClustering.h"
 #include "DynamicMesh/MeshNormals.h"
 #include "Selections/MeshConnectedComponents.h"
 #include "Util/IndexUtil.h"
@@ -582,12 +584,11 @@ bool FPolygroupsGenerator::FindSourceMeshPolygonPolygroups(
 
 
 bool FPolygroupsGenerator::FindPolygroupsFromFaceNormals(
-	double DotTolerance,
+	double OneMinusCosAngleTolerance,
 	bool bRespectUVSeams,
-	bool bRespectNormalSeams)
+	bool bRespectNormalSeams,
+	bool bUseAveragePolygroupNormals)
 {
-	DotTolerance = 1.0 - DotTolerance;
-
 	// if we are respecting seams or hard normals, find all those edges
 	TSet<int32> InvalidEdges;
 	if (bRespectUVSeams || bRespectNormalSeams)
@@ -595,56 +596,69 @@ bool FPolygroupsGenerator::FindPolygroupsFromFaceNormals(
 		GetSeamConstraintEdges(bRespectUVSeams, bRespectNormalSeams, InvalidEdges);
 	}
 
-	// compute face normals
-	FMeshNormals Normals(Mesh);
-	Normals.ComputeTriangleNormals();
+	if (bUseAveragePolygroupNormals)
+	{
+		FaceNormalClustering::FClusterOptions Options;
+		Options.NormalOneMinusCosTolerance = OneMinusCosAngleTolerance;
+		Options.bApplyNormalToleranceToClusters = true;
+		FaceNormalClustering::ComputeMeshPolyGroupsFromClusters(*Mesh, FoundPolygroups, Options, &InvalidEdges);
+	}
+	else
+	{
+		// Transform into normal-dot-product tolerance
+		const double DotTolerance = 1.0 - OneMinusCosAngleTolerance;
 
-	TArray<bool> DoneTriangle;
-	DoneTriangle.SetNum(Mesh->MaxTriangleID());
+		// compute face normals
+		FMeshNormals Normals(Mesh);
+		Normals.ComputeTriangleNormals();
 
-	TArray<int> Stack;
+		TArray<bool> DoneTriangle;
+		DoneTriangle.SetNum(Mesh->MaxTriangleID());
 
-	// grow outward from vertices until we have no more left
-	for (int TriID : Mesh->TriangleIndicesItr())
-	{
-		if (DoneTriangle[TriID] == true)
+		TArray<int> Stack;
+
+		// grow outward from vertices until we have no more left
+		for (int TriID : Mesh->TriangleIndicesItr())
 		{
-			continue;
-		}
+			if (DoneTriangle[TriID] == true)
+			{
+				continue;
+			}
 
-		TArray<int> Polygroup;
-		Polygroup.Add(TriID);
-		DoneTriangle[TriID] = true;
+			TArray<int> Polygroup;
+			Polygroup.Add(TriID);
+			DoneTriangle[TriID] = true;
 
-		Stack.SetNum(0);
-		Stack.Add(TriID);
-		while (Stack.Num() > 0)
-		{
-			int CurTri = Stack.Pop(false);
-			FIndex3i NbrTris = Mesh->GetTriNeighbourTris(CurTri);
-			FIndex3i NbrEdges = Mesh->GetTriEdges(CurTri);
-			for (int j = 0; j < 3; ++j)
+			Stack.Reset();
+			Stack.Add(TriID);
+			while (Stack.Num() > 0)
 			{
-				if (InvalidEdges.Contains(NbrEdges[j]))
+				int CurTri = Stack.Pop(false);
+				FIndex3i NbrTris = Mesh->GetTriNeighbourTris(CurTri);
+				FIndex3i NbrEdges = Mesh->GetTriEdges(CurTri);
+				for (int j = 0; j < 3; ++j)
 				{
-					continue;
-				}
+					if (InvalidEdges.Contains(NbrEdges[j]))
+					{
+						continue;
+					}
 
-				if (NbrTris[j] >= 0
-					&& DoneTriangle[NbrTris[j]] == false)
-				{
-					double Dot = Normals[CurTri].Dot(Normals[NbrTris[j]]);
-					if (Dot > DotTolerance)
+					if (NbrTris[j] >= 0
+						&& DoneTriangle[NbrTris[j]] == false)
 					{
-						Polygroup.Add(NbrTris[j]);
-						Stack.Add(NbrTris[j]);
-						DoneTriangle[NbrTris[j]] = true;
+						double Dot = Normals[CurTri].Dot(Normals[NbrTris[j]]);
+						if (Dot > DotTolerance)
+						{
+							Polygroup.Add(NbrTris[j]);
+							Stack.Add(NbrTris[j]);
+							DoneTriangle[NbrTris[j]] = true;
+						}
 					}
 				}
 			}
-		}
 
-		FoundPolygroups.Add(Polygroup);
+			FoundPolygroups.Add(Polygroup);
+		}
 	}
 
 	if (bApplyPostProcessing)
diff --git a/Engine/Plugins/Runtime/GeometryProcessing/Source/DynamicMesh/Public/Polygroups/PolygroupsGenerator.h b/Engine/Plugins/Runtime/GeometryProcessing/Source/DynamicMesh/Public/Polygroups/PolygroupsGenerator.h
--- a/Engine/Plugins/Runtime/GeometryProcessing/Source/DynamicMesh/Public/Polygroups/PolygroupsGenerator.h
+++ b/Engine/Plugins/Runtime/GeometryProcessing/Source/DynamicMesh/Public/Polygroups/PolygroupsGenerator.h
@@ -57,12 +57,13 @@ public:
 
 	/**
 	 * Find Polygroups by randomly picking initial seed triangles and then flood-filling outwards,
-	 * stopping when the opening angle at an edge is larger than the angle defined by the DotTolerance.
+	 * stopping when the opening angle at an edge is larger than the angle defined by the OneMinusCosAngleTolerance.
 	 */
 	bool FindPolygroupsFromFaceNormals(
-		double DotTolerance = 0.0001,
+		double OneMinusCosAngleTolerance = 0.0001,
 		bool bRespectUVSeams = false,
-		bool bRespectNormalSeams = false);
+		bool bRespectNormalSeams = false,
+		bool bUseAveragePolygroupNormals = false);
 
 	/**
 	 * Find Polygroups based on UV Islands, ie each UV Island becomes a Polygroup
diff --git a/Engine/Source/Runtime/GeometryCore/Private/Clustering/FaceNormalClustering.cpp b/Engine/Source/Runtime/GeometryCore/Private/Clustering/FaceNormalClustering.cpp
new file mode 100644
--- /dev/null
+++ b/Engine/Source/Runtime/GeometryCore/Private/Clustering/FaceNormalClustering.cpp
@@ -0,0 +1,154 @@
+// Copyright Epic Games, Inc. All Rights Reserved.
+
+#include "Clustering/FaceNormalClustering.h"
+
+#include "DynamicMesh/DynamicMesh3.h"
+#include "DynamicMesh/MeshNormals.h"
+#include "MathUtil.h"
+#include "Math/Vector.h"
+#include "Util/IndexPriorityQueue.h"
+#include "VectorUtil.h"
+#include "VertexConnectedComponents.h" // for FSizedDisjointSet
+
+namespace UE::Geometry::FaceNormalClustering
+{
+
+
+void ComputeMeshPolyGroupsFromClusters(FDynamicMesh3& Mesh, TArray<TArray<int32>>& OutPolyGroups, const FClusterOptions& Options, TSet<int32>* IgnoreEdges)
+{
+	int32 NumT = Mesh.MaxTriangleID();
+
+	struct FNormalArea
+	{
+		FVector3d Normal = FVector3d::Zero();
+		double Area = 0;
+	};
+
+	TArray<FNormalArea> GroupNormalAreas;
+	GroupNormalAreas.SetNumZeroed(NumT);
+	FVector3d TriV[3];
+	for (int32 TID = 0; TID < NumT; ++TID)
+	{
+		if (Mesh.IsTriangle(TID))
+		{
+			FNormalArea& NormalArea = GroupNormalAreas[TID];
+			Mesh.GetTriVertices(TID, TriV[0], TriV[1], TriV[2]);
+			NormalArea.Normal = VectorUtil::NormalArea(TriV[0], TriV[1], TriV[2], NormalArea.Area);
+		}
+	}
+
+	FSizedDisjointSet PlaneGroups;
+	PlaneGroups.Init(NumT, [&Mesh](int32 TID) { return Mesh.IsTriangle(TID); });
+	int32 NumEdges = Mesh.MaxEdgeID();
+	Mesh.GetEdge(0);
+	FIndexPriorityQueue EdgeQueue(NumEdges);
+
+	const double AreaThreshold = FMath::Max(FMathd::Epsilon, Options.SmallFaceAreaThreshold);
+	auto GetMergeWeightForGroups = [&GroupNormalAreas, AreaThreshold](const FIndex2i& EdgeGroups) -> float
+	{
+		if (GroupNormalAreas[EdgeGroups.A].Area < AreaThreshold || GroupNormalAreas[EdgeGroups.B].Area < AreaThreshold)
+		{
+			return 0;
+		}
+		float NormalAlignment = 1 - float(GroupNormalAreas[EdgeGroups.A].Normal.Dot(GroupNormalAreas[EdgeGroups.B].Normal));
+		return NormalAlignment;
+	};
+	auto GetMergeWeightForEdge = [&GetMergeWeightForGroups, &Mesh, &PlaneGroups](int32 EdgeID, FIndex2i& EdgeGroups) -> float
+	{
+		EdgeGroups = Mesh.GetEdgeT(EdgeID);
+		EdgeGroups.A = PlaneGroups.Find(EdgeGroups.A);
+		EdgeGroups.B = PlaneGroups.Find(EdgeGroups.B);
+		if (EdgeGroups.A == EdgeGroups.B)
+		{		
+			// return a value guaranteed to be above any merge tolerance, so we ignore edges connecting already-merged groups
+			constexpr float CannotMergeValue = FMathf::MaxReal;
+			return CannotMergeValue;
+		}
+		return GetMergeWeightForGroups(EdgeGroups);
+	};
+
+	const float Threshold = static_cast<float>(Options.NormalOneMinusCosTolerance);
+	for (int32 EdgeID : Mesh.EdgeIndicesItr())
+	{
+		if (IgnoreEdges && IgnoreEdges->Contains(EdgeID))
+		{
+			continue;
+		}
+		FIndex2i EdgeT = Mesh.GetEdgeT(EdgeID);
+		if (EdgeT.Contains(INDEX_NONE))
+		{
+			continue;
+		}
+		float Wt = GetMergeWeightForGroups(EdgeT);
+		if (Wt < Threshold)
+		{
+			EdgeQueue.Insert(EdgeID, Wt);
+		}
+	}
+
+	int32 NumRemainingPlanes = Mesh.TriangleCount();
+	while (EdgeQueue.GetCount() > 0)
+	{
+		int32 EdgeID = EdgeQueue.Dequeue();
+		FIndex2i EdgeGroups;
+
+		// If we're considering the overall cluster normals, we need to re-evaluate the potential merge based on the current cluster normals
+		if (Options.bApplyNormalToleranceToClusters)
+		{
+			float UpdatedWt = GetMergeWeightForEdge(EdgeID, EdgeGroups);
+
+			// If updated weight is too large, we can skip it
+			if (UpdatedWt >= Threshold)
+			{
+				continue;
+			}
+
+			// Test if the updated weight is still less than the next node's weight; if not, re-queue it
+			if (EdgeQueue.GetCount() > 0)
+			{
+				float NextWt = EdgeQueue.GetFirstNodePriority();
+				if (UpdatedWt > NextWt)
+				{
+					EdgeQueue.Insert(EdgeID, UpdatedWt);
+					continue;
+				}
+			}
+		}
+
+		// We know UpdatedWt < Threshold, so merge the groups
+		int32 Parent = PlaneGroups.Union(EdgeGroups.A, EdgeGroups.B);
+		// Update the NormalArea of the new parent
+		const FNormalArea& NA_A = GroupNormalAreas[EdgeGroups.A];
+		const FNormalArea& NA_B = GroupNormalAreas[EdgeGroups.B];
+		FNormalArea& NA_P = GroupNormalAreas[Parent];
+		NA_P.Normal = NA_A.Area * NA_A.Normal + NA_B.Area * NA_B.Normal;
+		NA_P.Normal.Normalize();
+		NA_P.Area = NA_A.Area + NA_B.Area;
+		NumRemainingPlanes--;
+		if (NumRemainingPlanes < Options.TargetMinGroups)
+		{
+			break;
+		}
+	}
+
+	// Copy triangle IDs of plane groups into the OutPolyGroups arrays
+	TArray<int32> CompactIdxToGroupID, GroupIDToCompactIdx;
+	int32 NumFoundGroups = PlaneGroups.CompactedGroupIndexToGroupID(&CompactIdxToGroupID, &GroupIDToCompactIdx);
+
+	OutPolyGroups.SetNum(NumFoundGroups);
+	for (int32 GroupIndex = 0; GroupIndex < NumFoundGroups; ++GroupIndex)
+	{
+		int32 GroupID = CompactIdxToGroupID[GroupIndex];
+		int32 Size = PlaneGroups.GetSize(GroupID);
+		OutPolyGroups[GroupIndex].Reserve(Size);
+	}
+	for (int32 TID : Mesh.TriangleIndicesItr())
+	{
+		int32 GroupID = PlaneGroups.Find(TID);
+		int32 CompactIdx = GroupIDToCompactIdx[GroupID];
+		OutPolyGroups[CompactIdx].Add(TID);
+	}
+}
+
+
+} // namespace UE::Geometry::FaceNormalClustering
\ No newline at end of file
diff --git a/Engine/Source/Runtime/GeometryCore/Public/Clustering/FaceNormalClustering.h b/Engine/Source/Runtime/GeometryCore/Public/Clustering/FaceNormalClustering.h
new file mode 100644
--- /dev/null
+++ b/Engine/Source/Runtime/GeometryCore/Public/Clustering/FaceNormalClustering.h
@@ -0,0 +1,43 @@
+// Copyright Epic Games, Inc. All Rights Reserved.
+
+#pragma once
+
+#include "MathUtil.h"
+
+namespace UE::Geometry
+{
+
+class FDynamicMesh3;
+
+namespace FaceNormalClustering
+{
+	struct FClusterOptions
+	{
+		// If one minus the dot product of normals is greater than this value, they are similar enough to be joined together.
+		double NormalOneMinusCosTolerance = .01;
+		// The area threshold below which we ignore the computed triangle normal, and always merge to an adjacent group. Helps avoid tiny triangles failing to cluster due to unreliable normals.
+		double SmallFaceAreaThreshold = FMathd::ZeroTolerance;
+		// The minimum number of cluster groups below which to stop merging clusters
+		int32 TargetMinGroups = 0;
+		// Whether to consider the average normal of the whole cluster before merging. If false, will only consider the normals of the original triangles within a group.
+		bool bApplyNormalToleranceToClusters = true;
+
+		void SetNormalAngleToleranceInDegrees(double ToleranceInDegrees)
+		{
+			NormalOneMinusCosTolerance = 1 - FMathd::Cos(FMathd::DegToRad * ToleranceInDegrees);
+		}
+	};
+
+	/**
+	 * Find Mesh PolyGroups by clustering neighboring faces (or polygroups) with similar normals.
+	 * 
+	 * @param Mesh				The Mesh to operate on
+	 * @param OutPolyGroups		An array holding the array-of-triangle-IDs for each found PolyGroup
+	 * @param Options			Options controlling the tolerances used for clustering
+	 * @param IgnoreEdges		Optional set of edges that clusters should not be merged across
+	 */
+	void GEOMETRYCORE_API ComputeMeshPolyGroupsFromClusters(FDynamicMesh3& Mesh, TArray<TArray<int32>>& OutPolyGroups, const FClusterOptions& Options, TSet<int32>* IgnoreEdges = nullptr);
+
+}
+
+} // end namespace UE::Geometry
\ No newline at end of file
diff --git a/Engine/Source/Runtime/GeometryCore/Public/VertexConnectedComponents.h b/Engine/Source/Runtime/GeometryCore/Public/VertexConnectedComponents.h
--- a/Engine/Source/Runtime/GeometryCore/Public/VertexConnectedComponents.h
+++ b/Engine/Source/Runtime/GeometryCore/Public/VertexConnectedComponents.h
@@ -48,13 +48,14 @@ struct FSizedDisjointSet
 		}
 	}
 
-	void Union(int32 A, int32 B)
+	// @return the Parent of the Union
+	int32 Union(int32 A, int32 B)
 	{
 		int32 ParentA = Find(A);
 		int32 ParentB = Find(B);
 		if (ParentA == ParentB)
 		{
-			return;
+			return ParentB;
 		}
 		if (Sizes[ParentA] > Sizes[ParentB])
 		{
@@ -62,6 +63,7 @@ struct FSizedDisjointSet
 		}
 		Parents[ParentA] = ParentB;
 		Sizes[ParentB] += Sizes[ParentA];
+		return ParentB;
 	}
 
 	int32 Find(int32 Idx)
@@ -82,6 +84,60 @@ struct FSizedDisjointSet
 		int32 Parent = Find(Idx);
 		return Sizes[Parent];
 	}
+
+	/**
+	 * Create mappings between compacted Group Index and group ID, where the compacted indices numbers the groups from 0 to NumGroups
+	 * 
+	 * @param CompactIdxToGroupID  Array to fill with the Compact Index -> Original Group ID mapping.  If null, array will not be set.
+	 * @param GroupIDToCompactIdx  Array to fill with the Original Group ID -> Compact Index mapping.  If null, array will not be set.
+	 * @param MinGroupSize  The minimum size of group to consider. Groups smaller than this will not be counted.
+	 * @return The number of groups found
+	 */
+	int32 CompactedGroupIndexToGroupID(TArray<int32>* CompactIdxToGroupID, TArray<int32>* GroupIDToCompactIdx, int32 MinGroupSize = 1)
+	{
+		MinGroupSize = FMath::Max(1, MinGroupSize);
+
+		int32 NumIDs = Sizes.Num();
+		if (CompactIdxToGroupID)
+		{
+			CompactIdxToGroupID->Reset();
+		}
+		if (GroupIDToCompactIdx)
+		{
+			GroupIDToCompactIdx->Init(-1, NumIDs);
+		}
+		int32 NumGroups = 0;
+		for (int32 ID = 0; ID < NumIDs; ++ID)
+		{
+			if (Parents[ID] == INDEX_NONE) // ignore invalid IDs
+			{
+				continue;
+			}
+
+			int32 Parent = Find(ID);
+			if (Parent != ID) // only count groups on the Parent==ID node
+			{
+				continue;
+			}
+
+			if (Sizes[Parent] < MinGroupSize) // ignore too-small groups
+			{
+				continue;
+			}
+
+			// Record the unique GroupID
+			if (CompactIdxToGroupID)
+			{
+				CompactIdxToGroupID->Add(Parent);
+			}
+			if (GroupIDToCompactIdx)
+			{
+				(*GroupIDToCompactIdx)[ID] = NumGroups;
+			}
+			NumGroups++;
+		}
+		return NumGroups;
+	}
 };
 
 /// Vertex-based connected components class -- can work with any mesh that has vertex IDs
