Fix issue in negative space voxel mesh generation where eroding the (convex hull - original mesh) by the distance tolerance could delete valid negative space when the convex hull was thinner than the distance tolerance.

We now directly construct the initial voxel mesh to be the tolerance distance from the input mesh (while allowing it to follow the convex hull), and also ensure clamp the grid cell size to not be larger than half the width of the region bounding box's min dimension.

diff --git a/Engine/Source/Runtime/GeometryCore/Private/CompGeom/ConvexDecomposition3.cpp b/Engine/Source/Runtime/GeometryCore/Private/CompGeom/ConvexDecomposition3.cpp
--- a/Engine/Source/Runtime/GeometryCore/Private/CompGeom/ConvexDecomposition3.cpp
+++ b/Engine/Source/Runtime/GeometryCore/Private/CompGeom/ConvexDecomposition3.cpp
@@ -145,40 +145,37 @@ bool FSphereCovering::AddNegativeSpace(const TFastWindingTree<FDynamicMesh3>& Sp
 		FAxisAlignedBox3d HullBox = HullAABB.GetBoundingBox();
 		HullBox.Expand(1);
 		FMarchingCubes MarchingCubes;
-		MarchingCubes.CubeSize = FMath::Max(HullBox.MaxDim() / 128.0, SampleSettings.ReduceRadiusMargin * .5);
+		MarchingCubes.CubeSize = FMath::Clamp(SampleSettings.ReduceRadiusMargin * .5, HullBox.MaxDim() / 128.0, HullBox.MinDim() * .5);
 		MarchingCubes.Bounds = HullBox;
 		MarchingCubes.RootMode = ERootfindingModes::Bisection;
 		MarchingCubes.RootModeSteps = 3;
 		MarchingCubes.IsoValue = 0;
-		MarchingCubes.Implicit = [&HullWinding, &Spatial, WindingSign](const FVector3d& Pt) -> double
+		MarchingCubes.Implicit = [&HullWinding, &Spatial, WindingSign, &SampleSettings](const FVector3d& Pt) -> double
 		{
-			// Volume is anything inside the hull and outside the input surface
-			return (HullWinding.FastWindingNumber(Pt) > .5) && (Spatial.FastWindingNumber(Pt) * WindingSign <= .5) ? 1.0 : -1.0;
+			// Volume is inside the hull and outside the input surface
+			if ((HullWinding.FastWindingNumber(Pt) > .5) && (Spatial.FastWindingNumber(Pt) * WindingSign <= .5))
+			{
+				// Volume is at least ReduceRadiusMargin away from the input surface
+				double NearDistSq = FMathd::MaxReal;
+				int32 NearTri = Spatial.GetTree()->FindNearestTriangle(Pt, NearDistSq, SampleSettings.ReduceRadiusMargin);
+				if (NearTri == INDEX_NONE)
+				{
+					return 1.0;
+				}
+				return -1.0;
+			}
+			return -1.0;
 		};
-		TArray<FVector3d> MCSeeds;
-		MCSeeds.Reserve(Mesh->VertexCount());
-		for (int32 VertIdx : Mesh->VertexIndicesItr())
-		{
-			FVector3d Vertex = Mesh->GetVertex(VertIdx);
-			MCSeeds.Add(Vertex);
-		}
-		FDynamicMesh3 HullMinusComponents(&MarchingCubes.GenerateContinuation(MCSeeds));
-		HullMinusComponents.DiscardAttributes();
-
-		// Contract the surface by the ReduceRadiusMargin to avoid sampling too-small concave regions
-		FDynamicMeshAABBTree3 MorphologyBVTree(&HullMinusComponents);
-		TImplicitMorphology<FDynamicMesh3> ImplicitMorphology;
-		ImplicitMorphology.MorphologyOp = TImplicitMorphology<FDynamicMesh3>::EMorphologyOp::Contract;
-		ImplicitMorphology.Source = &HullMinusComponents;
-		ImplicitMorphology.SourceSpatial = &MorphologyBVTree;
-		ImplicitMorphology.Distance = SampleSettings.ReduceRadiusMargin;
-		ImplicitMorphology.GridCellSize = MarchingCubes.CubeSize;
-		ImplicitMorphology.MeshCellSize = MarchingCubes.CubeSize;
-		FDynamicMesh3 MorphologyMesh(&ImplicitMorphology.Generate());
-		MorphologyMesh.DiscardAttributes();
+
+		// Note: we could try to sample the hull away from vertices (on edges > 2*ReduceRadiusMargin long, and similarly on large triangles)
+		// to generate seeds to allow us to call the (generally faster) MarchingCubes.GenerateContinuation(Seeds) function below.
+		// However, naively calling GenerateContinuation with the mesh vertices can give an empty result (due to the radius margin).
+		// For now, we just call the slower (but reliable) Generate() function.
+		FDynamicMesh3 NegativeSpaceMesh(&MarchingCubes.Generate());
+		NegativeSpaceMesh.DiscardAttributes();
 
 		// Make sure the mesh is compact to simplify downsampling below
-		MorphologyMesh.CompactInPlace();
+		NegativeSpaceMesh.CompactInPlace();
 
 		auto AddSample = [this, &Mesh, &Spatial, &SampleSettings, &bAddedPoints, WindingSign](FVector3d Pos)
 		{
@@ -232,26 +229,26 @@ bool FSphereCovering::AddNegativeSpace(const TFastWindingTree<FDynamicMesh3>& Sp
 		// if we have more vertices than we want samples, downsample so that we:
 		//  (1) ~uniformly cover space and
 		//  (2) prioritize samples at 'features' (sharp angles) of the offset mesh
-		if (MorphologyMesh.MaxVertexID() > SampleSettings.TargetNumSamples)
+		if (NegativeSpaceMesh.MaxVertexID() > SampleSettings.TargetNumSamples)
 		{
 			TArray<float> VertexAngleMetric;
-			VertexAngleMetric.SetNumZeroed(MorphologyMesh.MaxVertexID());
+			VertexAngleMetric.SetNumZeroed(NegativeSpaceMesh.MaxVertexID());
 			TArray<FVector3d> TriNormals;
-			TriNormals.SetNumZeroed(MorphologyMesh.MaxTriangleID());
-			for (int32 TID : MorphologyMesh.TriangleIndicesItr())
+			TriNormals.SetNumZeroed(NegativeSpaceMesh.MaxTriangleID());
+			for (int32 TID : NegativeSpaceMesh.TriangleIndicesItr())
 			{
-				TriNormals[TID] = MorphologyMesh.GetTriNormal(TID);
+				TriNormals[TID] = NegativeSpaceMesh.GetTriNormal(TID);
 			}
 			TArray<FVector3d> VertexPositions;
-			VertexPositions.SetNumZeroed(MorphologyMesh.MaxVertexID());
-			for (int32 VID : MorphologyMesh.VertexIndicesItr())
+			VertexPositions.SetNumZeroed(NegativeSpaceMesh.MaxVertexID());
+			for (int32 VID : NegativeSpaceMesh.VertexIndicesItr())
 			{
-				VertexPositions[VID] = MorphologyMesh.GetVertex(VID);
+				VertexPositions[VID] = NegativeSpaceMesh.GetVertex(VID);
 				// Note: Angle metric favors vertices on edges with larger dihedral angles
 				float AngleMetric = 0;
-				MorphologyMesh.EnumerateVertexEdges(VID, [&](int32 EID)
+				NegativeSpaceMesh.EnumerateVertexEdges(VID, [&](int32 EID)
 					{
-						FIndex2i EdgeT = MorphologyMesh.GetEdgeT(EID);
+						FIndex2i EdgeT = NegativeSpaceMesh.GetEdgeT(EID);
 						if (EdgeT.B != FDynamicMesh3::InvalidID)
 						{
 							AngleMetric = FMath::Max(AngleMetric, float(1 - TriNormals[EdgeT.A].Dot(TriNormals[EdgeT.B])));
@@ -265,7 +262,7 @@ bool FSphereCovering::AddNegativeSpace(const TFastWindingTree<FDynamicMesh3>& Sp
 			for (int32 SampleIdx = 0; SampleIdx < NumSamples; ++SampleIdx)
 			{
 				int32 VID = Ordering.Order[SampleIdx];
-				AddSample(MorphologyMesh.GetVertex(VID));
+				AddSample(NegativeSpaceMesh.GetVertex(VID));
 			}
 
 			if (SampleSettings.bRequireSearchSampleCoverage)
@@ -274,7 +271,7 @@ bool FSphereCovering::AddNegativeSpace(const TFastWindingTree<FDynamicMesh3>& Sp
 				for (int32 SampleIdx = NumSamples; SampleIdx < Ordering.Order.Num(); ++SampleIdx)
 				{
 					int32 VID = Ordering.Order[SampleIdx];
-					FVector3d Pos = MorphologyMesh.GetVertex(VID);
+					FVector3d Pos = NegativeSpaceMesh.GetVertex(VID);
 					// TODO: Consider accelerating this coverage search w/ e.g. a sparse dynamic octree representation of the sphere covering
 					bool bFoundCover = false;
 					for (int32 SphereIdx = 0; SphereIdx < Position.Num(); ++SphereIdx)
@@ -296,7 +293,7 @@ bool FSphereCovering::AddNegativeSpace(const TFastWindingTree<FDynamicMesh3>& Sp
 		else
 		{
 			// When we have fewer vertices than target samples, just try adding all of them
-			for (FVector3d Pos : MorphologyMesh.VerticesItr())
+			for (FVector3d Pos : NegativeSpaceMesh.VerticesItr())
 			{
 				AddSample(Pos);
 			}
