fix crash in convex hull simplification due to not tracking polygon group IDs correctly

diff --git a/Engine/Source/Runtime/GeometryCore/Private/CompGeom/ConvexHull3.cpp b/Engine/Source/Runtime/GeometryCore/Private/CompGeom/ConvexHull3.cpp
--- a/Engine/Source/Runtime/GeometryCore/Private/CompGeom/ConvexHull3.cpp
+++ b/Engine/Source/Runtime/GeometryCore/Private/CompGeom/ConvexHull3.cpp
@@ -1250,6 +1250,10 @@ void TConvexHull3<RealType>::GetSimplifiedFaces(TArray<FPolygonFace>& OutPolygon
 				{
 					OutPolygonNormals->Add((TVector<RealType>)PlaneNormals[GroupIdx]);
 				}
+				else
+				{
+					PolygonToGroup.Add(GroupIdx);
+				}
 			}
 			else
 			{
