support custom early-out overlap queries for sparse dynamic octree, and early-out Enumerate functions on SmallListSet and SparseListSet to support this

diff --git a/Engine/Source/Runtime/GeometryCore/Private/Spatial/SparseDynamicOctree3.cpp b/Engine/Source/Runtime/GeometryCore/Private/Spatial/SparseDynamicOctree3.cpp
--- a/Engine/Source/Runtime/GeometryCore/Private/Spatial/SparseDynamicOctree3.cpp
+++ b/Engine/Source/Runtime/GeometryCore/Private/Spatial/SparseDynamicOctree3.cpp
@@ -519,6 +519,40 @@ void FSparseDynamicOctree3::ParallelRangeQuery(
 }
 
 
+int FSparseDynamicOctree3::ParallelOverlapAnyQuery(const FAxisAlignedBox3d& ShapeBounds,
+	TFunctionRef<bool(int32)> ObjectOverlapFn, TFunctionRef<bool(const FAxisAlignedBox3d&)> BoundsOverlapFn) const
+{
+	FAxisAlignedBox3d EmptyCell = FAxisAlignedBox3d::Empty();
+	// first test spill objects
+	for (int ObjectID : SpillObjectSet)
+	{
+		if (ObjectOverlapFn(ObjectID))
+		{
+			return ObjectID;
+		}
+	}
+
+	TArray<const FSparseOctreeCell*, TInlineAllocator<32>> Queue = InitializeQueryQueue(ShapeBounds);
+
+	std::atomic<int> FoundID = INDEX_NONE;
+	ParallelFor(Queue.Num(), [&](int32 qi)
+	{
+		if (FoundID != INDEX_NONE)
+		{
+			return;
+		}
+
+		const FSparseOctreeCell* RootCell = Queue[qi];
+		int LocalFoundID = BranchCustomOverlapAnyQuery(RootCell, ShapeBounds, ObjectOverlapFn, BoundsOverlapFn);
+		if (LocalFoundID != INDEX_NONE)
+		{
+			FoundID = LocalFoundID;
+		}
+	});
+	return FoundID;
+}
+
+
 void FSparseDynamicOctree3::BranchRangeQuery(
 	const FSparseOctreeCell* ParentCell,
 	const FAxisAlignedBox3d& Bounds,
@@ -551,6 +585,51 @@ void FSparseDynamicOctree3::BranchRangeQuery(
 	}
 }
 
+int FSparseDynamicOctree3::BranchCustomOverlapAnyQuery(
+	const FSparseOctreeCell* ParentCell,
+	const FAxisAlignedBox3d& Bounds,
+	TFunctionRef<bool(int32)> ObjectOverlapFn, TFunctionRef<bool(const FAxisAlignedBox3d&)> BoundsOverlapFn) const
+{
+	TArray<const FSparseOctreeCell*, TInlineAllocator<32>> Queue;
+	Queue.Add(ParentCell);
+
+	while (Queue.Num() > 0)
+	{
+		const FSparseOctreeCell* CurCell = Queue.Pop(false);
+
+		// process elements
+		int32 FoundID = INDEX_NONE;
+		bool bNoOverlaps = CellObjectLists.EnumerateEarlyOut(CurCell->CellID, [&](int32 ObjectID)
+		{
+			if (ObjectOverlapFn(ObjectID))
+			{
+				FoundID = ObjectID;
+				return false;
+			}
+			return true;
+		});
+
+		if (!bNoOverlaps)
+		{
+			return FoundID;
+		}
+
+		for (int k = 0; k < 8; ++k)
+		{
+			if (CurCell->HasChild(k))
+			{
+				const FSparseOctreeCell* ChildCell = &Cells[CurCell->GetChildCellID(k)];
+				FAxisAlignedBox3d CellBox = GetCellBox(*ChildCell, MaxExpandFactor);
+				if (CellBox.Intersects(Bounds) && BoundsOverlapFn(CellBox))
+				{
+					Queue.Add(ChildCell);
+				}
+			}
+		}
+	}
+
+	return INDEX_NONE;
+}
 
 
 
@@ -649,10 +728,10 @@ void FSparseDynamicOctree3::CheckValidity(
 	// check that all object IDs in per-cell object lists is valid
 	for (int32 CellID : CellRefCounts.Indices())
 	{
-		for (int32 ObjectID : CellObjectLists.Values(CellID))
+		CellObjectLists.Enumerate(CellID, [&](int32 ObjectID)
 		{
 			CheckOrFailF(IsValidObjectIDFunc(ObjectID));
-		}
+		});
 	}
 
 	uint32 NumObjectIDs = ObjectIDToCellMap.Num();
diff --git a/Engine/Source/Runtime/GeometryCore/Private/Util/SmallListSet.cpp b/Engine/Source/Runtime/GeometryCore/Private/Util/SmallListSet.cpp
--- a/Engine/Source/Runtime/GeometryCore/Private/Util/SmallListSet.cpp
+++ b/Engine/Source/Runtime/GeometryCore/Private/Util/SmallListSet.cpp
@@ -369,7 +369,47 @@ void FSmallListSet::Enumerate(int32 ListIndex, TFunctionRef<void(int32)> ApplyFu
 	}
 }
 
-
+bool FSmallListSet::EnumerateEarlyOut(int32 ListIndex, TFunctionRef<bool(int32)> ApplyFunc) const
+{
+	int32 block_ptr = ListHeads[ListIndex];
+	if (block_ptr != NullValue)
+	{
+		int32 N = ListBlocks[block_ptr];
+		if (N < BLOCKSIZE)
+		{
+			int32 iEnd = block_ptr + N;
+			for (int32 i = block_ptr + 1; i <= iEnd; ++i)
+			{
+				if (!ApplyFunc(ListBlocks[i]))
+				{
+					return false;
+				}
+			}
+		}
+		else
+		{
+			// we spilled to linked list, have to iterate through it as well
+			int32 iEnd = block_ptr + BLOCKSIZE;
+			for (int32 i = block_ptr + 1; i <= iEnd; ++i)
+			{
+				if (!ApplyFunc(ListBlocks[i]))
+				{
+					return false;
+				}
+			}
+			int32 cur_ptr = ListBlocks[block_ptr + BLOCK_LIST_OFFSET];
+			while (cur_ptr != NullValue)
+			{
+				if (!ApplyFunc(LinkedListElements[cur_ptr]))
+				{
+					return false;
+				}
+				cur_ptr = LinkedListElements[cur_ptr + 1];
+			}
+		}
+	}
+	return true;
+}
 
 
 int32 FSmallListSet::AllocateBlock()
diff --git a/Engine/Source/Runtime/GeometryCore/Public/Spatial/SparseDynamicOctree3.h b/Engine/Source/Runtime/GeometryCore/Public/Spatial/SparseDynamicOctree3.h
--- a/Engine/Source/Runtime/GeometryCore/Public/Spatial/SparseDynamicOctree3.h
+++ b/Engine/Source/Runtime/GeometryCore/Public/Spatial/SparseDynamicOctree3.h
@@ -314,6 +314,35 @@ public:
 		TArray<int>& ObjectIDsOut) const;
 
 
+	/**
+	 * Find any overlap between a caller-defined query and any object ID. Returns the first overlap it finds.
+	 * @param ShapeBounds Overall bounds of the custom query shape
+	 * @param ObjectOverlapFn Custom function that indicates if the given ObjectID overlaps the query shape
+	 * @param BoundsOverlapFn Optional custom function that indicates if the given bounding box overlaps the query shape, used to filter the octree search.
+	 *                        Note the bounds overlap with ShapeBounds will already be called before calling this; only provide this if a more accurate overlap can be quickly checked.
+	 * @return index of overlapping ObjectID or INDEX_NONE if no overlaps are found
+	 */
+	GEOMETRYCORE_API int ParallelOverlapAnyQuery(const FAxisAlignedBox3d& ShapeBounds,
+		TFunctionRef<bool(int32)> ObjectOverlapFn, 
+		TFunctionRef<bool(const FAxisAlignedBox3d&)> BoundsOverlapFn) const;
+
+	/**
+	 * Find any overlap between a caller-defined query and any object ID. Returns the first overlap it finds.
+	 * @param ShapeBounds Overall bounds of the custom query shape
+	 * @param ObjectOverlapFn Custom function that indicates if the given ObjectID overlaps the query shape
+	 * @return index of overlapping ObjectID or INDEX_NONE if no overlaps are found
+	 */
+	GEOMETRYCORE_API int ParallelOverlapAnyQuery(const FAxisAlignedBox3d& ShapeBounds,
+		TFunctionRef<bool(int32)> ObjectOverlapFn) const
+	{
+		return ParallelOverlapAnyQuery(ShapeBounds, ObjectOverlapFn, 
+			// Note a bounds vs ShapeBounds test will always be performed automatically;
+			// the optional bounds overlap fn here can just return true to indicate no additional filtering is performed
+			[](const FAxisAlignedBox3d&) {return true;}
+		);
+	}
+
+
 	/**
 	 * Check that the octree is internally valid
 	 * @param IsValidObjectIDFunc function that returns true if given ObjectID is valid
@@ -360,6 +389,7 @@ protected:
 	// list of cells. Note that some cells may be unused, depending on CellRefCounts
 	TDynamicVector<FSparseOctreeCell> Cells;
 
+	// TODO: Consider switching FSmallListSet to TSparseListSet<int32> for more reliable performance in cases where we have more than 8 objects in many cells
 	FSmallListSet CellObjectLists;			// per-cell object ID lists
 	TSet<int32> SpillObjectSet;				// list of object IDs for objects that didn't fit in a root cell
 
@@ -461,6 +491,10 @@ protected:
 		TArray<int>& ObjectIDs) const;
 
 private:
+	int BranchCustomOverlapAnyQuery(const FSparseOctreeCell* ParentCell,
+		const FAxisAlignedBox3d& Bounds,
+		TFunctionRef<bool(int32)> ObjectOverlapFn, TFunctionRef<bool(const FAxisAlignedBox3d&)> BoundsOverlapFn) const;
+
 	// helper to find the root-level cells that could intersect with a given query object, specialized for point containment queries
 	TArray<const FSparseOctreeCell*, TInlineAllocator<32>> InitializeQueryQueue(const FVector3d& Point) const
 	{
diff --git a/Engine/Source/Runtime/GeometryCore/Public/Util/SmallListSet.h b/Engine/Source/Runtime/GeometryCore/Public/Util/SmallListSet.h
--- a/Engine/Source/Runtime/GeometryCore/Public/Util/SmallListSet.h
+++ b/Engine/Source/Runtime/GeometryCore/Public/Util/SmallListSet.h
@@ -188,6 +188,12 @@ public:
 	 */
 	GEOMETRYCORE_API void Enumerate(int32 ListIndex, TFunctionRef<void(int32)> ApplyFunc) const;
 
+	/**
+	 * Call ApplyFunc on each element of the list at ListIndex, until ApplyFunc returns false
+	 * @return true if all elements were processed and ApplyFunc never returned false
+	 */
+	bool EnumerateEarlyOut(int32 ListIndex, TFunctionRef<bool(int32)> ApplyFunc) const;
+
 	/**
 	 * Serialization operator for FSmallListSet.
 	 *
diff --git a/Engine/Source/Runtime/GeometryCore/Public/Util/SparseListSet.h b/Engine/Source/Runtime/GeometryCore/Public/Util/SparseListSet.h
--- a/Engine/Source/Runtime/GeometryCore/Public/Util/SparseListSet.h
+++ b/Engine/Source/Runtime/GeometryCore/Public/Util/SparseListSet.h
@@ -184,6 +184,33 @@ public:
 		}
 	}
 
+	/**
+	 * Call ApplyFunc on each element of the list at ListIndex, until ApplyFunc returns false
+	 * @return true if all elements were processed and ApplyFunc never returned false
+	 */
+	template<typename FuncType>
+	bool EnumerateEarlyOut(int32 ListIndex, FuncType ApplyFunc) const
+	{
+		checkSlow(IsAllocated(ListIndex));
+		const FList& List = Lists[ListIndex];
+		if (List.CurBlock != nullptr)
+		{
+			int32 N = List.Blocks.Num();
+			for (int32 BlockIndex = 0; BlockIndex < N; ++BlockIndex)
+			{
+				const ElementType* Elements = List.Blocks[BlockIndex];
+				const ElementType* EndElement = (BlockIndex == N - 1) ? (Elements + List.CurIndex) : (Elements + BlockSize);
+				while (Elements != EndElement)
+				{
+					if (!ApplyFunc(*Elements++))
+					{
+						return false;
+					}
+				}
+			}
+		}
+		return true;
+	}
 
 
 
