expose min part thickness out to geometry script collision functions + make min part thickness apply with compact shapes (just skipping the inapplicable cutting step)

diff --git a/Engine/Plugins/Experimental/GeometryScripting/Source/GeometryScriptingCore/Private/CollisionFunctions.cpp b/Engine/Plugins/Experimental/GeometryScripting/Source/GeometryScriptingCore/Private/CollisionFunctions.cpp
--- a/Engine/Plugins/Experimental/GeometryScripting/Source/GeometryScriptingCore/Private/CollisionFunctions.cpp
+++ b/Engine/Plugins/Experimental/GeometryScripting/Source/GeometryScriptingCore/Private/CollisionFunctions.cpp
@@ -1182,7 +1182,7 @@ FGeometryScriptSimpleCollision UGeometryScriptLibrary_CollisionFunctions::MergeS
 		NegativeSpace.Append(*MergeOptions.PrecomputedNegativeSpace.Spheres);
 	}
 	FSphereCovering* UseNegativeSpace = NegativeSpace.Num() > 0 ? &NegativeSpace : nullptr;
-	Decomposition.MergeBest(MergeOptions.MaxShapeCount, MergeOptions.ErrorTolerance, 0.0, true, false, MergeOptions.MaxShapeCount, UseNegativeSpace, nullptr /*optional FTransform for negative space*/);
+	Decomposition.MergeBest(MergeOptions.MaxShapeCount, MergeOptions.ErrorTolerance, MergeOptions.MinThicknessTolerance, true, false, MergeOptions.MaxShapeCount, UseNegativeSpace, nullptr /*optional FTransform for negative space*/);
 
 	// Algorithm decided not to merge
 	if (Decomposition.NumHulls() == InitialNumConvex)
diff --git a/Engine/Plugins/Experimental/GeometryScripting/Source/GeometryScriptingCore/Public/GeometryScript/CollisionFunctions.h b/Engine/Plugins/Experimental/GeometryScripting/Source/GeometryScriptingCore/Public/GeometryScript/CollisionFunctions.h
--- a/Engine/Plugins/Experimental/GeometryScripting/Source/GeometryScriptingCore/Public/GeometryScript/CollisionFunctions.h
+++ b/Engine/Plugins/Experimental/GeometryScripting/Source/GeometryScriptingCore/Public/GeometryScript/CollisionFunctions.h
@@ -181,6 +181,13 @@ public:
 	UPROPERTY(EditAnywhere, BlueprintReadWrite, Category = Options, meta = (UIMin = "0", UIMax = "100.", Units = cm))
 	double ErrorTolerance = 0.0;
 
+	/**
+	 * Always attempt to merge parts thicker than this, ignoring ErrorTolerance and MaxShapeCount.
+	 * Note: Negative space, if set, will still prevent merges.
+	 */
+	UPROPERTY(EditAnywhere, BlueprintReadWrite, Category = Options, meta = (UIMin = "0", UIMax = "100.", Units = cm))
+	double MinThicknessTolerance = 0.0;
+
 	// Negative space that must be preserved during merging
 	UPROPERTY(EditAnywhere, BlueprintReadWrite, Category = NegativeSpace)
 	FGeometryScriptSphereCovering PrecomputedNegativeSpace;
diff --git a/Engine/Source/Runtime/GeometryCore/Private/CompGeom/ConvexDecomposition3.cpp b/Engine/Source/Runtime/GeometryCore/Private/CompGeom/ConvexDecomposition3.cpp
--- a/Engine/Source/Runtime/GeometryCore/Private/CompGeom/ConvexDecomposition3.cpp
+++ b/Engine/Source/Runtime/GeometryCore/Private/CompGeom/ConvexDecomposition3.cpp
@@ -1723,15 +1723,16 @@ int32 FConvexDecomposition3::MergeBest(int32 InTargetNumParts, double MaxErrorTo
 	};
 
 	int32 MustMergeCount = 0;
+	bool bHaveCompactMustMergeParts = false;
 	if (MinThicknessTolerance > 0)
 	{
 		for (int32 PartIdx = 0; PartIdx < Decomposition.Num(); PartIdx++)
 		{
 			FConvexPart& Part = Decomposition[PartIdx];
-			// The 'must merge' logic is not designed to handle pre-compacted hull inputs
+			// Note if our thin parts are already compact, so we can skip processing that requires the underlying geometry
 			if (Part.IsCompact())
 			{
-				continue;
+				bHaveCompactMustMergeParts = true;
 			}
 			Part.bMustMerge = IsPartBelowSizeTolerance(Part);
 			MustMergeCount += Part.bMustMerge;
@@ -1740,7 +1741,7 @@ int32 FConvexDecomposition3::MergeBest(int32 InTargetNumParts, double MaxErrorTo
 	// TODO: if MustMergeCount is >= NumParts-1, we could shortcut the process and just return a single convex hull
 
 	// Try to further-split long/wide thin parts -- parts that are not thin in all directions, e.g. parts with max bounding box extent wider than MinThicknessTolerance
-	if (MustMergeCount > 0)
+	if (MustMergeCount > 0 && !bHaveCompactMustMergeParts)
 	{
 		TArray<int32> PartsToConsider;
 		for (int32 OrigPartIdx = 0, OrigNumDecomposition = Decomposition.Num(); OrigPartIdx < OrigNumDecomposition; OrigPartIdx++)
