make MergeCoincidentMeshEdges operation not ensure on inputs with near-zero bounding box size

diff --git a/Engine/Source/Runtime/GeometryCore/Private/DynamicMesh/Operations/MergeCoincidentMeshEdges.cpp b/Engine/Source/Runtime/GeometryCore/Private/DynamicMesh/Operations/MergeCoincidentMeshEdges.cpp
--- a/Engine/Source/Runtime/GeometryCore/Private/DynamicMesh/Operations/MergeCoincidentMeshEdges.cpp
+++ b/Engine/Source/Runtime/GeometryCore/Private/DynamicMesh/Operations/MergeCoincidentMeshEdges.cpp
@@ -30,9 +30,9 @@ bool FMergeCoincidentMeshEdges::Apply()
 	int hashN = 64;
 	if (Mesh->TriangleCount() > 100000)   hashN = 128;
 	if (Mesh->TriangleCount() > 1000000)  hashN = 256;
-	FAxisAlignedBox3d bounds = Mesh->GetBounds(true);
-	double CellSize = bounds.MaxDim() / (double)hashN;
-	MidpointsHash.Build(CellSize, bounds.Min);
+	FAxisAlignedBox3d Bounds = Mesh->GetBounds(true);
+	double CellSize = FMath::Max(FMathd::ZeroTolerance, Bounds.MaxDim() / (double)hashN);
+	MidpointsHash.Build(CellSize, Bounds.Min);
 
 	UseMergeSearchTol = FMathd::Min(CellSize, UseMergeSearchTol);
 
