fix geometrycore Polygon2::Simplify's handling of case when initial close-point-removal reduces polygon to two points

diff --git a/Engine/Source/Runtime/GeometryCore/Public/Polygon2.h b/Engine/Source/Runtime/GeometryCore/Public/Polygon2.h
--- a/Engine/Source/Runtime/GeometryCore/Public/Polygon2.h
+++ b/Engine/Source/Runtime/GeometryCore/Public/Polygon2.h
@@ -855,7 +855,16 @@ public:
 		}
 		else if (k == 2) 
 		{
-			NewVertices[k++] = Vertices[0];
+			if (pv + 1 < n) // pv < n-1 case; use vertex pv+1
+			{
+				NewVertices[2] = Vertices[pv + 1];
+			}
+			else // pv == n-1 case; use vertices n-2 and n-1
+			{
+				NewVertices[1] = Vertices[pv - 1];
+				NewVertices[2] = Vertices[pv];
+			}
+			k = 3;
 			skip_dp = true;
 		}
 
