Fix performance and memory issues when merging a large number of collision shapes:
 - More aggressively prune the cache of computed merge shapes after it reaches a threshold size
 - By default compute a proximity graph based on actual proximity, not the full n^2 possible connections
 - After finding the best of a number of possilbe merges, optionally only look for local improvements to the current best rather than doing a full search for the absolute best connection.

diff --git a/Engine/Plugins/Experimental/GeometryScripting/Source/GeometryScriptingCore/Private/CollisionFunctions.cpp b/Engine/Plugins/Experimental/GeometryScripting/Source/GeometryScriptingCore/Private/CollisionFunctions.cpp
--- a/Engine/Plugins/Experimental/GeometryScripting/Source/GeometryScriptingCore/Private/CollisionFunctions.cpp
+++ b/Engine/Plugins/Experimental/GeometryScripting/Source/GeometryScriptingCore/Private/CollisionFunctions.cpp
@@ -1161,14 +1161,16 @@ FGeometryScriptSimpleCollision UGeometryScriptLibrary_CollisionFunctions::MergeS
 		HullVertexStarts[HullIdx] = LastEnd;
 	}
 
-	// Currently we use dense proximity to propose which shapes can be merged.
-	// Note: To efficiently handle larger shape counts, consider optionally limiting these by some approximate proximity (e.g. expanded bounding box overlap)
 	TArray<TPair<int32, int32>> HullProximity;
-	for (int32 ConvexA = 0; ConvexA < InitialNumConvex; ++ConvexA)
+	if (MergeOptions.bConsiderAllPossibleMerges)
 	{
-		for (int32 ConvexB = ConvexA + 1; ConvexB < InitialNumConvex; ++ConvexB)
+		// Add all n^2 possible merge combinations for consideration
+		for (int32 ConvexA = 0; ConvexA < InitialNumConvex; ++ConvexA)
 		{
-			HullProximity.Emplace(ConvexA, ConvexB);
+			for (int32 ConvexB = ConvexA + 1; ConvexB < InitialNumConvex; ++ConvexB)
+			{
+				HullProximity.Emplace(ConvexA, ConvexB);
+			}
 		}
 	}
 
@@ -1176,11 +1178,13 @@ FGeometryScriptSimpleCollision UGeometryScriptLibrary_CollisionFunctions::MergeS
 	Decomposition.InitializeFromHulls(HullVertexStarts.Num(),
 		[&HullVolumes](int32 HullIdx) { return HullVolumes[HullIdx]; }, [&HullVertexCounts](int32 HullIdx) { return HullVertexCounts[HullIdx]; },
 		[&HullVertexStarts, &HullVertices](int32 HullIdx, int32 VertIdx) { return HullVertices[HullVertexStarts[HullIdx] + VertIdx]; }, HullProximity);
+	double MinProximityOverlapTolerance = 0;
 	FSphereCovering NegativeSpace;
 	// Build the negative space of the collision shapes, if requested
 	if (MergeOptions.bComputeNegativeSpace)
 	{
 		FNegativeSpaceSampleSettings SampleSettings = UELocal::ConvertNegativeSpaceOptions(MergeOptions.ComputeNegativeSpaceOptions);
+		MinProximityOverlapTolerance = FMath::Max(SampleSettings.ReduceRadiusMargin * .5, MinProximityOverlapTolerance);
 		FDynamicMeshAABBTree3 CollisionAABBTree(CollisionMesh.Get(), true);
 		TFastWindingTree<FDynamicMesh3> CollisionFastWinding(&CollisionAABBTree, true);
 		NegativeSpace.AddNegativeSpace(CollisionFastWinding, SampleSettings);
@@ -1191,7 +1195,15 @@ FGeometryScriptSimpleCollision UGeometryScriptLibrary_CollisionFunctions::MergeS
 		NegativeSpace.Append(*MergeOptions.PrecomputedNegativeSpace.Spheres);
 	}
 	FSphereCovering* UseNegativeSpace = NegativeSpace.Num() > 0 ? &NegativeSpace : nullptr;
-	Decomposition.MergeBest(MergeOptions.MaxShapeCount, MergeOptions.ErrorTolerance, MergeOptions.MinThicknessTolerance, true, false, MergeOptions.MaxShapeCount, UseNegativeSpace, nullptr /*optional FTransform for negative space*/);
+	if (!MergeOptions.bConsiderAllPossibleMerges)
+	{
+		// Find possible shape merges based on the shape bounding box overlaps,
+		// where bounds are expanded by: Max(a quarter their min dimension, a tenth their max dimension, the half reduce radius margin if negative space is computed)
+		Decomposition.InitializeProximityFromDecompositionBoundingBoxOverlaps(.25, .1, MinProximityOverlapTolerance);
+	}
+	int32 MaxShapeCount = FMath::Max(0, MergeOptions.MaxShapeCount);
+	Decomposition.RestrictMergeSearchToLocalAfterTestNumConnections = 1000 + MaxShapeCount * MaxShapeCount; // Restrict searches in very large search cases, when not close to max shape count, to avoid excessive search time
+	Decomposition.MergeBest(MaxShapeCount, MergeOptions.ErrorTolerance, MergeOptions.MinThicknessTolerance, true, false, MaxShapeCount, UseNegativeSpace, nullptr /*optional FTransform for negative space*/);
 
 	// Algorithm decided not to merge
 	if (Decomposition.NumHulls() == InitialNumConvex)
diff --git a/Engine/Plugins/Experimental/GeometryScripting/Source/GeometryScriptingCore/Public/GeometryScript/CollisionFunctions.h b/Engine/Plugins/Experimental/GeometryScripting/Source/GeometryScriptingCore/Public/GeometryScript/CollisionFunctions.h
--- a/Engine/Plugins/Experimental/GeometryScripting/Source/GeometryScriptingCore/Public/GeometryScript/CollisionFunctions.h
+++ b/Engine/Plugins/Experimental/GeometryScripting/Source/GeometryScriptingCore/Public/GeometryScript/CollisionFunctions.h
@@ -192,6 +192,10 @@ public:
 	UPROPERTY(EditAnywhere, BlueprintReadWrite, Category = Options, meta = (UIMin = "0", UIMax = "100.", Units = cm))
 	double MinThicknessTolerance = 0.0;
 
+	/** Whether to consider merges between every shape. If false, will only merge shapes that have overlapping or nearby bounding boxes. */
+	UPROPERTY(EditAnywhere, BlueprintReadWrite, Category = Options)
+	bool bConsiderAllPossibleMerges = false;
+
 	// Negative space that must be preserved during merging
 	UPROPERTY(EditAnywhere, BlueprintReadWrite, Category = NegativeSpace)
 	FGeometryScriptSphereCovering PrecomputedNegativeSpace;
diff --git a/Engine/Source/Runtime/GeometryCore/Private/CompGeom/ConvexDecomposition3.cpp b/Engine/Source/Runtime/GeometryCore/Private/CompGeom/ConvexDecomposition3.cpp
--- a/Engine/Source/Runtime/GeometryCore/Private/CompGeom/ConvexDecomposition3.cpp
+++ b/Engine/Source/Runtime/GeometryCore/Private/CompGeom/ConvexDecomposition3.cpp
@@ -14,6 +14,7 @@
 #include "DynamicMesh/MeshTransforms.h"
 #include "DynamicMesh/DynamicMeshAABBTree3.h"
 #include "Spatial/FastWinding.h"
+#include "Spatial/SparseDynamicOctree3.h"
 #include "Implicit/Morphology.h"
 
 #include "Async/ParallelFor.h"
@@ -1598,6 +1599,48 @@ void FConvexDecomposition3::InitializeFromHulls(int32 NumHulls, TFunctionRef<dou
 	}
 }
 
+void FConvexDecomposition3::InitializeProximityFromDecompositionBoundingBoxOverlaps(double ExpandByBoundsMinDimFactor, double ExpandByBoundsMaxDimFactor, double MinExpandAmount)
+{
+	Proximities.Empty();
+	DecompositionToProximity.Empty();
+	
+	FAxisAlignedBox3d OverallBounds;
+	double MaxDim = 0;
+	TArray<FAxisAlignedBox3d> ExpandedBounds;
+	ExpandedBounds.SetNum(Decomposition.Num());
+	for (int32 PartIdx = 0; PartIdx < Decomposition.Num(); ++PartIdx)
+	{
+		FAxisAlignedBox3d Bounds = Decomposition[PartIdx].Bounds;
+		double ExpandAmount = FMath::Max3(MinExpandAmount, Bounds.MinDim() * ExpandByBoundsMinDimFactor, Bounds.MaxDim() * ExpandByBoundsMaxDimFactor);
+		Bounds.Expand(ExpandAmount);
+		OverallBounds.Contain(Bounds);
+		MaxDim = FMath::Max(MaxDim, Bounds.MaxDim());
+		ExpandedBounds[PartIdx] = Bounds;
+	}
+	FSparseDynamicOctree3 Octree;
+	Octree.RootDimension = FMath::Max(MaxDim, OverallBounds.MaxDim() / 2.0);
+	TArray<int32> Overlaps;
+	for (int32 PartIdx = 0; PartIdx < Decomposition.Num(); ++PartIdx)
+	{
+		FAxisAlignedBox3d QueryBounds = ExpandedBounds[PartIdx];
+		
+		Overlaps.Reset();
+		Octree.RangeQuery(QueryBounds, Overlaps);
+		for (int32 OverlapIdx : Overlaps)
+		{
+			// Filter range query results for actual overlaps with query bounds
+			if (QueryBounds.Intersects(ExpandedBounds[OverlapIdx]))
+			{
+				int32 ProxIdx = Proximities.Emplace(FIndex2i(OverlapIdx, PartIdx), FPlane3d(), false /*bPlaneSeparates*/);
+				DecompositionToProximity.Add(OverlapIdx, ProxIdx);
+				DecompositionToProximity.Add(PartIdx, ProxIdx);
+			}
+		}
+
+		Octree.InsertObject(PartIdx, ExpandedBounds[PartIdx]);
+	}
+}
+
 int32 FConvexDecomposition3::MergeBest(int32 InTargetNumParts, double MaxErrorTolerance, double MinThicknessToleranceWorldSpace, bool bAllowCompact, bool bRequireHullTriangles, int32 MaxOutputHulls,
 	const FSphereCovering* OptionalNegativeSpace, const FTransform* OptionalTransformIntoNegativeSpace)
 {
@@ -1614,7 +1657,11 @@ int32 FConvexDecomposition3::MergeBest(int32 InTargetNumParts, double MaxErrorTo
 
 	int32 MergeNum = 0;
 
+	
+	// Map from proximity index to a pre-computed convex part that would result if the proximity-linked parts were merged
 	TMap<int32, TUniquePtr<FConvexPart>> ProximityComputedParts;
+	// threshold at which we will start to more aggressively evict cached computed parts from ProximityComputedParts
+	constexpr int32 EvictComputedPartsThreshold = 10000;
 
 	auto ConvexNegativeSpaceSphereOverlapTest = [OptionalNegativeSpace, OptionalTransformIntoNegativeSpace](FConvexPart* Part, int PtIdx)
 	{
@@ -1860,6 +1907,7 @@ int32 FConvexDecomposition3::MergeBest(int32 InTargetNumParts, double MaxErrorTo
 		double BestKnownCost = FMathd::MaxReal;
 		bool bOnlyAllowMustMerges = VolumeTolerance == 0 && Decomposition.Num() <= TargetNumParts;
 		int32 BestKnownIdx = -1;
+		int32 ConsideredCount = 0;
 		for (int32 ProxIdx = 0; ProxIdx < Proximities.Num(); ProxIdx++)
 		{
 			if (!Proximities[ProxIdx].bIsValidLink)
@@ -1867,6 +1915,19 @@ int32 FConvexDecomposition3::MergeBest(int32 InTargetNumParts, double MaxErrorTo
 				continue;
 			}
 
+			// Once we've considered a certain number of links and found a merge candidate, only consider links that are in the neighborhood of the current best link
+			// This prevents excessive computation to find the absolute best merge, when we only really need a local best merge
+			ConsideredCount++;
+			if (RestrictMergeSearchToLocalAfterTestNumConnections > -1 && ConsideredCount > RestrictMergeSearchToLocalAfterTestNumConnections && BestKnownIdx != -1)
+			{
+				FIndex2i BestLink = Proximities[BestKnownIdx].Link;
+				FIndex2i CurLink = Proximities[ProxIdx].Link;
+				if (!BestLink.Contains(CurLink.A) && !BestLink.Contains(CurLink.B))
+				{
+					continue;
+				}
+			}
+
 			bool bIncludesMustMergePart =
 				Decomposition[Proximities[ProxIdx].Link.A].bMustMerge ||
 				Decomposition[Proximities[ProxIdx].Link.B].bMustMerge;
@@ -1875,6 +1936,8 @@ int32 FConvexDecomposition3::MergeBest(int32 InTargetNumParts, double MaxErrorTo
 				Decomposition[Proximities[ProxIdx].Link.B].bMustMerge;
 
 			double MergeCost;
+			// attempt to remove parts from the proximity store if we're at the size limit
+			bool bShouldRemoveParts = ProximityComputedParts.Num() > EvictComputedPartsThreshold;
 			if (Proximities[ProxIdx].HasMergedVolume())
 			{
 				MergeCost = Proximities[ProxIdx].GetMergeCost(Decomposition);
@@ -1895,6 +1958,11 @@ int32 FConvexDecomposition3::MergeBest(int32 InTargetNumParts, double MaxErrorTo
 					bOnlyAllowMustMerges // or we're currently only considering must-merge parts (e.g., already at or below desired number of parts)
 				))
 			{
+				// remove pre-computed part if they won't pass volume tolerance or we're at the part limit
+				if (bShouldRemoveParts || (!bAboveMax && VolumeTolerance > 0 && MergeCost > VolumeTolerance))
+				{
+					ProximityComputedParts.Remove(ProxIdx);
+				}
 				continue;
 			}
 
@@ -1925,6 +1993,8 @@ int32 FConvexDecomposition3::MergeBest(int32 InTargetNumParts, double MaxErrorTo
 				{
 					// Mark link as invalid so we don't repeat this overlap test in later iterations
 					Proximities[ProxIdx].bIsValidLink = false;
+					// Never keep precomputed proximity for invalid links
+					ProximityComputedParts.Remove(ProxIdx);
 					continue;
 				}
 			}
@@ -1934,6 +2004,10 @@ int32 FConvexDecomposition3::MergeBest(int32 InTargetNumParts, double MaxErrorTo
 				BestKnownCost = MergeCost;
 				BestKnownIdx = ProxIdx;
 			}
+			else if (bShouldRemoveParts) // if we are at the precomputed proximity part limit and it's not the best part, don't keep it
+			{
+				ProximityComputedParts.Remove(ProxIdx);
+			}
 		}
 
 
diff --git a/Engine/Source/Runtime/GeometryCore/Public/CompGeom/ConvexDecomposition3.h b/Engine/Source/Runtime/GeometryCore/Public/CompGeom/ConvexDecomposition3.h
--- a/Engine/Source/Runtime/GeometryCore/Public/CompGeom/ConvexDecomposition3.h
+++ b/Engine/Source/Runtime/GeometryCore/Public/CompGeom/ConvexDecomposition3.h
@@ -169,6 +169,14 @@ public:
 	 */
 	GEOMETRYCORE_API void InitializeFromHulls(int32 NumHulls, TFunctionRef<double(int32)> HullVolumes, TFunctionRef<int32(int32)> HullNumVertices, TFunctionRef<FVector3d(int32, int32)> HullVertices, TArrayView<const TPair<int32, int32>> Proximity);
 
+	/**
+	 * Create the proximity graph from the current decomposition, using bounding box overlaps. To consider non-overlapping proximity, the bounding boxes can be expanded by a factor of their own size or by an absolute amount.
+	 * @param BoundsExpandByMinDimFactor		Part bounds will be expanded by at least this factor of their own min dimension, before finding overlaps
+	 * @param BoundsExpandByMaxDimFactor		Part bounds will be expanded by at least this factor of their own max dimension, before finding overlaps
+	 * @param MinBoundsExpand					Part bounds will be expanded by at least this fixed amount, before finding overlaps
+	 */
+	GEOMETRYCORE_API void InitializeProximityFromDecompositionBoundingBoxOverlaps(double BoundsExpandByMinDimFactor, double BoundsExpandByMaxDimFactor, double MinBoundsExpand);
+
 	GEOMETRYCORE_API void InitializeFromMesh(const FDynamicMesh3& SourceMesh, bool bMergeEdges);
 
 	/**
@@ -210,6 +218,10 @@ public:
 	// Larger values will cost more to run but can let the algorithm find a cleaner decomposition
 	int32 MaxConvexEdgePlanes = 50;
 
+
+	// If > 0, search for best merges will be restricted to a greedy local search after testing this many connections. Helpful when there are many potential merges.
+	int32 RestrictMergeSearchToLocalAfterTestNumConnections = -1;
+
 	// TODO: Provide hull approximation options?
 
 
