Add support for negative space protection in the 'splitting' part of the convex decomposition algorithm + associated initialization / support functions

diff --git a/Engine/Source/Runtime/GeometryCore/Private/CompGeom/ConvexDecomposition3.cpp b/Engine/Source/Runtime/GeometryCore/Private/CompGeom/ConvexDecomposition3.cpp
--- a/Engine/Source/Runtime/GeometryCore/Private/CompGeom/ConvexDecomposition3.cpp
+++ b/Engine/Source/Runtime/GeometryCore/Private/CompGeom/ConvexDecomposition3.cpp
@@ -24,6 +24,86 @@ namespace UE
 namespace Geometry
 {
 
+bool FConvexDecomposition3::ConvexPartVsSphereOverlap(const FConvexDecomposition3::FConvexPart& Part, FVector3d Center, double Radius, const FTransform* TransformIntoSphereSpace, double* OutDistanceSq)
+{
+	if (OutDistanceSq)
+	{
+		*OutDistanceSq = 0;
+	}
+	bool bMustTransformConvex = false;
+	if (TransformIntoSphereSpace)
+	{
+		FVector ScaleVec = TransformIntoSphereSpace->GetScale3D();
+		// Uniform scale: Can transform into sphere into part space
+		if (ScaleVec.AllComponentsEqual())
+		{
+			if (ScaleVec.X == 0)
+			{
+				return false;
+			}
+			Radius = FMath::Abs(Radius / ScaleVec.X);
+			Center = TransformIntoSphereSpace->InverseTransformPosition(Center);
+		}
+		else // Non-uniform scale: Must transform convex into sphere space
+		{
+			bMustTransformConvex = true;
+		}
+	}
+
+	// Note: Could test vs FConvexPart's bounding box here for an early out (Especially if !bMustTransformConvex and/or Part.HullPlanes.Num() is large!)
+
+	double ClosestRadiusSq = Radius * Radius;
+	double MaxPlaneDist = 0;
+	bool bFoundDistSq = false;
+	for (int32 PlaneIdx = 0; PlaneIdx < Part.HullPlanes.Num(); ++PlaneIdx)
+	{
+		FPlane3d Plane = Part.HullPlanes[PlaneIdx];
+		if (bMustTransformConvex)
+		{
+			Plane.Transform(*TransformIntoSphereSpace);
+		}
+		double PlaneDist = Plane.DistanceTo(Center);
+		MaxPlaneDist = FMath::Max(MaxPlaneDist, PlaneDist);
+		if (PlaneDist > Radius)
+		{
+			return false;
+		}
+		else if (Radius > 0 && PlaneDist > FMath::Max(0, MaxPlaneDist - FMathd::ZeroTolerance))
+		{
+			FIndex3i TriInds = Part.HullTriangles[PlaneIdx];
+			FTriangle3d Tri(Part.InternalGeo.GetVertex(TriInds.A), Part.InternalGeo.GetVertex(TriInds.B), Part.InternalGeo.GetVertex(TriInds.C));
+			if (bMustTransformConvex)
+			{
+				Tri.V[0] = TransformIntoSphereSpace->TransformPosition(Tri.V[0]);
+				Tri.V[1] = TransformIntoSphereSpace->TransformPosition(Tri.V[1]);
+				Tri.V[2] = TransformIntoSphereSpace->TransformPosition(Tri.V[2]);
+			}
+			// TODO: Can optimize this by writing custom Point-Tri distance logic to re-use what we already know and early-out if a vertex or edge distance is < Radius
+			FDistPoint3Triangle3d Dist(Center, Tri);
+			double TriDistSq = Dist.GetSquared();
+			if (TriDistSq < ClosestRadiusSq)
+			{
+				if (OutDistanceSq)
+				{
+					ClosestRadiusSq = TriDistSq;
+					bFoundDistSq = true;
+				}
+				else
+				{
+					return true;
+				}
+			}
+		}
+	}
+	if (OutDistanceSq && bFoundDistSq)
+	{
+		*OutDistanceSq = ClosestRadiusSq;
+		return true;
+	}
+	checkSlow(!OutDistanceSq || *OutDistanceSq == 0); // If we reach here distance should still be zero (as set by default)
+	return (MaxPlaneDist <= 0);
+}
+
 bool FSphereCovering::AddNegativeSpace(const TFastWindingTree<FDynamicMesh3>& Spatial, const FNegativeSpaceSampleSettings& SampleSettings, bool bHasFlippedTriangles)
 {
 	bool bAddedPoints = false;
@@ -144,8 +224,14 @@ bool FSphereCovering::AddNegativeSpace(const TFastWindingTree<FDynamicMesh3>& Sp
 
 		// Compute the empty space inside the convex hull as (Convex Hull - Original Mesh)
 		FAxisAlignedBox3d HullBox = HullAABB.GetBoundingBox();
-		HullBox.Expand(1);
+		// expand by 1 unit (rescaled as needed) to make sure we have some empty space at the boundary
+		// TODO: Test if a smaller expand (e.g., KINDA_SMALL_NUMBER) will do just as well here
+		HullBox.Expand(1.0 * SampleSettings.GetAppliedScaleFactor());
 		FMarchingCubes MarchingCubes;
+		if (SampleSettings.bDeterministic)
+		{
+			MarchingCubes.bParallelCompute = false;
+		}
 		MarchingCubes.CubeSize = FMath::Clamp(SampleSettings.ReduceRadiusMargin * .5, HullBox.MaxDim() / (double)SampleSettings.MaxVoxelsPerDim, HullBox.MinDim() * .5);
 		MarchingCubes.Bounds = HullBox;
 		MarchingCubes.RootMode = ERootfindingModes::Bisection;
@@ -467,7 +553,7 @@ struct FPartialCutResult
 		return HullMinusGeoVolume;
 	}
 
-	void ApplyToGeo(TIndirectArray<FConvexDecomposition3::FConvexPart>& Decomposition, int32 OrigPartIdx, const FPlane3d& Plane, int32& OtherSideStartIdxOut, double PlaneTol, double ConnectedComponentTolerance)
+	void ApplyToGeo(TIndirectArray<FConvexDecomposition3::FConvexPart>& Decomposition, int32 OrigPartIdx, const FPlane3d& Plane, int32& OtherSideStartIdxOut, double PlaneTol, double ConnectedComponentTolerance, const FSphereCovering& NegativeSpace)
 	{
 		FConvexDecomposition3::FConvexPart& OrigPart = Decomposition[OrigPartIdx];
 		bool bSourceGeometryVolumeUnreliable = OrigPart.bGeometryVolumeUnreliable;
@@ -905,6 +991,28 @@ struct FPartialCutResult
 				Decomposition[PartIdx].bGeometryVolumeUnreliable = true;
 			}
 		}
+
+		// Update lists of conflicting negative spaces (if negative space is available)
+		if (NegativeSpace.Num() > 0 && !OrigPart.OverlapsNegativeSpace.IsEmpty())
+		{
+			TArray<int32> OrigOverlaps = MoveTemp(OrigPart.OverlapsNegativeSpace);
+			auto AddIfOverlaps = [NegativeSpace, OrigOverlaps](FConvexDecomposition3::FConvexPart& Part)
+			{
+				for (int32 SphereIdx : OrigOverlaps)
+				{
+					if (FConvexDecomposition3::ConvexPartVsSphereOverlap(Part, NegativeSpace.GetCenter(SphereIdx), NegativeSpace.GetRadius(SphereIdx)))
+					{
+						Part.OverlapsNegativeSpace.Add(SphereIdx);
+					}
+				}
+			};
+
+			AddIfOverlaps(OrigPart);
+			for (int32 PartIdx = NewPartIdx; PartIdx < Decomposition.Num(); ++PartIdx)
+			{
+				AddIfOverlaps(Decomposition[PartIdx]);
+			}
+		}
 	}
 
 	FPartialCutResult() {}
@@ -1154,6 +1262,88 @@ void FConvexDecomposition3::InitializeFromIndexMesh(TArrayView<const FVector3f>
 	Decomposition.Add(Convex);
 }
 
+
+bool FConvexDecomposition3::InitializeNegativeSpace(const FNegativeSpaceSampleSettings& Settings)
+{
+	if (Decomposition.Num() != 1)
+	{
+		return false;
+	}
+	TMeshAABBTree3<FDynamicMesh3> InternalGeoAABBTree(&Decomposition[0].InternalGeo, true);
+	TFastWindingTree<FDynamicMesh3> InternalGeoWinding(&InternalGeoAABBTree, true);
+
+	FNegativeSpaceSampleSettings RescaledSettings = Settings;
+	RescaledSettings.Rescale(1.0 / ResultTransform.GetScale().X);
+	NegativeSpace.Reset();
+	NegativeSpace.AddNegativeSpace(InternalGeoWinding, RescaledSettings, false);
+	InitNegativeSpaceConvexPartMapping();
+	return true;
+}
+
+
+void FConvexDecomposition3::FixHullOverlapsInNegativeSpace(double NegativeSpaceTolerance, double NegativeSpaceMinRadius)
+{
+	if (!NegativeSpace.Num())
+	{
+		return;
+	}
+
+	// make sure tolerance and min radius are at least 0
+	double UseTolerance = FMath::Max(0, NegativeSpaceTolerance);
+	double UseMinRadius = FMath::Max(0, NegativeSpaceMinRadius);
+	// fix overlaps
+	for (int32 PartIdx = 0; PartIdx < Decomposition.Num(); ++PartIdx)
+	{
+		for (int32 SphereIdx : Decomposition[PartIdx].OverlapsNegativeSpace)
+		{
+			double DistSq = 0;
+			double Radius = NegativeSpace.GetRadius(SphereIdx);
+			if (Radius > UseMinRadius && ConvexPartVsSphereOverlap(Decomposition[PartIdx], NegativeSpace.GetCenter(SphereIdx), Radius, nullptr, &DistSq))
+			{
+				double Dist = FMath::Sqrt(DistSq);
+				check(Dist <= Radius);
+				NegativeSpace.SetRadius(SphereIdx, Dist - UseTolerance);
+			}
+		}
+		Decomposition[PartIdx].OverlapsNegativeSpace.Reset();
+	}
+	// delete all too-small spheres
+	NegativeSpace.RemoveSmaller(UseMinRadius);
+}
+
+void FConvexDecomposition3::Compute(int32 NumOutputHulls, int32 NumAdditionalSplits, double ErrorTolerance, double MinThicknessTolerance, int32 MaxOutputHulls, bool bOnlySplitIfNegativeSpaceCovered)
+{
+	if (MaxOutputHulls > 0)
+	{
+		NumOutputHulls = FMath::Min(NumOutputHulls, MaxOutputHulls);
+	}
+	bool bUseNegativeSpace = NegativeSpace.Num() > 0;
+	if (bUseNegativeSpace)
+	{
+		InitNegativeSpaceConvexPartMapping();
+	}
+	int32 TargetNumSplits = NumOutputHulls + NumAdditionalSplits;
+	for (int32 SplitIdx = 0; SplitIdx < TargetNumSplits; SplitIdx++)
+	{
+		int32 NumNewParts = SplitWorst(bool(SplitIdx % 2), ErrorTolerance, bOnlySplitIfNegativeSpaceCovered);
+		if (NumNewParts == 0)
+		{
+			break;
+		}
+	}
+
+	if (bUseNegativeSpace)
+	{
+		// remove overlaps with negative space that the splitting didn't resolve by just shrinking/deleting the associated spheres
+		// (b/c otherwise the overlapped parts will be incapable of merging with anything in the MergeBest step)
+		FixHullOverlapsInNegativeSpace();
+	}
+
+	constexpr bool bAllowCompact = true;
+	MergeBest(NumOutputHulls, ErrorTolerance, MinThicknessTolerance, bAllowCompact, false, MaxOutputHulls, nullptr, nullptr);
+}
+
+
 bool FConvexDecomposition3::FConvexPart::ComputeHull(bool bComputePlanes)
 {
 	FConvexHull3d HullCompute;
@@ -1202,7 +1392,7 @@ void FConvexDecomposition3::FConvexPart::ComputeStats()
 	HullError = HullVolume - GeoVolume;
 }
 
-int32 FConvexDecomposition3::SplitWorst(bool bCanSkipUnreliableGeoVolumes, double ErrorTolerance)
+int32 FConvexDecomposition3::SplitWorst(bool bCanSkipUnreliableGeoVolumes, double ErrorTolerance, bool bOnlySplitIfNegativeSpaceCovered)
 {
 	if (Decomposition.Num() == 0)
 	{
@@ -1211,12 +1401,36 @@ int32 FConvexDecomposition3::SplitWorst(bool bCanSkipUnreliableGeoVolumes, doubl
 
 	double VolumeTolerance = ConvertDistanceToleranceToLocalVolumeTolerance(ErrorTolerance);
 
-	double WorstError = Decomposition[0].HullError;
-	bool bErrorAboveTolerance = VolumeTolerance <= 0 || Decomposition[0].HullError > VolumeTolerance;
-	int32 WorstIdx = 0;
-	bool bCanSkipCurrent = bCanSkipUnreliableGeoVolumes && Decomposition[0].bGeometryVolumeUnreliable;
-	for (int32 PartIdx = 1; PartIdx < Decomposition.Num(); PartIdx++)
+	double WorstError = -FMathd::MaxReal;
+	bool bErrorAboveTolerance = false;
+	int32 WorstIdx = INDEX_NONE;
+	bool bCanSkipCurrent = true;
+	bool bHasOverlapsNegative = false;
+	// check if we have any negative-space-covered parts
+	if (NegativeSpace.Num() > 0)
+	{
+		for (int32 PartIdx = 0; PartIdx < Decomposition.Num(); PartIdx++)
+		{
+			if (!Decomposition[PartIdx].OverlapsNegativeSpace.IsEmpty())
+			{
+				bHasOverlapsNegative = true;
+				break;
+			}
+		}
+	}
+	// stop early if there are no negative-space overlaps, and we're only splitting in those cases
+	if (!bHasOverlapsNegative && bOnlySplitIfNegativeSpaceCovered)
+	{
+		return 0;
+	}
+	for (int32 PartIdx = 0; PartIdx < Decomposition.Num(); PartIdx++)
 	{
+		bool bOverlapsNegative = !Decomposition[PartIdx].OverlapsNegativeSpace.IsEmpty();
+		// always favor the parts that overlap negative space as long as we have them
+		if (bHasOverlapsNegative && !bOverlapsNegative)
+		{
+			continue;
+		}
 		if (Decomposition[PartIdx].HullError > VolumeTolerance)
 		{
 			bErrorAboveTolerance = true;
@@ -1232,12 +1446,17 @@ int32 FConvexDecomposition3::SplitWorst(bool bCanSkipUnreliableGeoVolumes, doubl
 			bCanSkipCurrent = bCanSkipCurrent && Decomposition[PartIdx].bGeometryVolumeUnreliable;
 		}
 	}
+	
+	if (WorstIdx == INDEX_NONE)
+	{
+		return 0;
+	}
 
-	// Stop splitting if we see no errors above tolerance
+	// Stop splitting if we see no errors above tolerance and no negative-space overlaps
 	// Note if bCanSkipUnreliableGeoVolumes==true, we may still split a below-tolerance part
 	//  -- TODO: consider what to do in this case!
 	//		It will just add one extra split currently as we alternate toggling bCanSkipUnreliableGeoVolumes off, but the behavior may be confusing.
-	if (!bErrorAboveTolerance)
+	if (!bErrorAboveTolerance && !bHasOverlapsNegative)
 	{
 		return 0;
 	}
@@ -1428,7 +1647,7 @@ int32 FConvexDecomposition3::SplitWorst(bool bCanSkipUnreliableGeoVolumes, doubl
 
 	int32 NewPartsStartIdx = Decomposition.Num();
 	int32 OtherSideStartIdx = -1;
-	BestCutResult.ApplyToGeo(Decomposition, WorstIdx, CandidatePlanes[BestPlaneIdx], OtherSideStartIdx, OnPlaneTolerance, ConnectedComponentTolerance);
+	BestCutResult.ApplyToGeo(Decomposition, WorstIdx, CandidatePlanes[BestPlaneIdx], OtherSideStartIdx, OnPlaneTolerance, ConnectedComponentTolerance, NegativeSpace);
 
 	UpdateProximitiesAfterSplit(WorstIdx, NewPartsStartIdx, CandidatePlanes[BestPlaneIdx], OtherSideStartIdx, OrigHullVolume);
 
@@ -1663,69 +1882,6 @@ int32 FConvexDecomposition3::MergeBest(int32 InTargetNumParts, double MaxErrorTo
 	// threshold at which we will start to more aggressively evict cached computed parts from ProximityComputedParts
 	constexpr int32 EvictComputedPartsThreshold = 10000;
 
-	auto ConvexNegativeSpaceSphereOverlapTest = [OptionalNegativeSpace, OptionalTransformIntoNegativeSpace](FConvexPart* Part, int PtIdx)
-	{
-		check(OptionalNegativeSpace);
-		double Radius = OptionalNegativeSpace->GetRadius(PtIdx);
-		FVector3d Center = OptionalNegativeSpace->GetCenter(PtIdx);
-
-		bool bMustTransformConvex = false;
-		if (OptionalTransformIntoNegativeSpace)
-		{
-			FVector ScaleVec = OptionalTransformIntoNegativeSpace->GetScale3D();
-			// Uniform scale: Can tranform into negative space into part space
-			if (ScaleVec.AllComponentsEqual())
-			{
-				if (ScaleVec.X == 0)
-				{
-					return false;
-				}
-				Radius = FMath::Abs(Radius / ScaleVec.X);
-				Center = OptionalTransformIntoNegativeSpace->InverseTransformPosition(Center);
-			}
-			else // Non-uniform scale: Must transform convex into negative space
-			{
-				bMustTransformConvex = true;
-			}
-		}
-		
-		double RadiusSq = Radius * Radius;
-		double MaxPlaneDist = 0;
-		for (int32 PlaneIdx = 0; PlaneIdx < Part->HullPlanes.Num(); ++PlaneIdx)
-		{
-			FPlane3d Plane = Part->HullPlanes[PlaneIdx];
-			if (bMustTransformConvex)
-			{
-				Plane.Transform(*OptionalTransformIntoNegativeSpace);
-			}
-			double PlaneDist = Plane.DistanceTo(Center);
-			MaxPlaneDist = FMath::Max(MaxPlaneDist, PlaneDist);
-			if (PlaneDist > Radius)
-			{
-				return false;
-			}
-			else if (Radius > 0 && PlaneDist > FMath::Max(0, MaxPlaneDist - FMathd::ZeroTolerance))
-			{
-				FIndex3i TriInds = Part->HullTriangles[PlaneIdx];
-				FTriangle3d Tri(Part->InternalGeo.GetVertex(TriInds.A), Part->InternalGeo.GetVertex(TriInds.B), Part->InternalGeo.GetVertex(TriInds.C));
-				if (bMustTransformConvex)
-				{
-					Tri.V[0] = OptionalTransformIntoNegativeSpace->TransformPosition(Tri.V[0]);
-					Tri.V[1] = OptionalTransformIntoNegativeSpace->TransformPosition(Tri.V[1]);
-					Tri.V[2] = OptionalTransformIntoNegativeSpace->TransformPosition(Tri.V[2]);
-				}
-				// TODO: Can optimize this by writing custom Point-Tri distance logic to re-use what we already know and early-out if a vertex or edge distance is < Radius
-				FDistPoint3Triangle3d Dist(Center, Tri);
-				double TriDistSq = Dist.GetSquared();
-				if (TriDistSq < RadiusSq)
-				{
-					return true;
-				}
-			}
-		}
-		return (MaxPlaneDist <= 0);
-	};
-
 	auto IsPartBelowSizeTolerance = [MinThicknessTolerance](const FConvexPart& Part) -> bool
 	{ 
 		if (MinThicknessTolerance <= 0) // tolerance is disabled
@@ -1833,7 +1989,7 @@ int32 FConvexDecomposition3::MergeBest(int32 InTargetNumParts, double MaxErrorTo
 						double OrigHullVolume = PartToSplit.HullVolume;
 						int32 NewPartsStartIdx = Decomposition.Num();
 						int32 OtherSideStartIdx = -1;
-						PlaneResult.ApplyToGeo(Decomposition, PartIdx, Plane, OtherSideStartIdx, OnPlaneTolerance, ConnectedComponentTolerance);
+						PlaneResult.ApplyToGeo(Decomposition, PartIdx, Plane, OtherSideStartIdx, OnPlaneTolerance, ConnectedComponentTolerance, NegativeSpace);
 						UpdateProximitiesAfterSplit(PartIdx, NewPartsStartIdx, Plane, OtherSideStartIdx, OrigHullVolume);
 						PartToSplit.bMustMerge = true;
 						for (int32 NewIdx = NewPartsStartIdx; NewIdx < Decomposition.Num(); NewIdx++)
@@ -1973,7 +2129,8 @@ int32 FConvexDecomposition3::MergeBest(int32 InTargetNumParts, double MaxErrorTo
 			}
 
 			// if this would be the new best candidate, and we have negative space, check if the merge would cross protected negative space
-			if (OptionalNegativeSpace && OptionalNegativeSpace->Num() > 0 && MergeCost < BestKnownCost)
+			bool bHasNegativeSpace = NegativeSpace.Num() || (OptionalNegativeSpace && OptionalNegativeSpace->Num() > 0);
+			if (bHasNegativeSpace && MergeCost < BestKnownCost)
 			{
 				TUniquePtr<FConvexPart>* FoundMergeHull = ProximityComputedParts.Find(ProxIdx);
 				if (!FoundMergeHull)
@@ -1981,14 +2138,30 @@ int32 FConvexDecomposition3::MergeBest(int32 InTargetNumParts, double MaxErrorTo
 					FoundMergeHull = &ProximityComputedParts.Add(ProxIdx, CreateMergedPart(Proximities[ProxIdx]));
 				}
 				bool bCoversProtectedPt = false;
-				for (int32 PtIdx = 0; PtIdx < OptionalNegativeSpace->Num(); ++PtIdx)
+				auto TestVsNegativeSpace = [](const FConvexPart& Part, const FSphereCovering& SphereCovering, const FTransform* OptionalTransform = nullptr) -> bool
+				{
+					for (int32 SphereIdx = 0; SphereIdx < SphereCovering.Num(); ++SphereIdx)
+					{
+						double Radius = SphereCovering.GetRadius(SphereIdx);
+						FVector3d Center = SphereCovering.GetCenter(SphereIdx);
+						if (ConvexPartVsSphereOverlap(Part, Center, Radius, OptionalTransform))
+						{
+							return true;
+						}
+					}
+					return false;
+				};
+				if (OptionalNegativeSpace)
 				{
-					if (ConvexNegativeSpaceSphereOverlapTest(FoundMergeHull->Get(), PtIdx))
+					if (TestVsNegativeSpace(**FoundMergeHull, *OptionalNegativeSpace, OptionalTransformIntoNegativeSpace))
 					{
 						bCoversProtectedPt = true;
-						break;
 					}
 				}
+				if (TestVsNegativeSpace(**FoundMergeHull, NegativeSpace))
+				{
+					bCoversProtectedPt = true;
+				}
 				if (bCoversProtectedPt)
 				{
 					// Mark link as invalid so we don't repeat this overlap test in later iterations
@@ -2167,6 +2340,59 @@ void FConvexDecomposition3::DeleteProximity(int32 ProxIdx, bool bDeleteMapRefere
 	}
 }
 
+void FConvexDecomposition3::InitNegativeSpaceConvexPartMapping()
+{
+	// Helper to check each sphere for overlap vs each convex part
+	auto IteratePartsVsSpheres = [this](TFunctionRef<void(FConvexPart& Part, int32 SphereIdx, bool bOverlapsPart)> ProcessPartVsSphere)
+	{
+		for (int32 PartIdx = 0; PartIdx < Decomposition.Num(); ++PartIdx)
+		{
+			FConvexPart& Part = Decomposition[PartIdx];
+			if (!ensure(Part.OverlapsNegativeSpace.IsEmpty())) // expect this Init to only be called when there is no existing mapping
+			{
+				Part.OverlapsNegativeSpace.Reset();
+			}
+			for (int32 SphereIdx = 0; SphereIdx < NegativeSpace.Num(); ++SphereIdx)
+			{
+				// Note: We assume the part is in the same coordinate space as the negative space spheres
+				bool bOverlapsPart = ConvexPartVsSphereOverlap(Part, NegativeSpace.GetCenter(SphereIdx), NegativeSpace.GetRadius(SphereIdx));
+				ProcessPartVsSphere(Part, SphereIdx, bOverlapsPart);
+			}
+		}
+	};
+
+	if (Decomposition.Num() == 1)
+	{
+		// Special handling of the single-part case, since any spheres that don't overlap that must be useless (future splits/merges will never be affected by them)
+		IteratePartsVsSpheres([this](FConvexPart& Part, int32 SphereIdx, bool bOverlapsPart)
+		{
+			if (!bOverlapsPart)
+			{
+				// Set negative radius to indicate sphere is not useful
+				NegativeSpace.SetRadius(SphereIdx, -FMathd::MaxReal);
+			}
+		});
+		// Clear negative-radius / invalid spheres from the NegativeSpace
+		NegativeSpace.RemoveSmaller(0);
+		// Add all remaining spheres
+		Decomposition[0].OverlapsNegativeSpace.SetNumUninitialized(NegativeSpace.Num());
+		for (int32 Idx = 0; Idx < NegativeSpace.Num(); ++Idx)
+		{
+			Decomposition[0].OverlapsNegativeSpace[Idx] = Idx;
+		}
+	}
+	else
+	{
+		// If there are multiple parts, just assign them to parts based on overlap
+		IteratePartsVsSpheres([](FConvexPart& Part, int32 SphereIdx, bool bOverlapsPart)
+		{
+			if (bOverlapsPart)
+			{
+				Part.OverlapsNegativeSpace.Add(SphereIdx);
+			}
+		});
+	}
+}
 
 } // end namespace UE::Geometry
 } // end namespace UE
\ No newline at end of file
diff --git a/Engine/Source/Runtime/GeometryCore/Public/CompGeom/ConvexDecomposition3.h b/Engine/Source/Runtime/GeometryCore/Public/CompGeom/ConvexDecomposition3.h
--- a/Engine/Source/Runtime/GeometryCore/Public/CompGeom/ConvexDecomposition3.h
+++ b/Engine/Source/Runtime/GeometryCore/Public/CompGeom/ConvexDecomposition3.h
@@ -72,6 +72,22 @@ struct FNegativeSpaceSampleSettings
 	bool bOnlyConnectedToHull = false;
 	// Maximum number of voxels to use per dimension, when performing VoxelSearch
 	int32 MaxVoxelsPerDim = 128;
+	// Attempt to keep negative space computation deterministic (e.g., do not run marching cubes in parallel for voxel search)
+	bool bDeterministic = false;
+
+	// @return the scale factor that has been applied by Rescale()
+	double GetAppliedScaleFactor() const
+	{
+		return AppliedScaleFactor;
+	}
+
+	void Rescale(double ScaleFactor)
+	{
+		MinSpacing *= ScaleFactor;
+		ReduceRadiusMargin *= ScaleFactor;
+		MinRadius *= ScaleFactor;
+		AppliedScaleFactor *= ScaleFactor;
+	}
 
 	// Make sure the settings values are in valid ranges
 	void Sanitize()
@@ -82,6 +98,11 @@ struct FNegativeSpaceSampleSettings
 		MinRadius = FMath::Max(0.0, MinRadius);
 		MaxVoxelsPerDim = FMath::Clamp(MaxVoxelsPerDim, 4, 4096);
 	}
+	
+private:
+
+	// Track how the settings have been rescaled
+	double AppliedScaleFactor = 1;
 };
 
 // Define a volume with a set of spheres
@@ -104,6 +125,11 @@ public:
 		return Radius[Idx];
 	}
 
+	inline void SetRadius(int32 Idx, double UpdateRadius)
+	{
+		Radius[Idx] = UpdateRadius;
+	}
+
 	// Add spheres covering the negative space of the given fast winding tree
 	// @param Spatial	Fast winding tree of the reference mesh, used to compute negative space
 	// @param Settings	Settings controlling how negative space is computed
@@ -130,6 +156,20 @@ public:
 		Radius.Append(Other.Radius);
 	}
 
+	// Remove spheres with radius below a threshold
+	void RemoveSmaller(double MinRadius)
+	{
+		for (int32 Idx = 0; Idx < Radius.Num(); ++Idx)
+		{
+			if (Radius[Idx] < MinRadius)
+			{
+				Radius.RemoveAtSwap(Idx, 1, false);
+				Position.RemoveAtSwap(Idx, 1, false);
+				--Idx;
+			}
+		}
+	}
+
 	void AppendSpheres(TArrayView<const FSphere> Spheres)
 	{
 		const int32 AddCount = Spheres.Num();
@@ -195,6 +235,13 @@ public:
 	 */
 	GEOMETRYCORE_API void InitializeFromIndexMesh(TArrayView<const FVector3f> Vertices, TArrayView<const FIntVector> Faces, bool bMergeEdges, int32 FaceVertexOffset = 0);
 
+	/**
+	 * Find negative space that should be protected. Uses the mesh passed on construction or to InitializeFromMesh or InitializeFromIndexMesh
+	 * @param Settings	Settings to use to find the negative space
+	 * @return			False on failure -- e.g., if there was no mesh available
+	 */
+	GEOMETRYCORE_API bool InitializeNegativeSpace(const FNegativeSpaceSampleSettings& Settings);
+
 	//
 	// Settings
 	//
@@ -241,31 +288,15 @@ public:
 	 * @param ErrorTolerance		Stop splitting when hulls have error less than this (expressed in cm; will be cubed for volumetric error). Overrides NumOutputHulls if specified
 	 * @param MinThicknessTolerance	Optionally specify a minimum thickness (in cm) for convex parts; parts below this thickness will always be merged away. Overrides NumOutputHulls and ErrorTolerance when needed
 	 * @param MaxOutputHulls		If > 0, maximum number of convex hulls to generate. Overrides ErrorTolerance and TargetNumParts when needed
+	 * @param bOnlySplitIfOverlapNegativeSpace	If true, use NegativeSpace to guide splits, and only split parts that overlap negative space
 	 */
-	void Compute(int32 NumOutputHulls, int32 NumAdditionalSplits = 10, double ErrorTolerance = 0.0, double MinThicknessTolerance = 0, int32 MaxOutputHulls = -1)
-	{
-		if (MaxOutputHulls > 0)
-		{
-			NumOutputHulls = FMath::Min(NumOutputHulls, MaxOutputHulls);
-		}
-		int32 TargetNumSplits = NumOutputHulls + NumAdditionalSplits;
-		for (int32 SplitIdx = 0; SplitIdx < TargetNumSplits; SplitIdx++)
-		{
-			int32 NumNewParts = SplitWorst(bool(SplitIdx % 2), ErrorTolerance);
-			if (NumNewParts == 0)
-			{
-				break;
-			}
-		}
-		
-		constexpr bool bAllowCompact = true;
-		MergeBest(NumOutputHulls, ErrorTolerance, MinThicknessTolerance, bAllowCompact, false, MaxOutputHulls);
-	}
+	GEOMETRYCORE_API void Compute(int32 NumOutputHulls, int32 NumAdditionalSplits = 10, double ErrorTolerance = 0.0, double MinThicknessTolerance = 0, int32 MaxOutputHulls = -1, bool bOnlySplitIfNegativeSpaceCovered = false);
 
 	// Split the worst convex part, and return the increase in the total number of convex parts after splitting (can be more than 1 if result has multiple separate connected components)
 	// Note: could return 0 if no splits were possible
 	// @param bCanSkipUnreliableGeoVolumes		if true, don't split hulls where we have questionable geometry volume results, unless there is no hull with good geometry volume results
-	GEOMETRYCORE_API int32 SplitWorst(bool bCanSkipUnreliableGeoVolumes = false, double ErrorTolerance = 0.0);
+	// @param bOnlySplitIfNegativeSpaceCovered	if true, don't split hulls unless they overlap with some covered Negative Space (stored in the corresponding member variable)
+	GEOMETRYCORE_API int32 SplitWorst(bool bCanSkipUnreliableGeoVolumes = false, double ErrorTolerance = 0.0, bool bOnlySplitIfNegativeSpaceCovered = false);
 
 	// Merge the pairs of convex hulls in the decomposition that will least increase the error.  Intermediate results can be used across merges, so it is best to do all merges in one call.
 	// Note: A future version of this function may replace NumOutputHulls with MaxOutputHulls, but this version keeps both parameters for compatibility / consistent behavior.
@@ -403,6 +434,7 @@ public:
 			HullTriangles.Reset();
 			HullPlanes.Reset();
 			HullSourceID = -1;
+			OverlapsNegativeSpace.Reset();
 			HullVolume = 0;
 			GeoVolume = 0;
 			SumHullsVolume = -FMathd::MaxReal;
@@ -434,6 +466,9 @@ public:
 
 		int32 HullSourceID = -1; // Optional ID, cleared on merge, to track the origin of un-merged hulls
 
+		// Indices of NegativeSpace spheres that are overlapped by the part
+		TArray<int32> OverlapsNegativeSpace;
+
 		// Measurements of the geo and hull, to be used when evaluating potential further splits
 		double HullVolume = 0, GeoVolume = 0;
 		double SumHullsVolume = -FMathd::MaxReal; // Sum of volume of any hulls that have been merged to form this hull; only valid if the FConvexPart was created by merging
@@ -555,6 +590,36 @@ public:
 	// @param SecondSideIdxStart	The first index in the Decomposition array of new FConvexParts that were on the opposite side of the plane from SplitIdx.  Will be different from NewIdxStart only if the first part was split into multiple components.
 	// @param OrigHullVolume		The volume of the convex hull *before* the split was performed -- used to precompute the merge cost.
 	GEOMETRYCORE_API void UpdateProximitiesAfterSplit(int32 SplitIdx, int32 NewIdxStart, FPlane3d CutPlane, int32 SecondSideIdxStart, double OrigHullVolume);
+
+	// Fix overlaps between current Decomposition and associated NegativeSpace, using the saved associations in the FConvexPart::OverlapsNegativeSpace
+	// Note the tolerance and min radius may not be the same as originally used to generate the negative space; they should just be large enough to
+	// avoid 'locking' the merge step for any hulls that still overlap negative space after splitting
+	GEOMETRYCORE_API void FixHullOverlapsInNegativeSpace(double NegativeSpaceTolerance = UE_DOUBLE_KINDA_SMALL_NUMBER, double NegativeSpaceMinRadius = UE_DOUBLE_KINDA_SMALL_NUMBER);
+	
+	// Tests if a convex part overlaps a sphere
+	// @param Part				Convex part to test for overlap with sphere
+	// @param Center			Center of sphere
+	// @param Radius			Radius of sphere
+	// @param TransformIntoSphereSpace	If non-null, transform from convex hull to sphere coordinate space. Otherwise, hulls and spheres are assumed to be in the same space.
+	// @param OutDistanceSq		If non-null, and there is an overlap, will be filled with the squared distance from the sphere center to the convex hull (0 if the center is inside the hull). If no overlap is found, value is not meaningful.
+	// @return true if the Part overlaps the sphere, false otherwise
+	GEOMETRYCORE_API static bool ConvexPartVsSphereOverlap(const FConvexPart& Part, FVector3d Center, double Radius, const FTransform* OptionalTransformIntoSphereSpace = nullptr, double* OutDistanceSq = nullptr);
+
+	// Get the current negative space tracked by the convex decomposition
+	// Note: Does not include externally-managed negative space passed to MergeBest.
+	// Note: Direct access is const-only; use InitializeNegativeSpace() to set the negative space.
+	const FSphereCovering& GetNegativeSpace() const
+	{
+		return NegativeSpace;
+	}
+
+private:
+	
+	// Initialize mappings from Decomposition parts to NegativeSpace indices (FConvexPart::OverlapsNegativeSpace)
+	void InitNegativeSpaceConvexPartMapping();
+
+	// Negative space to attempt to protect; can be referenced by the Decomposition parts
+	FSphereCovering NegativeSpace;
 };
 
 } // end namespace UE::Geometry
