in convex decomposition helper function, fix mismatch between parameters listed in comment and actual parameters in method

diff --git a/Engine/Source/Runtime/GeometryCore/Public/CompGeom/ConvexDecomposition3.h b/Engine/Source/Runtime/GeometryCore/Public/CompGeom/ConvexDecomposition3.h
--- a/Engine/Source/Runtime/GeometryCore/Public/CompGeom/ConvexDecomposition3.h
+++ b/Engine/Source/Runtime/GeometryCore/Public/CompGeom/ConvexDecomposition3.h
@@ -603,7 +603,7 @@ public:
 	// @param TransformIntoSphereSpace	If non-null, transform from convex hull to sphere coordinate space. Otherwise, hulls and spheres are assumed to be in the same space.
 	// @param OutDistanceSq		If non-null, and there is an overlap, will be filled with the squared distance from the sphere center to the convex hull (0 if the center is inside the hull). If no overlap is found, value is not meaningful.
 	// @return true if the Part overlaps the sphere, false otherwise
-	GEOMETRYCORE_API static bool ConvexPartVsSphereOverlap(const FConvexPart& Part, FVector3d Center, double Radius, const FTransform* OptionalTransformIntoSphereSpace = nullptr, double* OutDistanceSq = nullptr);
+	GEOMETRYCORE_API static bool ConvexPartVsSphereOverlap(const FConvexPart& Part, FVector3d Center, double Radius, const FTransform* TransformIntoSphereSpace = nullptr, double* OutDistanceSq = nullptr);
 
 	// Get the current negative space tracked by the convex decomposition
 	// Note: Does not include externally-managed negative space passed to MergeBest.
