put the exact triangle intersection find-segments branch in an 'else' to prevent the compiler from complaining about unreachable code when it takes the other constexpr if branch

diff --git a/Engine/Source/Runtime/GeometryCore/Public/Intersection/ExactIntrTriangle3Triangle3.h b/Engine/Source/Runtime/GeometryCore/Public/Intersection/ExactIntrTriangle3Triangle3.h
--- a/Engine/Source/Runtime/GeometryCore/Public/Intersection/ExactIntrTriangle3Triangle3.h
+++ b/Engine/Source/Runtime/GeometryCore/Public/Intersection/ExactIntrTriangle3Triangle3.h
@@ -146,60 +146,62 @@ private:
 				ExactPredicates::Orient3<Real>(P1, R1, R2, P2) <= 0;
 			return bIntersects;
 		}
-
-		// Decision tree from Fig. 5 of Devillers & Guigue paper (see reference above)
-		if (ExactPredicates::Orient3<Real>(P1, Q1, R2, P2) <= 0)
+		else
 		{
-			if (ExactPredicates::Orient3<Real>(P1, Q1, Q2, P2) >= 0)
+			// Decision tree from Fig. 5 of Devillers & Guigue paper (see reference above)
+			if (ExactPredicates::Orient3<Real>(P1, Q1, R2, P2) <= 0)
 			{
-				if (ExactPredicates::Orient3<Real>(P1, R1, Q2, P2) < 0)
+				if (ExactPredicates::Orient3<Real>(P1, Q1, Q2, P2) >= 0)
 				{
-					// Intersection is k,j segment -- crossings of p2,q2 to p1,q1
-					TVector<Real> N1 = VectorUtil::NormalDirection(P1, Q1, R1);
-					TVector<Real> N2 = VectorUtil::NormalDirection(P2, Q2, R2);
-					OutSegA = GetEdgePlaneCrossing(P2, Q2, N1, P1);
-					OutSegB = GetEdgePlaneCrossing(P1, Q1, N2, P2);
+					if (ExactPredicates::Orient3<Real>(P1, R1, Q2, P2) < 0)
+					{
+						// Intersection is k,j segment -- crossings of p2,q2 to p1,q1
+						TVector<Real> N1 = VectorUtil::NormalDirection(P1, Q1, R1);
+						TVector<Real> N2 = VectorUtil::NormalDirection(P2, Q2, R2);
+						OutSegA = GetEdgePlaneCrossing(P2, Q2, N1, P1);
+						OutSegB = GetEdgePlaneCrossing(P1, Q1, N2, P2);
+					}
+					else
+					{
+						// Intersection is i,j segment -- crossings of p1,r1 to p1,q1
+						TVector<Real> N2 = VectorUtil::NormalDirection(P2, Q2, R2);
+						OutSegA = GetEdgePlaneCrossing(P1, R1, N2, P2);
+						OutSegB = GetEdgePlaneCrossing(P1, Q1, N2, P2);
+					}
 				}
 				else
 				{
-					// Intersection is i,j segment -- crossings of p1,r1 to p1,q1
-					TVector<Real> N2 = VectorUtil::NormalDirection(P2, Q2, R2);
-					OutSegA = GetEdgePlaneCrossing(P1, R1, N2, P2);
-					OutSegB = GetEdgePlaneCrossing(P1, Q1, N2, P2);
+					return false;
 				}
 			}
 			else
 			{
-				return false;
-			}
-		}
-		else
-		{
-			if (ExactPredicates::Orient3<Real>(P1, R1, R2, P2) <= 0)
-			{
-				if (ExactPredicates::Orient3<Real>(P1, R1, Q2, P2) <= 0)
+				if (ExactPredicates::Orient3<Real>(P1, R1, R2, P2) <= 0)
 				{
-					// Intersection is k,l segment -- crossings p2,q2 to p2, r2
-					TVector<Real> N1 = VectorUtil::NormalDirection(P1, Q1, R1);
-					OutSegA = GetEdgePlaneCrossing(P2, Q2, N1, P1);
-					OutSegB = GetEdgePlaneCrossing(P2, R2, N1, P1);
+					if (ExactPredicates::Orient3<Real>(P1, R1, Q2, P2) <= 0)
+					{
+						// Intersection is k,l segment -- crossings p2,q2 to p2, r2
+						TVector<Real> N1 = VectorUtil::NormalDirection(P1, Q1, R1);
+						OutSegA = GetEdgePlaneCrossing(P2, Q2, N1, P1);
+						OutSegB = GetEdgePlaneCrossing(P2, R2, N1, P1);
+					}
+					else
+					{
+						// Intersection is i,l segment -- crossings of p1,r1 to p2,r2
+						TVector<Real> N1 = VectorUtil::NormalDirection(P1, Q1, R1);
+						TVector<Real> N2 = VectorUtil::NormalDirection(P2, Q2, R2);
+						OutSegA = GetEdgePlaneCrossing(P1, R1, N2, P2);
+						OutSegB = GetEdgePlaneCrossing(P2, R2, N1, P1);
+					}
 				}
 				else
 				{
-					// Intersection is i,l segment -- crossings of p1,r1 to p2,r2
-					TVector<Real> N1 = VectorUtil::NormalDirection(P1, Q1, R1);
-					TVector<Real> N2 = VectorUtil::NormalDirection(P2, Q2, R2);
-					OutSegA = GetEdgePlaneCrossing(P1, R1, N2, P2);
-					OutSegB = GetEdgePlaneCrossing(P2, R2, N1, P1);
+					return false;
 				}
 			}
-			else
-			{
-				return false;
-			}
-		}
 
-		return true;
+			return true;
+		}
 	}
 
 };
