Fix convex decomposition's min sample spacing test to run after finding the final sample location, rather than running on the initial estimate location, so it will function as intended when the samples are placed far from their initial estimate locations.

diff --git a/Engine/Source/Runtime/GeometryCore/Private/CompGeom/ConvexDecomposition3.cpp b/Engine/Source/Runtime/GeometryCore/Private/CompGeom/ConvexDecomposition3.cpp
--- a/Engine/Source/Runtime/GeometryCore/Private/CompGeom/ConvexDecomposition3.cpp
+++ b/Engine/Source/Runtime/GeometryCore/Private/CompGeom/ConvexDecomposition3.cpp
@@ -317,7 +317,7 @@ bool FSphereCovering::AddNegativeSpace(const TFastWindingTree<FDynamicMesh3>& Sp
 		// Make sure the mesh is compact to simplify downsampling below
 		NegativeSpaceMesh.CompactInPlace();
 
-		auto AddSample = [this, &Mesh, &Spatial, &SampleSettings, &bAddedPoints, WindingSign](FVector3d Pos)
+		auto AddSample = [this, &Mesh, &Spatial, &SampleSettings, &bAddedPoints, WindingSign](FVector3d Pos, bool bTestCover = false)
 		{
 			double Winding = Spatial.FastWindingNumber(Pos) * WindingSign;
 			if (Winding > .5)
@@ -360,6 +360,19 @@ bool FSphereCovering::AddNegativeSpace(const TFastWindingTree<FDynamicMesh3>& Sp
 			}
 			if (R >= SampleSettings.MinRadius)
 			{
+				if (bTestCover)
+				{
+					double SpacingThresholdSq = SampleSettings.MinSpacing * SampleSettings.MinSpacing;
+					// TODO: Consider accelerating this coverage search w/ e.g. a sparse dynamic octree representation of the sphere covering
+					for (int32 SphereIdx = 0; SphereIdx < Position.Num(); ++SphereIdx)
+					{
+						double ThreshSq = FMath::Max(SpacingThresholdSq, Radius[SphereIdx] * Radius[SphereIdx]);
+						if (FVector3d::DistSquared(Position[SphereIdx], Pos) < ThreshSq)
+						{
+							return;
+						}
+					}
+				}
 				bAddedPoints = true;
 				Position.Add(Pos);
 				Radius.Add(R);
@@ -412,21 +425,7 @@ bool FSphereCovering::AddNegativeSpace(const TFastWindingTree<FDynamicMesh3>& Sp
 				{
 					int32 VID = Ordering.Order[SampleIdx];
 					FVector3d Pos = NegativeSpaceMesh.GetVertex(VID);
-					// TODO: Consider accelerating this coverage search w/ e.g. a sparse dynamic octree representation of the sphere covering
-					bool bFoundCover = false;
-					for (int32 SphereIdx = 0; SphereIdx < Position.Num(); ++SphereIdx)
-					{
-						double ThreshSq = FMath::Max(SpacingThresholdSq, Radius[SphereIdx] * Radius[SphereIdx]);
-						if (FVector3d::DistSquared(Position[SphereIdx], Pos) < ThreshSq)
-						{
-							bFoundCover = true;
-							break;
-						}
-					}
-					if (!bFoundCover)
-					{
-						AddSample(Pos);
-					}
+					AddSample(Pos, true /*test cover*/);
 				}
 			}
 		}
