remove redundant init of negative space mapping in convex decomposition on Compute(), since it is already initialized by the function which adds it

diff --git a/Engine/Source/Runtime/GeometryCore/Private/CompGeom/ConvexDecomposition3.cpp b/Engine/Source/Runtime/GeometryCore/Private/CompGeom/ConvexDecomposition3.cpp
--- a/Engine/Source/Runtime/GeometryCore/Private/CompGeom/ConvexDecomposition3.cpp
+++ b/Engine/Source/Runtime/GeometryCore/Private/CompGeom/ConvexDecomposition3.cpp
@@ -1317,10 +1317,6 @@ void FConvexDecomposition3::Compute(int32 NumOutputHulls, int32 NumAdditionalSpl
 		NumOutputHulls = FMath::Min(NumOutputHulls, MaxOutputHulls);
 	}
 	bool bUseNegativeSpace = NegativeSpace.Num() > 0;
-	if (bUseNegativeSpace)
-	{
-		InitNegativeSpaceConvexPartMapping();
-	}
 	int32 TargetNumSplits = NumOutputHulls + NumAdditionalSplits;
 	for (int32 SplitIdx = 0; SplitIdx < TargetNumSplits; SplitIdx++)
 	{
