make negative space sampling default to 'deterministic' mode, and make that mode sort the vertices after marching cubes has run rather than disabling parallelism during marching cubes.

also always run the point prioritization step, and don't early-stop it if we may take more samples than the minimum requested.

diff --git a/Engine/Source/Runtime/GeometryCore/Private/CompGeom/ConvexDecomposition3.cpp b/Engine/Source/Runtime/GeometryCore/Private/CompGeom/ConvexDecomposition3.cpp
--- a/Engine/Source/Runtime/GeometryCore/Private/CompGeom/ConvexDecomposition3.cpp
+++ b/Engine/Source/Runtime/GeometryCore/Private/CompGeom/ConvexDecomposition3.cpp
@@ -228,10 +228,6 @@ bool FSphereCovering::AddNegativeSpace(const TFastWindingTree<FDynamicMesh3>& Sp
 		// TODO: Test if a smaller expand (e.g., KINDA_SMALL_NUMBER) will do just as well here
 		HullBox.Expand(1.0 * SampleSettings.GetAppliedScaleFactor());
 		FMarchingCubes MarchingCubes;
-		if (SampleSettings.bDeterministic)
-		{
-			MarchingCubes.bParallelCompute = false;
-		}
 		MarchingCubes.CubeSize = FMath::Clamp(SampleSettings.ReduceRadiusMargin * .5, HullBox.MaxDim() / (double)SampleSettings.MaxVoxelsPerDim, HullBox.MinDim() * .5);
 		MarchingCubes.Bounds = HullBox;
 		MarchingCubes.RootMode = ERootfindingModes::Bisection;
@@ -379,62 +375,94 @@ bool FSphereCovering::AddNegativeSpace(const TFastWindingTree<FDynamicMesh3>& Sp
 			}
 		};
 
-		// if we have more vertices than we want samples, downsample so that we:
-		//  (1) ~uniformly cover space and
-		//  (2) prioritize samples at 'features' (sharp angles) of the offset mesh
-		if (NegativeSpaceMesh.MaxVertexID() > SampleSettings.TargetNumSamples)
+		// Compute an angle metric to prioritize samples on 'features' of the negative space mesh
+		TArray<float> VertexAngleMetric; // favor points on sharper edges (near 'features')
+		VertexAngleMetric.SetNumZeroed(NegativeSpaceMesh.MaxVertexID());
+		TArray<FVector3d> TriNormals;
+		TriNormals.SetNumZeroed(NegativeSpaceMesh.MaxTriangleID());
+		for (int32 TID : NegativeSpaceMesh.TriangleIndicesItr())
+		{
+			TriNormals[TID] = NegativeSpaceMesh.GetTriNormal(TID);
+		}
+		TArray<FVector3d> VertexPositions;
+		VertexPositions.SetNumZeroed(NegativeSpaceMesh.MaxVertexID());
+		for (int32 VID : NegativeSpaceMesh.VertexIndicesItr())
+		{
+			VertexPositions[VID] = NegativeSpaceMesh.GetVertex(VID);
+			// Note: Angle metric favors vertices on edges with larger dihedral angles
+			float AngleMetric = 0;
+			NegativeSpaceMesh.EnumerateVertexEdges(VID, [&](int32 EID)
+				{
+					FIndex2i EdgeT = NegativeSpaceMesh.GetEdgeT(EID);
+					if (EdgeT.B != FDynamicMesh3::InvalidID)
+					{
+						AngleMetric = FMath::Max(AngleMetric, float(1 - TriNormals[EdgeT.A].Dot(TriNormals[EdgeT.B])));
+					}
+				});
+			VertexAngleMetric[VID] = AngleMetric;
+		}
+
+		// If we want consistent results, sort the points -- since parallelism in marching cubes can add vertices in arbitrary order
+		if (SampleSettings.bDeterministic)
 		{
-			TArray<float> VertexAngleMetric;
-			VertexAngleMetric.SetNumZeroed(NegativeSpaceMesh.MaxVertexID());
-			TArray<FVector3d> TriNormals;
-			TriNormals.SetNumZeroed(NegativeSpaceMesh.MaxTriangleID());
-			for (int32 TID : NegativeSpaceMesh.TriangleIndicesItr())
+			check(NegativeSpaceMesh.IsCompactV());
+			TArray<int32> Indices;
+			Indices.SetNumUninitialized(VertexPositions.Num());
+			for (int32 Idx = 0; Idx < Indices.Num(); ++Idx)
 			{
-				TriNormals[TID] = NegativeSpaceMesh.GetTriNormal(TID);
+				Indices[Idx] = Idx;
 			}
-			TArray<FVector3d> VertexPositions;
-			VertexPositions.SetNumZeroed(NegativeSpaceMesh.MaxVertexID());
-			for (int32 VID : NegativeSpaceMesh.VertexIndicesItr())
-			{
-				VertexPositions[VID] = NegativeSpaceMesh.GetVertex(VID);
-				// Note: Angle metric favors vertices on edges with larger dihedral angles
-				float AngleMetric = 0;
-				NegativeSpaceMesh.EnumerateVertexEdges(VID, [&](int32 EID)
+			Indices.Sort([&](int32 A, int32 B)
+				{
+					FVector3d VA = VertexPositions[A];
+					FVector3d VB = VertexPositions[B];
+					if (VA.X < VB.X)
+					{
+						return true;
+					}
+					else if (VA.X == VB.X)
 					{
-						FIndex2i EdgeT = NegativeSpaceMesh.GetEdgeT(EID);
-						if (EdgeT.B != FDynamicMesh3::InvalidID)
+						if (VA.Y < VB.Y)
 						{
-							AngleMetric = FMath::Max(AngleMetric, float(1 - TriNormals[EdgeT.A].Dot(TriNormals[EdgeT.B])));
+							return true;
 						}
-					});
-				VertexAngleMetric[VID] = AngleMetric;
-			}
-			FPriorityOrderPoints Ordering;
-			Ordering.ComputeUniformSpaced(VertexPositions, VertexAngleMetric, SampleSettings.TargetNumSamples);
-			int32 NumSamples = FMath::Min(Ordering.Order.Num(), SampleSettings.TargetNumSamples);
-			for (int32 SampleIdx = 0; SampleIdx < NumSamples; ++SampleIdx)
+						else if (VA.Y == VB.Y)
+						{
+							return VA.Z < VB.Z;
+						}
+					}
+					return false;
+				});
+			TArray<FVector> SortedPositions;
+			TArray<float> SortedAngleMetric;
+			SortedPositions.SetNumUninitialized(VertexPositions.Num());
+			SortedAngleMetric.SetNumUninitialized(VertexPositions.Num());
+			for (int32 Idx = 0; Idx < VertexPositions.Num(); ++Idx)
 			{
-				int32 VID = Ordering.Order[SampleIdx];
-				AddSample(NegativeSpaceMesh.GetVertex(VID));
+				SortedPositions[Idx] = VertexPositions[Indices[Idx]];
+				SortedAngleMetric[Idx] = VertexAngleMetric[Indices[Idx]];
 			}
+			Swap(VertexPositions, SortedPositions);
+			Swap(VertexAngleMetric, SortedAngleMetric);
+		}
 
-			if (SampleSettings.bRequireSearchSampleCoverage)
-			{
-				double SpacingThresholdSq = SampleSettings.MinSpacing * SampleSettings.MinSpacing;
-				for (int32 SampleIdx = NumSamples; SampleIdx < Ordering.Order.Num(); ++SampleIdx)
-				{
-					int32 VID = Ordering.Order[SampleIdx];
-					FVector3d Pos = NegativeSpaceMesh.GetVertex(VID);
-					AddSample(Pos, true /*test cover*/);
-				}
-			}
+		int32 NumSamples = FMath::Min(VertexPositions.Num(), SampleSettings.TargetNumSamples);
+		FPriorityOrderPoints Ordering;
+		Ordering.ComputeUniformSpaced(VertexPositions, VertexAngleMetric, SampleSettings.bRequireSearchSampleCoverage ? -1 : NumSamples);
+		for (int32 SampleIdx = 0; SampleIdx < NumSamples; ++SampleIdx)
+		{
+			int32 VID = Ordering.Order[SampleIdx];
+			AddSample(VertexPositions[VID]);
 		}
-		else
+
+		if (SampleSettings.bRequireSearchSampleCoverage)
 		{
-			// When we have fewer vertices than target samples, just try adding all of them
-			for (FVector3d Pos : NegativeSpaceMesh.VerticesItr())
+			double SpacingThresholdSq = SampleSettings.MinSpacing * SampleSettings.MinSpacing;
+			for (int32 SampleIdx = NumSamples; SampleIdx < Ordering.Order.Num(); ++SampleIdx)
 			{
-				AddSample(Pos);
+				int32 VID = Ordering.Order[SampleIdx];
+				FVector3d Pos = VertexPositions[VID];
+				AddSample(Pos, true /*test cover*/);
 			}
 		}
 	}
diff --git a/Engine/Source/Runtime/GeometryCore/Public/CompGeom/ConvexDecomposition3.h b/Engine/Source/Runtime/GeometryCore/Public/CompGeom/ConvexDecomposition3.h
--- a/Engine/Source/Runtime/GeometryCore/Public/CompGeom/ConvexDecomposition3.h
+++ b/Engine/Source/Runtime/GeometryCore/Public/CompGeom/ConvexDecomposition3.h
@@ -72,8 +72,8 @@ struct FNegativeSpaceSampleSettings
 	bool bOnlyConnectedToHull = false;
 	// Maximum number of voxels to use per dimension, when performing VoxelSearch
 	int32 MaxVoxelsPerDim = 128;
-	// Attempt to keep negative space computation deterministic (e.g., do not run marching cubes in parallel for voxel search)
-	bool bDeterministic = false;
+	// Attempt to keep negative space computation deterministic, at some additional runtime cost
+	bool bDeterministic = true;
 
 	// @return the scale factor that has been applied by Rescale()
 	double GetAppliedScaleFactor() const
