FDynamicMesh: add a CanCollapseEdge function to check if an edge collapse might fail before actually trying it

diff --git a/Engine/Source/Runtime/GeometryCore/Private/DynamicMesh/DynamicMesh3_Edits.cpp b/Engine/Source/Runtime/GeometryCore/Private/DynamicMesh/DynamicMesh3_Edits.cpp
--- a/Engine/Source/Runtime/GeometryCore/Private/DynamicMesh/DynamicMesh3_Edits.cpp
+++ b/Engine/Source/Runtime/GeometryCore/Private/DynamicMesh/DynamicMesh3_Edits.cpp
@@ -1268,13 +1268,8 @@ bool FDynamicMesh3::SplitVertexWouldLeaveIsolated(int VertexID, const TArrayView
 }
 
 
-
-
-EMeshResult FDynamicMesh3::CollapseEdge(int vKeep, int vRemove, double collapse_t, FEdgeCollapseInfo& CollapseInfo)
+EMeshResult FDynamicMesh3::CanCollapseEdgeInternal(int vKeep, int vRemove, double collapse_t, FEdgeCollapseInfo* OutCollapseInfo) const
 {
-	
-	CollapseInfo = FEdgeCollapseInfo();
-
 	if (IsVertex(vKeep) == false || IsVertex(vRemove) == false)
 	{
 		return EMeshResult::Failed_NotAnEdge;
@@ -1316,8 +1311,6 @@ EMeshResult FDynamicMesh3::CollapseEdge(int vKeep, int vRemove, double collapse_
 		bIsBoundaryEdge = true;
 	}
 
-	CollapseInfo.OpposingVerts = FIndex2i(c, d);
-
 	// We cannot collapse if edge lists of a and b share vertices other
 	//  than c and d  (because then we will make a triangle [x b b].
 	//  Unfortunately I cannot see a way to do this more efficiently than brute-force search
@@ -1397,6 +1390,55 @@ EMeshResult FDynamicMesh3::CollapseEdge(int vKeep, int vRemove, double collapse_
 		return EMeshResult::Failed_InvalidNeighbourhood;
 	}
 
+	if (OutCollapseInfo)
+	{
+		OutCollapseInfo->OpposingVerts = FIndex2i(c, d);
+		OutCollapseInfo->KeptVertex = b;
+		OutCollapseInfo->RemovedVertex = a;
+		OutCollapseInfo->bIsBoundary = bIsBoundaryEdge;
+		OutCollapseInfo->CollapsedEdge = eab;
+		OutCollapseInfo->RemovedTris = FIndex2i(t0, t1);
+		OutCollapseInfo->RemovedEdges = FIndex2i(eac, ead);
+		OutCollapseInfo->KeptEdges = FIndex2i(ebc, ebd);
+		OutCollapseInfo->CollapseT = collapse_t;
+	}
+
+	return EMeshResult::Ok;
+}
+
+
+EMeshResult FDynamicMesh3::CanCollapseEdge(int vKeep, int vRemove, double collapse_t) const
+{
+	return CanCollapseEdgeInternal(vKeep, vRemove, collapse_t, nullptr);
+}
+
+
+EMeshResult FDynamicMesh3::CollapseEdge(int vKeep, int vRemove, double collapse_t, FEdgeCollapseInfo& CollapseInfo)
+{
+	CollapseInfo = FEdgeCollapseInfo();
+
+	const EMeshResult CanCollapseResult = CanCollapseEdgeInternal(vKeep, vRemove, collapse_t, &CollapseInfo);
+	if (CanCollapseResult != EMeshResult::Ok)
+	{
+		return CanCollapseResult;
+	}
+
+	const int b = vKeep;		// renaming for sanity. We remove a and keep b
+	const int a = vRemove;
+	const int c = CollapseInfo.OpposingVerts[0];
+	const int d = CollapseInfo.OpposingVerts[1];
+	const int t0 = CollapseInfo.RemovedTris[0];
+	const int t1 = CollapseInfo.RemovedTris[1];
+	const int eab = CollapseInfo.CollapsedEdge;
+	const bool bIsBoundaryEdge = CollapseInfo.bIsBoundary;
+	const int eac = CollapseInfo.RemovedEdges[0];
+	const int ead = CollapseInfo.RemovedEdges[1];
+
+	// This may or may not have been computed already
+	int ebc = CollapseInfo.KeptEdges[0];
+	int ebd = InvalidID;
+
+
 	// save vertex positions before we delete removed (can defer kept?)
 	FVector3d KeptPos = GetVertex(vKeep);
 	FVector3d RemovedPos = GetVertex(vRemove);
@@ -1600,14 +1642,7 @@ EMeshResult FDynamicMesh3::CollapseEdge(int vKeep, int vRemove, double collapse_
 		SetVertexColor(vKeep, Lerp(GetVertexColor(vKeep), RemovedColor, (float)collapse_t));
 	}
 
-	CollapseInfo.KeptVertex = vKeep;
-	CollapseInfo.RemovedVertex = vRemove;
-	CollapseInfo.bIsBoundary = bIsBoundaryEdge;
-	CollapseInfo.CollapsedEdge = eab;
-	CollapseInfo.RemovedTris = FIndex2i(t0, t1);
-	CollapseInfo.RemovedEdges = FIndex2i(eac, ead);
 	CollapseInfo.KeptEdges = FIndex2i(ebc, ebd);
-	CollapseInfo.CollapseT = collapse_t;
 
 	if (HasAttributes())
 	{
diff --git a/Engine/Source/Runtime/GeometryCore/Public/DynamicMesh/DynamicMesh3.h b/Engine/Source/Runtime/GeometryCore/Public/DynamicMesh/DynamicMesh3.h
--- a/Engine/Source/Runtime/GeometryCore/Public/DynamicMesh/DynamicMesh3.h
+++ b/Engine/Source/Runtime/GeometryCore/Public/DynamicMesh/DynamicMesh3.h
@@ -1270,6 +1270,14 @@ public:
 	 */
 	GEOMETRYCORE_API virtual bool SplitVertexWouldLeaveIsolated(int VertexID, const TArrayView<const int>& TrianglesToUpdate);
 
+	/**
+	 * Tests whether collapsing the specified edge using the CollapseEdge function would succeed
+	 * @param KeepVertID index of the vertex that should be kept
+	 * @param RemoveVertID index of the vertex that should be removed
+	 * @param EdgeParameterT vKeep is moved to Lerp(KeepPos, RemovePos, EdgeParameterT)
+	 * @return Ok if the edge can be collapsed, or enum value indicating why the operation cannot be applied
+	 */
+	GEOMETRYCORE_API virtual EMeshResult CanCollapseEdge(int vKeep, int vRemove, double collapse_t) const;
 
 	/**
 	 * Collapse the edge between the two vertices, if topologically possible.
@@ -1539,6 +1547,9 @@ protected:
 	template<int Variant>
 	void SerializeInternal(FArchive& Ar, void* Ptr);
 
+	/* We keep this version of CanCollapseEdge internal because the CollapseInfo struct may only be partially filled out by the function */
+	virtual EMeshResult CanCollapseEdgeInternal(int vKeep, int vRemove, double collapse_t, FEdgeCollapseInfo* OutCollapseInfo) const;
+
 };
 
 
