Add Douglas Peucker simplification algorithm to TPolyline

diff --git a/Engine/Source/Runtime/GeometryCore/Public/Polyline3.h b/Engine/Source/Runtime/GeometryCore/Public/Polyline3.h
--- a/Engine/Source/Runtime/GeometryCore/Public/Polyline3.h
+++ b/Engine/Source/Runtime/GeometryCore/Public/Polyline3.h
@@ -431,6 +431,122 @@ public:
 		}
 		NewPolyline.Vertices[k] = Vertices[N];
 	}
+
+	/**
+	 * Simplify the Polyline to reduce the vertex count
+	 * @param ClusterTolerance Vertices closer than this distance will be merged into a single vertex
+	 * @param LineDeviationTolerance Vertices are allowed to deviate this much from the polylines
+	 */
+	void Simplify(T ClusterTolerance, T LineDeviationTolerance)
+	{
+		const int32 VertexCount = Vertices.Num();
+		if (VertexCount < 3)
+		{
+			// we need at least 3 vertices to be able to simplify a line
+			return;
+		}
+
+		TArray<TVector<T>> NewVertices;
+		NewVertices.Reserve(VertexCount);
+
+		// STAGE 1.  Vertex Reduction within tolerance of prior vertex cluster
+		if (ClusterTolerance > 0)
+		{
+			T ClusterToleranceSquared = ClusterTolerance * ClusterTolerance;
+			NewVertices.Add(Vertices[0]);				// keep the first vertex
+			for (int32 Index = 1; Index < VertexCount-1; Index++) 
+			{
+				if (Geometry::DistanceSquared(Vertices[Index], NewVertices.Last()) < ClusterToleranceSquared)
+				{
+					continue;
+				}
+				NewVertices.Add(Vertices[Index]);
+			}
+			NewVertices.Add(Vertices.Last());		// keep the last vertex
+		}
+		else
+		{
+			NewVertices = Vertices;
+		}
+
+		// STAGE 2.  Douglas-Peucker polyline simplification
+		TArray<bool> Marked;
+		if (LineDeviationTolerance > 0 && NewVertices.Num() >= 3)
+		{
+			Marked.SetNumZeroed(NewVertices.Num());
+
+			// mark the first and last vertices to make sure we keep them
+			Marked[0] = true;
+			Marked.Last() = true;
+			SimplifyDouglasPeucker(LineDeviationTolerance, NewVertices, 0, NewVertices.Num()-1, Marked);
+		}
+
+		// STAGE 3. Copy back values in Vertices
+		if (Marked.IsEmpty())
+		{
+			// we have not run Douglas-Pecker so we can just copy the list
+			Vertices = MoveTemp(NewVertices);
+		}
+		else
+		{
+			// only keep the marked ones
+			Vertices.Reset();
+			const int32 MarkedCount = Marked.Num();
+			for (int32 Index=0; Index<MarkedCount; ++Index)
+			{
+				if (Marked[Index])
+				{
+					Vertices.Add(NewVertices[Index]);
+				}
+			}
+		}
+	}
+
+private:
+	// Polygon simplification
+	// code adapted from: http://softsurfer.com/Archive/algorithm_0205/algorithm_0205.htm
+	// simplifyDP():
+	//  This is the Douglas-Peucker recursive simplification routine
+	//  It just marks Vertices that are part of the simplified polyline
+	//  for approximating the polyline subchain v[j] to v[k].
+	//    Input:  tol = approximation tolerance
+	//            v[] = polyline array of vertex points
+	//            j,k = indices for the subchain v[j] to v[k]
+	//    Output: mk[] = array of markers matching vertex array v[]
+	static void SimplifyDouglasPeucker(T Tolerance, const TArray<TVector<T>>& Vertices, int32 j, int32 k, TArray<bool>& Marked)
+	{
+		Marked.SetNum(Vertices.Num());
+		if (k <= j + 1) // there is nothing to simplify
+			return;
+
+		// check for adequate approximation by segment S from v[j] to v[k]
+		int maxi = j;          // index of vertex farthest from S
+		T maxd2 = 0;         // distance squared of farthest vertex
+		T tol2 = Tolerance * Tolerance;  // tolerance squared
+		TSegment3<T> S = TSegment3<T>(Vertices[j], Vertices[k]);    // segment from v[j] to v[k]
+
+		// test each vertex v[i] for max distance from S
+		// Note: this works in any dimension (2D, 3D, ...)
+		for (int i = j + 1; i < k; i++)
+		{
+			T dv2 = S.DistanceSquared(Vertices[i]);
+			if (dv2 <= maxd2)
+				continue;
+			// v[i] is a max vertex
+			maxi = i;
+			maxd2 = dv2;
+		}
+		if (maxd2 > tol2)       // error is worse than the tolerance
+		{
+			// split the polyline at the farthest vertex from S
+			Marked[maxi] = true;      // mark v[maxi] for the simplified polyline
+			// recursively simplify the two subpolylines at v[maxi]
+			SimplifyDouglasPeucker(Tolerance, Vertices, j, maxi, Marked);  // polyline v[j] to v[maxi]
+			SimplifyDouglasPeucker(Tolerance, Vertices, maxi, k, Marked);  // polyline v[maxi] to v[k]
+		}
+		// else the approximation is OK, so ignore intermediate Vertices
+		return;
+	}
 };
 
 typedef TPolyline3<double> FPolyline3d;
