Update TQuaternion::EpsilonEqual to account for the Quaternion == -Quaternion case, and add an IsIdentity() method

diff --git a/Engine/Source/Runtime/GeometryCore/Public/Quaternion.h b/Engine/Source/Runtime/GeometryCore/Public/Quaternion.h
--- a/Engine/Source/Runtime/GeometryCore/Public/Quaternion.h
+++ b/Engine/Source/Runtime/GeometryCore/Public/Quaternion.h
@@ -71,7 +71,13 @@ PRAGMA_DISABLE_DEPRECATION_WARNINGS
 PRAGMA_ENABLE_DEPRECATION_WARNINGS
 	}
 
-	bool EpsilonEqual(const TQuaternion<RealType>& Other, RealType Epsilon) const;
+	// Test whether quaternions represent the same rotation w/in a tolerance. (Note: Considers the negative representation of the same rotation as equal.)
+	bool EpsilonEqual(const TQuaternion<RealType>& Other, RealType Tolerance = TMathUtil<RealType>::ZeroTolerance) const;
+	// Test whether the quaternion is EpsilonEqual to the identity quaternion (i.e., W=1 or W=-1, within tolerance)
+	bool IsIdentity(RealType Tolerance = TMathUtil<RealType>::ZeroTolerance) const
+	{
+		return EpsilonEqual(Identity(), Tolerance);
+	}
 
 	RealType Length() const { return (RealType)sqrt(X*X + Y*Y + Z*Z + W*W); }
 	RealType SquaredLength() const { return X*X + Y*Y + Z*Z + W*W; }
@@ -564,12 +570,18 @@ TMatrix3<RealType> TQuaternion<RealType>::ToRotationMatrix() const
 
 
 template<typename RealType>
-bool TQuaternion<RealType>::EpsilonEqual(const TQuaternion<RealType>& Other, RealType Epsilon) const
-{
-	return (RealType)fabs(X - Other.X) <= Epsilon &&
-		(RealType)fabs(Y - Other.Y) <= Epsilon &&
-		(RealType)fabs(Z - Other.Z) <= Epsilon &&
-		(RealType)fabs(W - Other.W) <= Epsilon;
+bool TQuaternion<RealType>::EpsilonEqual(const TQuaternion<RealType>& Other, RealType Tolerance) const
+{
+	return 
+		((RealType)fabs(X - Other.X) <= Tolerance &&
+		(RealType)fabs(Y - Other.Y) <= Tolerance &&
+		(RealType)fabs(Z - Other.Z) <= Tolerance &&
+		(RealType)fabs(W - Other.W) <= Tolerance) 
+		||
+		((RealType)fabs(X + Other.X) <= Tolerance &&
+		(RealType)fabs(Y + Other.Y) <= Tolerance &&
+		(RealType)fabs(Z + Other.Z) <= Tolerance &&
+		(RealType)fabs(W + Other.W) <= Tolerance);
 }
 
 
