Add support for inverting Transform Sequences, appending inverses of transforms to Transform Sequences, and for testing if a TTransformSRT3 can represent its own inverse (to help efficiently represent Transform Sequence inverses)

diff --git a/Engine/Source/Runtime/GeometryCore/Public/TransformSequence.h b/Engine/Source/Runtime/GeometryCore/Public/TransformSequence.h
--- a/Engine/Source/Runtime/GeometryCore/Public/TransformSequence.h
+++ b/Engine/Source/Runtime/GeometryCore/Public/TransformSequence.h
@@ -76,6 +76,50 @@ public:
 		return false;
 	}
 
+
+	void AppendInverse(const FTransform& Transform, const RealType Tolerance = TMathUtil<RealType>::ZeroTolerance)
+	{
+		AppendInverse(TTransformSRT3<RealType>(Transform), Tolerance);
+	}
+
+	void AppendInverse(const TTransformSRT3<RealType>& Transform, const RealType Tolerance = TMathUtil<RealType>::ZeroTolerance)
+	{
+		if (Transform.CanRepresentInverse(Tolerance))
+		{
+			Append(Transform.InverseUnsafe(Tolerance));
+		}
+		else
+		{
+			TTransformSRT3<RealType> InvS = TTransformSRT3<RealType>::Identity();
+			InvS.SetScale(TTransformSRT3<RealType>::GetSafeScaleReciprocal(Transform.GetScale3D(), Tolerance));
+
+			TTransformSRT3<RealType> InvRT = Transform;
+			InvRT.SetScale(TVector<RealType>::One());
+			InvRT = InvRT.InverseUnsafe();
+			Append(InvRT);
+			Append(InvS);
+		}
+	}
+
+	void AppendInverse(const TTransformSequence3<RealType>& SequenceToAppend, const RealType Tolerance = TMathUtil<RealType>::ZeroTolerance)
+	{
+		for (int32 Idx = Transforms.Num() - 1; Idx >= 0; --Idx)
+		{
+			AppendInverse(Transforms[Idx], Tolerance);
+		}
+	}
+
+	// Create the inverse of a transform sequence. Note: Transforms with both non-uniform scale and rotation may be split into two transforms. Zeros in scale will be replaced with unit scales.
+	TTransformSequence3<RealType> GetInverse(const RealType Tolerance = TMathUtil<RealType>::ZeroTolerance) const
+	{
+		TTransformSequence3<RealType> InverseTransformSeq;
+		for (int32 Idx = Transforms.Num() - 1; Idx >= 0; --Idx)
+		{
+			InverseTransformSeq.AppendInverse(Transforms[Idx]);
+		}
+		return InverseTransformSeq;
+	}
+
 	/**
 	 * @return Cumulative scale across Transforms.
 	 */
diff --git a/Engine/Source/Runtime/GeometryCore/Public/TransformTypes.h b/Engine/Source/Runtime/GeometryCore/Public/TransformTypes.h
--- a/Engine/Source/Runtime/GeometryCore/Public/TransformTypes.h
+++ b/Engine/Source/Runtime/GeometryCore/Public/TransformTypes.h
@@ -197,14 +197,20 @@ public:
 	 * Attempts to return an inverse, but will give an incorrect result if the transform has both non-uniform scaling and rotation,
 	 * because TTransformSRT3<RealType> cannot represent the true inverse in that case.
 	 */
-	TTransformSRT3<RealType> InverseUnsafe() const
+	TTransformSRT3<RealType> InverseUnsafe(RealType Tolerance = TMathUtil<RealType>::ZeroTolerance) const
 	{
 		TQuaternion<RealType> InvRotation = Rotation.Inverse();
-		TVector<RealType> InvScale3D = GetSafeScaleReciprocal(Scale3D);
+		TVector<RealType> InvScale3D = GetSafeScaleReciprocal(Scale3D, Tolerance);
 		TVector<RealType> InvTranslation = InvRotation * (InvScale3D * -Translation);
 		return TTransformSRT3<RealType>(InvRotation, InvTranslation, InvScale3D);
 	}
 
+	/** Reports wheth the inverse is representable with a single TTransformSRT3. (Ignores zeros in scale.) */
+	bool CanRepresentInverse(RealType Tolerance = TMathUtil<RealType>::ZeroTolerance) const
+	{
+		// Note: Could also return true if there is a non-uniform scale aligned with the rotation axis ...
+		return Scale3D.AllComponentsEqual(Tolerance) || Rotation.IsIdentity(Tolerance);
+	}
 
 	/**
 	 * @return input point with QST transformation applied, ie QST(P) = Rotate(Scale*P) + Translate
