Add way to retrieve the real allocated memory size of a Point Provider Cluster

diff --git a/Engine/Source/Runtime/GeometryCore/Public/Spatial/SparseGrid2.h b/Engine/Source/Runtime/GeometryCore/Public/Spatial/SparseGrid2.h
--- a/Engine/Source/Runtime/GeometryCore/Public/Spatial/SparseGrid2.h
+++ b/Engine/Source/Runtime/GeometryCore/Public/Spatial/SparseGrid2.h
@@ -202,7 +202,10 @@ public:
 		}
 	}
 
-
+	SIZE_T GetAllocatedSize() const
+	{
+		return Elements.GetAllocatedSize() + GetCount() * sizeof(ElemType);
+	}
 
 protected:
 
