fix convex hull failure case when tolerances are used to discard points, and the best point has a reported plane distance less than zero due to floating point error (but the exact test has reported that the point is on the correct side of the plane)

diff --git a/Engine/Source/Runtime/GeometryCore/Private/CompGeom/ConvexHull3.cpp b/Engine/Source/Runtime/GeometryCore/Private/CompGeom/ConvexHull3.cpp
--- a/Engine/Source/Runtime/GeometryCore/Private/CompGeom/ConvexHull3.cpp
+++ b/Engine/Source/Runtime/GeometryCore/Private/CompGeom/ConvexHull3.cpp
@@ -221,7 +221,7 @@ struct FHullConnectivity
 		{
 			if (MaxIdx == SourcePointIdx)
 			{
-				MaxValue = 0;
+				MaxValue = -FMathd::MaxReal;
 				MaxIdx = -1;
 				for (int32 SubIdx = 0; SubIdx < Indices.Num(); ++SubIdx)
 				{
