improve performance for merging coincident edges and fix test coverage

diff --git a/Engine/Source/Runtime/GeometryCore/Private/DynamicMesh/Operations/MergeCoincidentMeshEdges.cpp b/Engine/Source/Runtime/GeometryCore/Private/DynamicMesh/Operations/MergeCoincidentMeshEdges.cpp
--- a/Engine/Source/Runtime/GeometryCore/Private/DynamicMesh/Operations/MergeCoincidentMeshEdges.cpp
+++ b/Engine/Source/Runtime/GeometryCore/Private/DynamicMesh/Operations/MergeCoincidentMeshEdges.cpp
@@ -5,7 +5,7 @@
 #include "DynamicMesh/Operations/SplitAttributeWelder.h"
 #include "DynamicMesh/DynamicMesh3.h"
 #include "DynamicMesh/MeshAdapterUtil.h"
-#include "Spatial/PointSetHashTable.h"
+#include "Spatial/PointHashGrid3.h"
 #include "Util/IndexPriorityQueue.h"
 #include "Util/IndexUtil.h"
 
@@ -23,16 +23,23 @@ bool FMergeCoincidentMeshEdges::Apply()
 	// construct hash table for edge midpoints
 	//
 
-	FPointSetAdapterd EdgeMidpoints = UE::Geometry::MakeBoundaryEdgeMidpointsAdapter(Mesh);
-	FPointSetHashtable MidpointsHash(&EdgeMidpoints);
+	TArray<FVector3d> BoundaryMidPoints;
+	TArray<int32> ToMidPt;
+	ToMidPt.Init(-1, Mesh->MaxEdgeID());
+	for (int32 EID : Mesh->BoundaryEdgeIndicesItr())
+	{
+		ToMidPt[EID] = BoundaryMidPoints.Add(Mesh->GetEdgePoint(EID, 0.5));
+	}
+	InitialNumBoundaryEdges = BoundaryMidPoints.Num();
 
-	// use denser grid as vertex count increases
+	// use denser grid as number of boundary edges increases
 	int hashN = 64;
-	if (Mesh->TriangleCount() > 100000)   hashN = 128;
-	if (Mesh->TriangleCount() > 1000000)  hashN = 256;
+	if (InitialNumBoundaryEdges > 1000)   hashN = 128;
+	if (InitialNumBoundaryEdges > 10000)  hashN = 256;
+	if (InitialNumBoundaryEdges > 100000)  hashN = 512;
 	FAxisAlignedBox3d Bounds = Mesh->GetBounds(true);
 	double CellSize = FMath::Max(FMathd::ZeroTolerance, Bounds.MaxDim() / (double)hashN);
-	MidpointsHash.Build(CellSize, Bounds.Min);
+	TPointHashGrid3<int32, double> MidpointsHash(CellSize, -1);
 
 	UseMergeSearchTol = FMathd::Min(CellSize, UseMergeSearchTol);
 
@@ -40,7 +47,7 @@ bool FMergeCoincidentMeshEdges::Apply()
 	FVector3d A, B, C, D;
 	TArray<int> equivBuffer;
 	TArray<int> SearchMatches;
-	SearchMatches.SetNum(1024); SearchMatches.Reset();  // allocate buffer
+	SearchMatches.Reserve(1024);  // allocate buffer
 
 	//
 	// construct edge equivalence sets. First we find all other edges with same
@@ -53,42 +60,43 @@ bool FMergeCoincidentMeshEdges::Apply()
 	EquivalenceSets.Init(nullptr, Mesh->MaxEdgeID());
 	TSet<int> RemainingEdges;
 
-	// @todo equivalence sets should be symmetric. this neither enforces that,
-	// nor takes advantage of it.
-	InitialNumBoundaryEdges = 0;
 	for (int eid : Mesh->BoundaryEdgeIndicesItr()) 
 	{
-		InitialNumBoundaryEdges++;
-
-		FVector3d midpt = Mesh->GetEdgePoint(eid, 0.5);
+		const int32 MidPtIdx = ToMidPt[eid];
+		FVector3d midpt = BoundaryMidPoints[MidPtIdx];
 
 		// find all other edges with same midpoint in query sphere
 		SearchMatches.Reset();
-		MidpointsHash.FindPointsInBall(midpt, UseMergeSearchTol, SearchMatches);
+		MidpointsHash.FindPointsInBall(midpt, UseMergeSearchTol, [&](const int32& PtIdx)
+			{
+				return FVector3d::DistSquared(midpt, BoundaryMidPoints[ToMidPt[PtIdx]]);
+			}, SearchMatches);
+		// add each point after querying for neighbors, so we only find edges with earlier IDs
+		MidpointsHash.InsertPointUnsafe(eid, midpt);
 
 		int N = SearchMatches.Num();
-		if (N == 1 && SearchMatches[0] != eid)
-		{
-			check(false);	// how could this happen?!
-		}
-		if (N <= 1)
+		if (N == 0)
 		{
 			continue;		// edge has no matches
 		}
 
 		Mesh->GetEdgeV(eid, A, B);
 
-		// if same endpoints, add to equivalence set for this edge
+		// if same endpoints, add to equivalence set for this edge (and matching reverse equivalence)
 		equivBuffer.Reset();
 		for (int i = 0; i < N; ++i) 
 		{
-			if (SearchMatches[i] != eid) 
+			int32 MatchEID = SearchMatches[i];
+			Mesh->GetEdgeV(SearchMatches[i], C, D);
+			if ( IsSameEdge(A, B, C, D) ) 
 			{
-				Mesh->GetEdgeV(SearchMatches[i], C, D);
-				if ( IsSameEdge(A, B, C, D) ) 
+				equivBuffer.Add(SearchMatches[i]);
+				if (!EquivalenceSets[MatchEID])
 				{
-					equivBuffer.Add(SearchMatches[i]);
+					EquivalenceSets[MatchEID] = new EdgesList();
+					RemainingEdges.Add(MatchEID);
 				}
+				EquivalenceSets[MatchEID]->Add(eid);
 			}
 		}
 		if (equivBuffer.Num() > 0)
