fix race condition in sparse dynamic point octree sampling method

diff --git a/Engine/Source/Runtime/GeometryCore/Public/Spatial/SparseDynamicPointOctree3.h b/Engine/Source/Runtime/GeometryCore/Public/Spatial/SparseDynamicPointOctree3.h
--- a/Engine/Source/Runtime/GeometryCore/Public/Spatial/SparseDynamicPointOctree3.h
+++ b/Engine/Source/Runtime/GeometryCore/Public/Spatial/SparseDynamicPointOctree3.h
@@ -575,13 +575,20 @@ void FSparseDynamicPointOctree3::ParallelInsertDensePointSet(int32 MaxPointID, T
 	}
 
 	// create all the Root Cells
+	TArray<FSparsePointOctreeCell*> RootCellPtrs;
 	{
 		//TRACE_CPUPROFILER_EVENT_SCOPE(CreateRootCells);
-		for (FVector3i RootIndex : RootCellIndexes)
+		// Note RootCellPtrs holds pointers to elements of Cells; this should be safe
+		// thanks to the block storage of the TDynamicVector used for Cells
+		RootCellPtrs.SetNumUninitialized(RootCellIndexes.Num());
+		for (int32 Index = 0; Index < RootCellIndexes.Num(); ++Index)
 		{
+			FVector3i RootIndex = RootCellIndexes[Index];
 			FSparsePointOctreeCell NewRootCell(0, RootIndex);
-			CreateNewRootCell(NewRootCell, false);
+			int32 CellsIndex = CreateNewRootCell(NewRootCell, false);
+			RootCellPtrs[Index] = &Cells[CellsIndex];
 		}
+		RootCellIndexes.Empty(); // discard this index array; we will only use the pointer array below
 	}
 
 	// make sure the PointIDToCellMap is large enough for all the PointIDs, then we never have to insert/resize
@@ -600,15 +607,13 @@ void FSparseDynamicPointOctree3::ParallelInsertDensePointSet(int32 MaxPointID, T
 		bool operator==(const FParentChildCell& OtherCell) const { return ChildCellIndex == OtherCell.ChildCellIndex; }
 	};
 
-	// Parallel iteration over the Root cells. For each Root cell, we iterate over it's points, and for each point, 
+	// Parallel iteration over the Root cells. For each Root cell, we iterate over its points, and for each point, 
 	// descend down to the max depth to figure out which cells need to be created. Then we create those cells, and 
 	// then we insert the points into the leaf cells. 
 	FCriticalSection SharedDataLock, CellPointListsLock;
-	ParallelFor(RootCellIndexes.Num(), [&](int32 Index)
+	ParallelFor(RootCellPtrs.Num(), [&](int32 Index)
 	{
-		FVector3i RootIndex = RootCellIndexes[Index];
-		uint32* RootCellID = RootCells.Get(RootIndex, false);
-		FSparsePointOctreeCell& RootCell = Cells[*RootCellID];
+		FSparsePointOctreeCell& RootCell = *RootCellPtrs[Index];
 		
 		// array of new child cells at each level
 		TArray<TArray<FParentChildCell>> LevelCells;
@@ -664,12 +669,13 @@ void FSparseDynamicPointOctree3::ParallelInsertDensePointSet(int32 MaxPointID, T
 				{
 					LeafCellIDs.Add(NewChildCell.CellID);
 				}
+				FSparsePointOctreeCell* NewCellPtr = &Cells[NewChildCell.CellID];
 				SharedDataLock.Unlock();
 
 				(*ParentCell)->SetChild(ParentChild.ChildNum, NewChildCell);
 				ParentChild.NewCellID = NewChildCell.CellID;
 
-				NextParents.Add( &Cells[NewChildCell.CellID] );
+				NextParents.Add(NewCellPtr);
 			}
 			Swap(LastParents, NextParents);
 			NextParents.Reset();
