fix bug in convex decomposition proximity graph updates on split

diff --git a/Engine/Source/Runtime/GeometryCore/Private/CompGeom/ConvexDecomposition3.cpp b/Engine/Source/Runtime/GeometryCore/Private/CompGeom/ConvexDecomposition3.cpp
--- a/Engine/Source/Runtime/GeometryCore/Private/CompGeom/ConvexDecomposition3.cpp
+++ b/Engine/Source/Runtime/GeometryCore/Private/CompGeom/ConvexDecomposition3.cpp
@@ -1753,10 +1753,7 @@ void FConvexDecomposition3::UpdateProximitiesAfterSplit(int32 SplitIdx, int32 Ne
 
 	// Delete removed edges (including updating the map as needed)
 	// Note this must happen *after* connecting up new edges, or else the StartNewProximitiesIdx will be wrong
-	for (int32 ToRemove : ToRemoveProximities)
-	{
-		DeleteProximity(ToRemove, true);
-	}
+	DeleteProximity(MoveTemp(ToRemoveProximities), true);
 
 	// Create new proximity links across the new plane
 	bool bOrigOnPlane = IsMeshOnPlane(SplitIdx, CutPlane);
@@ -2334,32 +2331,45 @@ int32 FConvexDecomposition3::MergeBest(int32 InTargetNumParts, double MaxErrorTo
 	return MergeNum;
 }
 
-void FConvexDecomposition3::DeleteProximity(int32 ProxIdx, bool bDeleteMapReferences)
+void FConvexDecomposition3::DeleteProximity(TArray<int32>&& ToRemoveProx, bool bDeleteMapReferences)
 {
-	if (bDeleteMapReferences)
+	for (int32 ToRemoveIdx = 0; ToRemoveIdx < ToRemoveProx.Num(); ++ToRemoveIdx)
 	{
-		const FProximity& OldProx = Proximities[ProxIdx];
-		DecompositionToProximity.RemoveSingle(OldProx.Link.A, ProxIdx);
-		DecompositionToProximity.RemoveSingle(OldProx.Link.B, ProxIdx);
-	}
-	checkSlow(DecompositionToProximity.FindKey(ProxIdx) == nullptr); // verify there are no more references to this Proximity's index
+		int32 ProxIdx = ToRemoveProx[ToRemoveIdx];
+		if (bDeleteMapReferences)
+		{
+			const FProximity& OldProx = Proximities[ProxIdx];
+			DecompositionToProximity.RemoveSingle(OldProx.Link.A, ProxIdx);
+			DecompositionToProximity.RemoveSingle(OldProx.Link.B, ProxIdx);
+		}
+		checkSlow(DecompositionToProximity.FindKey(ProxIdx) == nullptr); // verify there are no more references to this Proximity's index
 
-	if (ProxIdx == Proximities.Num() - 1)
-	{
-		Proximities.RemoveAt(ProxIdx, 1, EAllowShrinking::No);
-		return;
-	}
-	else
-	{
-		// Remove by swapping the last element to the new slot, and removing the last element
-		// Then update links for that swapped-in proximity
-		int32 LastProxIdx = Proximities.Num() - 1;
-		Proximities.RemoveAtSwap(ProxIdx, 1, EAllowShrinking::No);
-		const FProximity& NewProx = Proximities[ProxIdx];
-		DecompositionToProximity.RemoveSingle(NewProx.Link.A, LastProxIdx);
-		DecompositionToProximity.RemoveSingle(NewProx.Link.B, LastProxIdx);
-		DecompositionToProximity.Add(NewProx.Link.A, ProxIdx);
-		DecompositionToProximity.Add(NewProx.Link.B, ProxIdx);
+		if (ProxIdx == Proximities.Num() - 1)
+		{
+			Proximities.RemoveAt(ProxIdx, 1, EAllowShrinking::No);
+			return;
+		}
+		else
+		{
+			// Remove by swapping the last element to the new slot, and removing the last element
+			// Then update links for that swapped-in proximity
+			int32 LastProxIdx = Proximities.Num() - 1;
+			Proximities.RemoveAtSwap(ProxIdx, 1, EAllowShrinking::No);
+			const FProximity& NewProx = Proximities[ProxIdx];
+			DecompositionToProximity.RemoveSingle(NewProx.Link.A, LastProxIdx);
+			DecompositionToProximity.RemoveSingle(NewProx.Link.B, LastProxIdx);
+			DecompositionToProximity.Add(NewProx.Link.A, ProxIdx);
+			DecompositionToProximity.Add(NewProx.Link.B, ProxIdx);
+			// check if the swapped-back index is one we planned to remove, and if so update the index
+			for (int32 WillRemoveIdx = ToRemoveIdx + 1; WillRemoveIdx < ToRemoveProx.Num(); ++WillRemoveIdx)
+			{
+				if (LastProxIdx == ToRemoveProx[WillRemoveIdx])
+				{
+					ToRemoveProx[WillRemoveIdx] = ProxIdx;
+					break;
+				}
+			}
+		}
 	}
 }
 
diff --git a/Engine/Source/Runtime/GeometryCore/Public/CompGeom/ConvexDecomposition3.h b/Engine/Source/Runtime/GeometryCore/Public/CompGeom/ConvexDecomposition3.h
--- a/Engine/Source/Runtime/GeometryCore/Public/CompGeom/ConvexDecomposition3.h
+++ b/Engine/Source/Runtime/GeometryCore/Public/CompGeom/ConvexDecomposition3.h
@@ -580,8 +580,9 @@ public:
 	TMultiMap<int32, int32> DecompositionToProximity;
 
 	// Helper function to delete a proximity relationship, and fix index references to account for the new array ordering
+	// @param ToRemove					Indices of proximities that will be removed; note this function is destructive to this array
 	// @param bDeleteMapReferences		If true, also update references in the DecompositionToProximity map
-	GEOMETRYCORE_API void DeleteProximity(int32 ProxIdx, bool bDeleteMapReferences);
+	GEOMETRYCORE_API void DeleteProximity(TArray<int32>&& ToRemove, bool bDeleteMapReferences);
 
 	// Helper function to update all proximity data after splitting an FConvexPart apart
 	// @param SplitIdx				The index of the FConvexPart that has been split into multiple parts.  Now contains the first new part.
