Fix convex hull vs sphere test for convex decomposition, and track/exclude degenerate planes from the test

diff --git a/Engine/Source/Runtime/GeometryCore/Private/CompGeom/ConvexDecomposition3.cpp b/Engine/Source/Runtime/GeometryCore/Private/CompGeom/ConvexDecomposition3.cpp
--- a/Engine/Source/Runtime/GeometryCore/Private/CompGeom/ConvexDecomposition3.cpp
+++ b/Engine/Source/Runtime/GeometryCore/Private/CompGeom/ConvexDecomposition3.cpp
@@ -50,7 +50,29 @@ bool FConvexDecomposition3::ConvexPartVsSphereOverlap(const FConvexDecomposition
 		}
 	}
 
-	// Note: Could test vs FConvexPart's bounding box here for an early out (Especially if !bMustTransformConvex and/or Part.HullPlanes.Num() is large!)
+	// Quick test vs FConvexPart's bounding box here for an early out
+	if (!OutDistanceSq)
+	{
+		double BoundsOverlapDist = (Part.Bounds.DiagonalLength() * .5) + Radius;
+		if (FVector3d::DistSquared(Part.Bounds.Center(), Center) > BoundsOverlapDist * BoundsOverlapDist)
+		{
+			return false;
+		}
+	}
+
+	auto GetTriDistanceSq = [&Part, bMustTransformConvex, Center, TransformIntoSphereSpace](int32 PlaneIdx)
+	{
+		FIndex3i TriInds = Part.HullTriangles[PlaneIdx];
+		FTriangle3d Tri(Part.InternalGeo.GetVertex(TriInds.A), Part.InternalGeo.GetVertex(TriInds.B), Part.InternalGeo.GetVertex(TriInds.C));
+		if (bMustTransformConvex)
+		{
+			Tri.V[0] = TransformIntoSphereSpace->TransformPosition(Tri.V[0]);
+			Tri.V[1] = TransformIntoSphereSpace->TransformPosition(Tri.V[1]);
+			Tri.V[2] = TransformIntoSphereSpace->TransformPosition(Tri.V[2]);
+		}
+		FDistPoint3Triangle3d Dist(Center, Tri);
+		return Dist.GetSquared();
+	};
 
 	double ClosestRadiusSq = Radius * Radius;
 	double MaxPlaneDist = 0;
@@ -58,29 +80,38 @@ bool FConvexDecomposition3::ConvexPartVsSphereOverlap(const FConvexDecomposition
 	for (int32 PlaneIdx = 0; PlaneIdx < Part.HullPlanes.Num(); ++PlaneIdx)
 	{
 		FPlane3d Plane = Part.HullPlanes[PlaneIdx];
+		if (Plane.Normal == FVector::ZeroVector)
+		{
+			// for degenerate tri, don't consider the plane distance but still test if we're close to the tri
+			double TriDistSq = GetTriDistanceSq(PlaneIdx);
+			if (TriDistSq < ClosestRadiusSq)
+			{
+				if (OutDistanceSq)
+				{
+					ClosestRadiusSq = TriDistSq;
+					bFoundDistSq = true;
+				}
+				else
+				{
+					return true;
+				}
+			}
+
+			continue;
+		}
 		if (bMustTransformConvex)
 		{
 			Plane.Transform(*TransformIntoSphereSpace);
 		}
 		double PlaneDist = Plane.DistanceTo(Center);
 		MaxPlaneDist = FMath::Max(MaxPlaneDist, PlaneDist);
-		if (PlaneDist > Radius)
+		if (PlaneDist > Radius) // can quick-reject based on plane distance
 		{
 			return false;
 		}
-		else if (Radius > 0 && PlaneDist > FMath::Max(0, MaxPlaneDist - FMathd::ZeroTolerance))
+		else if (Radius > 0 && PlaneDist > FMath::Max(0, MaxPlaneDist - FMathd::ZeroTolerance)) // use a heuristic to only test the most-likely culprit plane
 		{
-			FIndex3i TriInds = Part.HullTriangles[PlaneIdx];
-			FTriangle3d Tri(Part.InternalGeo.GetVertex(TriInds.A), Part.InternalGeo.GetVertex(TriInds.B), Part.InternalGeo.GetVertex(TriInds.C));
-			if (bMustTransformConvex)
-			{
-				Tri.V[0] = TransformIntoSphereSpace->TransformPosition(Tri.V[0]);
-				Tri.V[1] = TransformIntoSphereSpace->TransformPosition(Tri.V[1]);
-				Tri.V[2] = TransformIntoSphereSpace->TransformPosition(Tri.V[2]);
-			}
-			// TODO: Can optimize this by writing custom Point-Tri distance logic to re-use what we already know and early-out if a vertex or edge distance is < Radius
-			FDistPoint3Triangle3d Dist(Center, Tri);
-			double TriDistSq = Dist.GetSquared();
+			double TriDistSq = GetTriDistanceSq(PlaneIdx);
 			if (TriDistSq < ClosestRadiusSq)
 			{
 				if (OutDistanceSq)
@@ -95,6 +126,42 @@ bool FConvexDecomposition3::ConvexPartVsSphereOverlap(const FConvexDecomposition
 			}
 		}
 	}
+	if (MaxPlaneDist > 0) // if we're outside the hull, but didn't accept or reject yet, do one more pass w/ the "most-likely culprit" heuristic flipped
+	{
+		double LocalMaxPlaneDist = 0;
+		for (int32 PlaneIdx = 0; PlaneIdx < Part.HullPlanes.Num(); ++PlaneIdx)
+		{
+			FPlane3d Plane = Part.HullPlanes[PlaneIdx];
+			if (Plane.Normal == FVector3d::ZeroVector)
+			{
+				continue;
+			}
+			
+			if (bMustTransformConvex)
+			{
+				Plane.Transform(*TransformIntoSphereSpace);
+			}
+			double PlaneDist = Plane.DistanceTo(Center);
+			LocalMaxPlaneDist = FMath::Max(LocalMaxPlaneDist, PlaneDist);
+			// Note: We don't test for negative plane distances, etc, since the first pass, above, would have caught those
+			if (PlaneDist <= FMath::Max(0, LocalMaxPlaneDist - FMathd::ZeroTolerance)) // cases the heuristic skipped in the first pass
+			{
+				double TriDistSq = GetTriDistanceSq(PlaneIdx);
+				if (TriDistSq < ClosestRadiusSq)
+				{
+					if (OutDistanceSq)
+					{
+						ClosestRadiusSq = TriDistSq;
+						bFoundDistSq = true;
+					}
+					else
+					{
+						return true;
+					}
+				}
+			}
+		}
+	}
 	if (OutDistanceSq && bFoundDistSq)
 	{
 		*OutDistanceSq = ClosestRadiusSq;
@@ -1394,8 +1461,16 @@ void FConvexDecomposition3::FConvexPart::ComputeHullPlanes()
 	{
 		const FIndex3i& Tri = HullTriangles[TriIdx];
 		FVector3d VA = InternalGeo.GetVertex(Tri.A);
-		FVector3d Normal = VectorUtil::Normal(VA, InternalGeo.GetVertex(Tri.B), InternalGeo.GetVertex(Tri.C));
-		HullPlanes[TriIdx] = FPlane3d(Normal, VA);
+		double Area = 0;
+		FVector3d Normal = VectorUtil::NormalArea(VA, InternalGeo.GetVertex(Tri.B), InternalGeo.GetVertex(Tri.C), Area);
+		if (Area < FMathd::ZeroTolerance)
+		{
+			HullPlanes[TriIdx] = FPlane3d(FVector3d::ZeroVector, 0); // track degenerate plane with a zero-vector normal
+		}
+		else
+		{
+			HullPlanes[TriIdx] = FPlane3d(Normal, VA);
+		}
 	}
 }
 
diff --git a/Engine/Source/Runtime/GeometryCore/Public/CompGeom/ConvexDecomposition3.h b/Engine/Source/Runtime/GeometryCore/Public/CompGeom/ConvexDecomposition3.h
--- a/Engine/Source/Runtime/GeometryCore/Public/CompGeom/ConvexDecomposition3.h
+++ b/Engine/Source/Runtime/GeometryCore/Public/CompGeom/ConvexDecomposition3.h
@@ -493,14 +493,18 @@ public:
 			{
 				return 0;
 			}
-			double MaxDist = HullPlanes[0].DistanceTo(Pt);
+			double MaxDist = -FMathd::MaxReal;
 			for (int32 Idx = 0; Idx < NumPlanes; Idx++)
 			{
 				const FPlane3d& Plane = HullPlanes[Idx];
+				if (Plane.Normal == FVector3d::ZeroVector)
+				{
+					continue;
+				}
 				double Dist = Plane.DistanceTo(Pt);
 				MaxDist = FMath::Max(Dist, MaxDist);
 			}
-			return -MaxDist;
+			return MaxDist == -FMathd::MaxReal ? 0 : -MaxDist;
 		}
 
 		// Helper to create hull after InternalGeo is set
