robustness improvements for convex decomposition: continue splitting other parts after one part fails to split, add a min part size to stop splitting too-small parts, add thickening for degenerate cases where convex hull could not be constructed, add offset on-tri planes by plane tolerance to generate fewer very thin parts

+ support adding custom negative space, and parallelize negative space search

#changelist validated
#virtualized

diff --git a/Engine/Source/Runtime/GeometryCore/Private/CompGeom/ConvexDecomposition3.cpp b/Engine/Source/Runtime/GeometryCore/Private/CompGeom/ConvexDecomposition3.cpp
--- a/Engine/Source/Runtime/GeometryCore/Private/CompGeom/ConvexDecomposition3.cpp
+++ b/Engine/Source/Runtime/GeometryCore/Private/CompGeom/ConvexDecomposition3.cpp
@@ -11,6 +11,7 @@
 
 #include "DynamicMesh/DynamicMesh3.h"
 #include "DynamicMesh/Operations/MergeCoincidentMeshEdges.h"
+#include "DynamicMesh/MeshNormals.h"
 #include "DynamicMesh/MeshTransforms.h"
 #include "DynamicMesh/DynamicMeshAABBTree3.h"
 #include "Spatial/FastWinding.h"
@@ -632,6 +633,7 @@ double GetVolumeUsingReferencePoint(const TriangleMeshType& Mesh, FVector3d RefP
 struct FPartialCutResult
 {
 	TArray<FVector3d> CutVertices;
+	TArray<FVector3d> OffsetVertices[2]; // vertices added to thicken an otherwise-degenerate convex hull; in most cases, not needed
 	TArray<int32> CrossingEdgeIDs;
 	TArray<double> SignDists;
 
@@ -647,11 +649,14 @@ struct FPartialCutResult
 		return HullMinusGeoVolume;
 	}
 
-	void ApplyToGeo(TIndirectArray<FConvexDecomposition3::FConvexPart>& Decomposition, int32 OrigPartIdx, const FPlane3d& Plane, int32& OtherSideStartIdxOut, double PlaneTol, double ConnectedComponentTolerance, const FSphereCovering& NegativeSpace)
+	void ApplyToGeo(TIndirectArray<FConvexDecomposition3::FConvexPart>& Decomposition, int32 OrigPartIdx, const FPlane3d& Plane, int32& OtherSideStartIdxOut, double PlaneTol, double ConnectedComponentTolerance, const FSphereCovering& NegativeSpace, bool bTreatAsSolid, bool bSplitDisconnectedComponents)
 	{
 		FConvexDecomposition3::FConvexPart& OrigPart = Decomposition[OrigPartIdx];
 		bool bSourceGeometryVolumeUnreliable = OrigPart.bGeometryVolumeUnreliable;
 
+		// Save out the negative space overlaps right away (as we will end up replacing this original part)
+		TArray<int32> OrigOverlaps = MoveTemp(OrigPart.OverlapsNegativeSpace);
+
 		int32 NewPartIdx = Decomposition.Add(new FConvexDecomposition3::FConvexPart);
 		FConvexDecomposition3::FConvexPart& NewPart = Decomposition[NewPartIdx];
 		int32 PartIndices[2]{ OrigPartIdx, NewPartIdx };
@@ -659,7 +664,6 @@ struct FPartialCutResult
 		TArray<int32> OnCutEdges;
 		check(CrossingEdgeIDs.Num() == CutVertices.Num());
 
-		TMap<int32, int32> VIDMaps[2]; // mapping from indices used in HullTriangles -> indices used in OrigPart and NewPart
 		TArray<int32> HullCrossingVertexMap;
 		int32 HullCrossingVertexStart = OrigPart.InternalGeo.MaxVertexID();
 		HullCrossingVertexMap.SetNumUninitialized(CrossingEdgeIDs.Num());
@@ -684,20 +688,23 @@ struct FPartialCutResult
 		}
 
 		// remove all fully on-plane triangles
-		for (int32 TID = 0, MaxTID = OrigPart.InternalGeo.MaxTriangleID(); TID < MaxTID; TID++)
+		if (bTreatAsSolid)
 		{
-			if (!OrigPart.InternalGeo.IsTriangle(TID))
+			for (int32 TID = 0, MaxTID = OrigPart.InternalGeo.MaxTriangleID(); TID < MaxTID; TID++)
 			{
-				continue;
-			}
-			FIndex3i Tri = OrigPart.InternalGeo.GetTriangle(TID);
-			// Note any new vertex (ID beyond sign dist or un-set before) will be on the plane by construction (introduced by above edge split)
-			// SignDists was zero-initialized, so we don't need to update it with the new vertices, and can just assume IDs beyond the previous max are on the plane
-			if ((Tri.A >= SignDists.Num() || FMathd::Abs(SignDists[Tri.A]) <= PlaneTol) &&
-				(Tri.B >= SignDists.Num() || FMathd::Abs(SignDists[Tri.B]) <= PlaneTol) &&
-				(Tri.C >= SignDists.Num() || FMathd::Abs(SignDists[Tri.C]) <= PlaneTol))
-			{
-				EMeshResult Res = OrigPart.InternalGeo.RemoveTriangle(TID, true, false);
+				if (!OrigPart.InternalGeo.IsTriangle(TID))
+				{
+					continue;
+				}
+				FIndex3i Tri = OrigPart.InternalGeo.GetTriangle(TID);
+				// Note any new vertex (ID beyond sign dist or un-set before) will be on the plane by construction (introduced by above edge split)
+				// SignDists was zero-initialized, so we don't need to update it with the new vertices, and can just assume IDs beyond the previous max are on the plane
+				if ((Tri.A >= SignDists.Num() || FMathd::Abs(SignDists[Tri.A]) <= PlaneTol) &&
+					(Tri.B >= SignDists.Num() || FMathd::Abs(SignDists[Tri.B]) <= PlaneTol) &&
+					(Tri.C >= SignDists.Num() || FMathd::Abs(SignDists[Tri.C]) <= PlaneTol))
+				{
+					EMeshResult Res = OrigPart.InternalGeo.RemoveTriangle(TID, true, false);
+				}
 			}
 		}
 		
@@ -719,6 +726,34 @@ struct FPartialCutResult
 			}
 		}
 
+		// for non-solid meshes, remove on-plane triangles based on the side they're facing
+		if (!bTreatAsSolid)
+		{
+			for (int32 Side = 0; Side < 2; Side++)
+			{
+				for (int32 TID = 0, MaxTID = PartsMeshes[Side]->MaxTriangleID(); TID < MaxTID; TID++)
+				{
+					if (!PartsMeshes[Side]->IsTriangle(TID))
+					{
+						continue;
+					}
+					FIndex3i Tri = PartsMeshes[Side]->GetTriangle(TID);
+					// Note any new vertex (ID beyond sign dist or un-set before) will be on the plane by construction (introduced by above edge split)
+					// SignDists was zero-initialized, so we don't need to update it with the new vertices, and can just assume IDs beyond the previous max are on the plane
+					if ((Tri.A >= SignDists.Num() || FMathd::Abs(SignDists[Tri.A]) <= PlaneTol) &&
+						(Tri.B >= SignDists.Num() || FMathd::Abs(SignDists[Tri.B]) <= PlaneTol) &&
+						(Tri.C >= SignDists.Num() || FMathd::Abs(SignDists[Tri.C]) <= PlaneTol))
+					{
+						double DotTriPlane = PartsMeshes[Side]->GetTriNormal(TID).Dot(Plane.Normal);
+						if (DotTriPlane > 0 != (bool)Side)
+						{
+							PartsMeshes[Side]->RemoveTriangle(TID, true, false);
+						}
+					}
+				}
+			}
+		}
+
 		// Remove unreferenced vertices, which may be created by the above deletion (i.e., vertices on the plane that were only connected to now-deleted triangles)
 		for (int32 Side = 0; Side < 2; Side++)
 		{
@@ -820,6 +855,7 @@ struct FPartialCutResult
 					// TODO: if result is not well defined, is there anything more robust we could do here?
 					// Perhaps fill based on the winding number of the input mesh? (But more expensive, and we'd have to handle ~coplanar cases as well)
 				}
+				// TODO: Consider only attempting to append the hole fill if bSuccessFill is true
 				if (!Triangles2.IsEmpty()) // Delaunay triangulation succeeded at generating some triangulation
 				{
 					for (const FIndex3i& Tri2 : Triangles2)
@@ -904,7 +940,7 @@ struct FPartialCutResult
 		};
 
 		// TODO: We may be able to skip hole filling if we're using an error metric that is not volume-based
-		bool bHoleFillResultUnreliable = !HoleFill();
+		bool bHoleFillResultUnreliable = !bTreatAsSolid || !HoleFill();
 
 		auto ComputeHullsIfMultipleComponents = [&Decomposition, &ConnectedComponentTolerance](int32 PartIdx)->int32
 		{
@@ -1002,19 +1038,22 @@ struct FPartialCutResult
 			Part->HullTriangles.Reset(HullTriangles[Side].Num());
 			Part->HullTriangles.Reset();
 
-			// If there are multiple components, we'll need to compute all new hulls for each of them
-			int32 NumComponents = ComputeHullsIfMultipleComponents(PartIndices[Side]);
-			if (NumComponents > 1)
+			// If there are multiple components, we'll need to compute all new hulls for each of them (if enabled)
+			if (bSplitDisconnectedComponents)
 			{
-				if (Side == 0) // keep the components from each side of the plane contiguous
+				int32 NumComponents = ComputeHullsIfMultipleComponents(PartIndices[Side]);
+				if (NumComponents > 1)
 				{
-					ComponentsOnSide0 = NumComponents;
-					int32 LastIdx = Decomposition.Num() - 1;
-					Decomposition.Swap(NewPartIdx, LastIdx);
-					NewPartIdx = LastIdx;
-					OtherSideStartIdxOut = NewPartIdx;
+					if (Side == 0) // keep the components from each side of the plane contiguous
+					{
+						ComponentsOnSide0 = NumComponents;
+						int32 LastIdx = Decomposition.Num() - 1;
+						Decomposition.Swap(NewPartIdx, LastIdx);
+						NewPartIdx = LastIdx;
+						OtherSideStartIdxOut = NewPartIdx;
+					}
+					continue;
 				}
-				continue;
 			}
 
 			// There was only one component -- we can re-use the convex hull we computed, just need to fix the vertex indices
@@ -1026,11 +1065,28 @@ struct FPartialCutResult
 			{
 				if (OrigHullIdx < HullCrossingVertexStart)
 				{
-					return CompactMaps.GetVertexMapping(OrigHullIdx);
+					int32 Idx = CompactMaps.GetVertexMapping(OrigHullIdx);
+					ensure(Idx > -1);
+					return Idx;
 				}
 				else
 				{
-					int32 CrossingVertexIdx = OrigHullIdx - HullCrossingVertexStart;
+					const int32 CrossingVertexIdx = OrigHullIdx - HullCrossingVertexStart;
+					if (CrossingVertexIdx >= CutVertices.Num()) // it's an offset vertex; will always need to be appended to the geo on first encounter
+					{
+						int32* AddedVertPtr = AddedVertsMap.Find(CrossingVertexIdx);
+						if (!AddedVertPtr)
+						{
+							const int32 OffsetVertexIdx = CrossingVertexIdx - CutVertices.Num();
+							int32 AppendedVertID = PartsMeshes[Side]->AppendVertex(OffsetVertices[Side][OffsetVertexIdx]);
+							AddedVertsMap.Add(CrossingVertexIdx, AppendedVertID);
+							return AppendedVertID;
+						}
+						else
+						{
+							return *AddedVertPtr;
+						}
+					}
 					int32 NewIdx = HullCrossingVertexMap[CrossingVertexIdx];
 					if (NewIdx == -1)
 					{
@@ -1050,6 +1106,7 @@ struct FPartialCutResult
 					{
 						NewIdx = CompactMaps.GetVertexMapping(NewIdx);
 					}
+					checkSlow(NewIdx > -1);
 					return NewIdx;
 				}
 			};
@@ -1087,16 +1144,48 @@ struct FPartialCutResult
 		}
 
 		// Update lists of conflicting negative spaces (if negative space is available)
-		if (NegativeSpace.Num() > 0 && !OrigPart.OverlapsNegativeSpace.IsEmpty())
+		if (NegativeSpace.Num() > 0 && !OrigOverlaps.IsEmpty())
 		{
-			TArray<int32> OrigOverlaps = MoveTemp(OrigPart.OverlapsNegativeSpace);
-			auto AddIfOverlaps = [NegativeSpace, OrigOverlaps](FConvexDecomposition3::FConvexPart& Part)
+			auto AddIfOverlaps = [&NegativeSpace, &OrigOverlaps](FConvexDecomposition3::FConvexPart& Part)
 			{
-				for (int32 SphereIdx : OrigOverlaps)
+				// heuristics to multithread the convex part vs sphere overlap tests when there are enough spheres to be potentially worth it
+				// TODO: consider also using some kind of spatial acceleration + revisit tuning
+				if (OrigOverlaps.Num() > 200)
+				{
+					const int32 NumBatch = 7 + int32(OrigOverlaps.Num()/200);
+					TArray<TArray<int32>> OverlapArr;
+					OverlapArr.SetNum(NumBatch);
+					const int32 PerBatch = 1+OrigOverlaps.Num() / NumBatch;
+					ParallelFor(NumBatch, [&](int32 BatchIdx)
+						{
+							for (int32 Idx = PerBatch * BatchIdx, Num = FMath::Min(OrigOverlaps.Num(), PerBatch * (BatchIdx + 1)); Idx < Num; ++Idx)
+							{
+								int32 SphereIdx = OrigOverlaps[Idx];
+								if (FConvexDecomposition3::ConvexPartVsSphereOverlap(Part, NegativeSpace.GetCenter(SphereIdx), NegativeSpace.GetRadius(SphereIdx)))
+								{
+									OverlapArr[BatchIdx].Add(SphereIdx);
+								}
+							}
+						});
+					int32 TotalSize = 0;
+					for (int32 BatchIdx = 0; BatchIdx < NumBatch; ++BatchIdx)
+					{
+						TotalSize += OverlapArr[BatchIdx].Num();
+					}
+					Part.OverlapsNegativeSpace.Reserve(TotalSize);
+					for (int32 BatchIdx = 0; BatchIdx < NumBatch; ++BatchIdx)
+					{
+						Part.OverlapsNegativeSpace.Append(OverlapArr[BatchIdx]);
+					}
+				}
+				else
 				{
-					if (FConvexDecomposition3::ConvexPartVsSphereOverlap(Part, NegativeSpace.GetCenter(SphereIdx), NegativeSpace.GetRadius(SphereIdx)))
+					for (int32 SphereIdx : OrigOverlaps)
 					{
-						Part.OverlapsNegativeSpace.Add(SphereIdx);
+						if (FConvexDecomposition3::ConvexPartVsSphereOverlap(Part, NegativeSpace.GetCenter(SphereIdx), NegativeSpace.GetRadius(SphereIdx)))
+						{
+							Part.OverlapsNegativeSpace.Add(SphereIdx);
+						}
 					}
 				}
 			};
@@ -1111,7 +1200,7 @@ struct FPartialCutResult
 
 	FPartialCutResult() {}
 
-	FPartialCutResult(const FConvexDecomposition3::FConvexPart& Convex, const FPlane3d& Plane, double PlaneTol)
+	FPartialCutResult(const FConvexDecomposition3::FConvexPart& Convex, const FPlane3d& Plane, double PlaneTol, bool bTreatAsSolid, double ThickenHullAfterFailure)
 	{
 		SignDists.SetNumZeroed(Convex.InternalGeo.MaxVertexID());
 		TArray<bool> OnHull; // indicates if vertex was on the original hull
@@ -1164,7 +1253,31 @@ struct FPartialCutResult
 					ForHull[Side][EdgeVID.B] = true;
 				}
 			}
-			// else {} Note: fully on-plane case does not need anything
+			// else edge is on the plane -- rely on the one-side case edges to protect these vertices in most cases
+		}
+		// for non-solid shapes, we cannot rely on non-planar segments to preserve on-plane geometry
+		// so we also also mark on-plane triangle vertices based on the triangle normal
+		if (!bTreatAsSolid)
+		{
+			for (int32 TID : Convex.InternalGeo.TriangleIndicesItr())
+			{
+				FIndex3i TriVID = Convex.InternalGeo.GetTriangle(TID);
+				auto IsVertOnPlane = [this, PlaneTol](int32 TestVID) -> bool
+					{
+						double Dist = SignDists[TestVID];
+						return Dist >= -PlaneTol && Dist <= PlaneTol;
+					};
+				bool bOnPlane = IsVertOnPlane(TriVID.A) && IsVertOnPlane(TriVID.B) && IsVertOnPlane(TriVID.C);
+				if (bOnPlane)
+				{
+					FVector3d TriNormal = Convex.InternalGeo.GetTriNormal(TID);
+					double TriDot = TriNormal.Dot(Plane.Normal);
+					int32 KeepSide = int32(TriDot > 0);
+					ForHull[KeepSide][TriVID.A] = true;
+					ForHull[KeepSide][TriVID.B] = true;
+					ForHull[KeepSide][TriVID.C] = true;
+				}
+			}
 		}
 
 		// compute the hulls, s.t. the hull triangle indices are referencing the original internal geo vertices w/ the cut vertices appended
@@ -1213,24 +1326,57 @@ struct FPartialCutResult
 				[this, &Convex, &ForHull, Side](int32 Index) { return Index < Convex.InternalGeo.MaxVertexID() ? ForHull[Side][Index] : true; });
 			if (!bOK)
 			{
-				// don't cut the hull on a plane that would lead to a degenerate hull on either side
-				CutVertices.Empty();
-				CrossingEdgeIDs.Empty();
-				bSuccess = false;
-				return;
+				// optionally re-try w/ minimal offsets in degen directions
+				if (ThickenHullAfterFailure > 0)
+				{
+					const int32 OrigMaxID = Convex.InternalGeo.MaxVertexID();
+					FMeshNormals Normals(&Convex.InternalGeo);
+					Normals.ComputeVertexNormals();
+					const double OffsetFactor = ThickenHullAfterFailure;
+					for (int32 VID = 0; VID < OrigMaxID; ++VID)
+					{
+						if (Convex.InternalGeo.IsVertex(VID) && ForHull[Side][VID])
+						{
+							FVector3d Normal = Normals[VID];
+							if (Normal == FVector::ZeroVector)
+							{
+								Normal = FVector::OneVector * FMathd::InvSqrt3; // for degenerate normals, arbitrarily pick a diagonal offset direction
+							}
+							OffsetVertices[Side].Add(Convex.InternalGeo.GetVertex(VID) - Normals[VID] * OffsetFactor);
+						}
+					}
+					bOK = HullCompute.Solve(Convex.InternalGeo.MaxVertexID() + CutVertices.Num() + OffsetVertices[Side].Num(),
+						[this, &Convex, &ForHull, Side](int32 Index)
+						{
+							const int32 MaxVID = Convex.InternalGeo.MaxVertexID();
+							const int32 MaxCutVID = MaxVID + CutVertices.Num();
+							return Index < MaxVID ? Convex.InternalGeo.GetVertex(Index) : Index < MaxCutVID ? CutVertices[Index - MaxVID] : OffsetVertices[Side][Index - MaxCutVID];
+						},
+						[this, &Convex, &ForHull, Side](int32 Index) { return Index < Convex.InternalGeo.MaxVertexID() ? ForHull[Side][Index] : true; });
+
+				}
+				if (!bOK)
+				{
+					// don't cut the hull on a plane that would lead to a degenerate hull on either side
+					CutVertices.Empty();
+					CrossingEdgeIDs.Empty();
+					bSuccess = false;
+					return;
+				}
 			}
 			HullTriangles[Side] = HullCompute.MoveTriangles();
 
 			// Custom hull volume calculation to account for the triangles indexing partly into InternalGeo, partly into CutVertices
 			HullVolumes[Side] = 0;
-			int32 MaxVID = Convex.InternalGeo.MaxVertexID();
+			const int32 MaxVID = Convex.InternalGeo.MaxVertexID();
+			const int32 MaxCutVID = MaxVID + CutVertices.Num();
 			for (const FIndex3i& Tri : HullTriangles[Side])
 			{
 				FVector3d Verts[3];
 				for (int32 SubIdx = 0; SubIdx < 3; SubIdx++)
 				{
 					int32 VID = Tri[SubIdx];
-					Verts[SubIdx] = VID < MaxVID ? Convex.InternalGeo.GetVertex(VID) : CutVertices[VID - MaxVID];
+					Verts[SubIdx] = VID < MaxVID ? Convex.InternalGeo.GetVertex(VID) : VID < MaxCutVID ? CutVertices[VID - MaxVID] : OffsetVertices[Side][VID - MaxCutVID];
 				}
 
 				// Get cross product of edges and (un-normalized) normal vector.
@@ -1357,7 +1503,7 @@ void FConvexDecomposition3::InitializeFromIndexMesh(TArrayView<const FVector3f>
 }
 
 
-bool FConvexDecomposition3::InitializeNegativeSpace(const FNegativeSpaceSampleSettings& Settings)
+bool FConvexDecomposition3::InitializeNegativeSpace(const FNegativeSpaceSampleSettings& Settings, TArrayView<const FVector3d> RequestedSamples)
 {
 	if (Decomposition.Num() != 1)
 	{
@@ -1367,9 +1513,25 @@ bool FConvexDecomposition3::InitializeNegativeSpace(const FNegativeSpaceSampleSe
 	TFastWindingTree<FDynamicMesh3> InternalGeoWinding(&InternalGeoAABBTree, true);
 
 	FNegativeSpaceSampleSettings RescaledSettings = Settings;
-	RescaledSettings.Rescale(1.0 / ResultTransform.GetScale().X);
+	RescaledSettings.SetResultTransform(ResultTransform);
 	NegativeSpace.Reset();
 	NegativeSpace.AddNegativeSpace(InternalGeoWinding, RescaledSettings, false);
+	for (const FVector3d& Req : RequestedSamples)
+	{
+		FVector3d Center = ResultTransform.InverseTransformPosition(Req);
+		if (!RescaledSettings.bAllowSamplesInsideMesh)
+		{
+			double Winding = InternalGeoWinding.FastWindingNumber(Center);
+			if (Winding > .5) // inside, should skip
+			{
+				continue;
+			}
+		}
+		FVector3d NearSample = InternalGeoAABBTree.FindNearestPoint(Center);
+		double Dist = FVector3d::Dist(Center, NearSample);
+		double Rad = FMath::Max(0, Dist - RescaledSettings.ReduceRadiusMargin);
+		NegativeSpace.AddSphere(Center, Rad);
+	}
 	InitNegativeSpaceConvexPartMapping();
 	return true;
 }
@@ -1490,13 +1652,29 @@ void FConvexDecomposition3::FConvexPart::ComputeStats()
 	HullError = HullVolume - GeoVolume;
 }
 
-int32 FConvexDecomposition3::SplitWorst(bool bCanSkipUnreliableGeoVolumes, double ErrorTolerance, bool bOnlySplitIfNegativeSpaceCovered)
+int32 FConvexDecomposition3::SplitWorst(bool bCanSkipUnreliableGeoVolumes, double ErrorTolerance, bool bOnlySplitIfNegativeSpaceCovered, double MinSplitSizeInWorldSpace)
+{
+	int32 InitialNum = Decomposition.Num();
+	int32 NumAttempts = InitialNum + 1;
+	while (!SplitWorstHelper(bCanSkipUnreliableGeoVolumes, ErrorTolerance, bOnlySplitIfNegativeSpaceCovered, MinSplitSizeInWorldSpace)) // keep attempting to split until we succeed or run out of parts to try
+	{
+		if (!ensure(NumAttempts-- > 0)) // guard against infinite loop
+		{
+			break;
+		}
+	}
+	return Decomposition.Num() - InitialNum;
+}
+
+bool FConvexDecomposition3::SplitWorstHelper(bool bCanSkipUnreliableGeoVolumes, double ErrorTolerance, bool bOnlySplitIfNegativeSpaceCovered, double MinSplitSizeInWorldSpace)
 {
 	if (Decomposition.Num() == 0)
 	{
-		return 0;
+		return true;
 	}
 
+	double MinSplitSize = MinSplitSizeInWorldSpace <= 0 ? MinSplitSizeInWorldSpace : ConvertDistanceToleranceToLocalSpace(MinSplitSizeInWorldSpace);
+
 	double VolumeTolerance = ConvertDistanceToleranceToLocalVolumeTolerance(ErrorTolerance);
 
 	double WorstError = -FMathd::MaxReal;
@@ -1509,6 +1687,14 @@ int32 FConvexDecomposition3::SplitWorst(bool bCanSkipUnreliableGeoVolumes, doubl
 	{
 		for (int32 PartIdx = 0; PartIdx < Decomposition.Num(); PartIdx++)
 		{
+			if (Decomposition[PartIdx].bSplitFailed)
+			{
+				continue;
+			}
+			if (Decomposition[PartIdx].Bounds.MaxDim() < MinSplitSize)
+			{
+				continue;
+			}
 			if (!Decomposition[PartIdx].OverlapsNegativeSpace.IsEmpty())
 			{
 				bHasOverlapsNegative = true;
@@ -1523,6 +1709,14 @@ int32 FConvexDecomposition3::SplitWorst(bool bCanSkipUnreliableGeoVolumes, doubl
 	}
 	for (int32 PartIdx = 0; PartIdx < Decomposition.Num(); PartIdx++)
 	{
+		if (Decomposition[PartIdx].bSplitFailed)
+		{
+			continue;
+		}
+		if (Decomposition[PartIdx].Bounds.MaxDim() < MinSplitSize)
+		{
+			continue;
+		}
 		bool bOverlapsNegative = !Decomposition[PartIdx].OverlapsNegativeSpace.IsEmpty();
 		// always favor the parts that overlap negative space as long as we have them
 		if (bHasOverlapsNegative && !bOverlapsNegative)
@@ -1547,7 +1741,7 @@ int32 FConvexDecomposition3::SplitWorst(bool bCanSkipUnreliableGeoVolumes, doubl
 	
 	if (WorstIdx == INDEX_NONE)
 	{
-		return 0;
+		return true;
 	}
 
 	// Stop splitting if we see no errors above tolerance and no negative-space overlaps
@@ -1556,7 +1750,7 @@ int32 FConvexDecomposition3::SplitWorst(bool bCanSkipUnreliableGeoVolumes, doubl
 	//		It will just add one extra split currently as we alternate toggling bCanSkipUnreliableGeoVolumes off, but the behavior may be confusing.
 	if (!bErrorAboveTolerance && !bHasOverlapsNegative)
 	{
-		return 0;
+		return true;
 	}
 
 	FConvexPart& Part = Decomposition[WorstIdx];
@@ -1565,20 +1759,24 @@ int32 FConvexDecomposition3::SplitWorst(bool bCanSkipUnreliableGeoVolumes, doubl
 
 	TArray<FPlane3d> CandidatePlanes;
 
-	// Always start with the major axis planes
+	// Always start with the major axis planes (as long as the axis bounds length is at least the min split size)
 	FVector3d Center = Part.Bounds.Center();
-	CandidatePlanes.Emplace(FVector3d::XAxisVector, Center);
-	CandidatePlanes.Emplace(FVector3d::YAxisVector, Center);
-	CandidatePlanes.Emplace(FVector3d::ZAxisVector, Center);
 	FVector3d Extents = Part.Bounds.Extents();
-	int32 MaxBoundsDim = 0;
-	if (Extents.Y > Extents.X)
-	{
-		MaxBoundsDim = 1;
-	}
-	if (Extents.Z > Extents[MaxBoundsDim])
+	int32 MaxBoundsDim = FMath::Max3Index(Extents.X, Extents.Y, Extents.Z);
+	int32 MaxAxisAlignedIdx = 0;
+	for (int32 AxisIdx = 0; AxisIdx < 3; ++AxisIdx)
 	{
-		MaxBoundsDim = 2;
+		bool bIsMaxDim = AxisIdx == MaxBoundsDim;
+		if (bIsMaxDim || Extents[AxisIdx] >= MinSplitSize)
+		{
+			FVector3d AxisVector = FVector3d::ZeroVector;
+			AxisVector[AxisIdx] = 1.0;
+			int32 PlaneIdx = CandidatePlanes.Emplace(AxisVector, Center);
+			if (bIsMaxDim)
+			{
+				MaxAxisAlignedIdx = PlaneIdx;
+			}
+		}
 	}
 
 	// make candidate planes out of convex and boundary edges
@@ -1600,14 +1798,14 @@ int32 FConvexDecomposition3::SplitWorst(bool bCanSkipUnreliableGeoVolumes, doubl
 			Convexities.Add((float)Convexity);
 		}
 
-		void AddMultiple(FConvexPart& ConvPart, const FVector3d& Center, FVector3d Normals[MaxPerEdge], int32 NumToAdd, double Convexity)
+		void AddMultiple(FConvexPart& ConvPart, const FVector3d& Center, FVector3d Normals[MaxPerEdge], FVector3d Offsets[MaxPerEdge], int32 NumToAdd, double Convexity)
 		{
 			TArray<FPlane3d, TFixedAllocator<MaxPerEdge>> PlaneSet;
 			for (int32 NormalIdx = 0; NormalIdx < NumToAdd; NormalIdx++)
 			{
 				if (Normals[NormalIdx].Normalize())
 				{
-					PlaneSet.Add(FPlane3d(Normals[NormalIdx], Center));
+					PlaneSet.Add(FPlane3d(Normals[NormalIdx], Center + Offsets[NormalIdx]));
 				}
 			}
 			if (PlaneSet.IsEmpty())
@@ -1696,7 +1894,11 @@ int32 FConvexDecomposition3::SplitWorst(bool bCanSkipUnreliableGeoVolumes, doubl
 				PlaneNormals[3] = (PlaneNormals[0] + PlaneNormals[1]);
 				PlaneNormals[4] = (PlaneNormals[0] + PlaneNormals[2]);
 			}
-			ConvexEdgeCandidates.AddMultiple(Part, (VA + VB) * .5, PlaneNormals, NumPlanes, CosHalfAngle);
+			// Add small offsets to the on-triangle planes, so the source triangle is fully on one side of the plane
+			FVector3d Offsets[5] {FVector3d::ZeroVector, FVector3d::ZeroVector, FVector3d::ZeroVector, FVector3d::ZeroVector, FVector3d::ZeroVector };
+			Offsets[1] = NormalA * OnPlaneTolerance;
+			Offsets[2] = NormalB * OnPlaneTolerance;
+			ConvexEdgeCandidates.AddMultiple(Part, (VA + VB) * .5, PlaneNormals, Offsets, NumPlanes, CosHalfAngle);
 		}
 	}
 
@@ -1717,18 +1919,15 @@ int32 FConvexDecomposition3::SplitWorst(bool bCanSkipUnreliableGeoVolumes, doubl
 	for (int32 PlaneIdx = 0; PlaneIdx < CandidatePlanes.Num(); PlaneIdx++)
 	{
 		const FPlane3d& Plane = CandidatePlanes[PlaneIdx];
-		FPartialCutResult PlaneResult(Part, Plane, OnPlaneTolerance);
+		FPartialCutResult PlaneResult(Part, Plane, OnPlaneTolerance, bTreatAsSolid, ThickenAfterHullFailure);
 		if (!PlaneResult.bSuccess)
 		{
 			continue;
 		}
 		double PlaneError = PlaneResult.Score(Part);
-		if (PlaneIdx < 3)
+		if (MaxAxisAlignedIdx == PlaneIdx) // favor the larger axis (make predictable cuts when none of the planes are reducing the hull volume)
 		{
-			if (MaxBoundsDim == PlaneIdx) // favor the larger axis (make predictable cuts when none of the planes are reducing the hull volume)
-			{
-				PlaneError *= CutLargestAxisErrorScale;
-			}
+			PlaneError *= CutLargestAxisErrorScale;
 		}
 		if (PlaneError < LowestError)
 		{
@@ -1740,16 +1939,17 @@ int32 FConvexDecomposition3::SplitWorst(bool bCanSkipUnreliableGeoVolumes, doubl
 
 	if (BestPlaneIdx == -1)
 	{
-		return 0;
+		Decomposition[WorstIdx].bSplitFailed = true; // mark this part so we know not to try splitting it again
+		return false; // return failure from the helper so split can be re-attempted
 	}
 
 	int32 NewPartsStartIdx = Decomposition.Num();
 	int32 OtherSideStartIdx = -1;
-	BestCutResult.ApplyToGeo(Decomposition, WorstIdx, CandidatePlanes[BestPlaneIdx], OtherSideStartIdx, OnPlaneTolerance, ConnectedComponentTolerance, NegativeSpace);
+	BestCutResult.ApplyToGeo(Decomposition, WorstIdx, CandidatePlanes[BestPlaneIdx], OtherSideStartIdx, OnPlaneTolerance, ConnectedComponentTolerance, NegativeSpace, bTreatAsSolid, bSplitDisconnectedComponents);
 
 	UpdateProximitiesAfterSplit(WorstIdx, NewPartsStartIdx, CandidatePlanes[BestPlaneIdx], OtherSideStartIdx, OrigHullVolume);
 
-	return Decomposition.Num() - NewPartsStartIdx;
+	return true;
 }
 
 void FConvexDecomposition3::UpdateProximitiesAfterSplit(int32 SplitIdx, int32 NewIdxStart, FPlane3d CutPlane, int32 SecondSideIdxStart, double OrigHullVolume)
@@ -2075,7 +2275,7 @@ int32 FConvexDecomposition3::MergeBest(int32 InTargetNumParts, double MaxErrorTo
 						{
 							continue;
 						}
-						FPartialCutResult PlaneResult(PartToSplit, Plane, OnPlaneTolerance);
+						FPartialCutResult PlaneResult(PartToSplit, Plane, OnPlaneTolerance, bTreatAsSolid, ThickenAfterHullFailure);
 						if (!PlaneResult.bSuccess)
 						{
 							continue;
@@ -2084,7 +2284,7 @@ int32 FConvexDecomposition3::MergeBest(int32 InTargetNumParts, double MaxErrorTo
 						double OrigHullVolume = PartToSplit.HullVolume;
 						int32 NewPartsStartIdx = Decomposition.Num();
 						int32 OtherSideStartIdx = -1;
-						PlaneResult.ApplyToGeo(Decomposition, PartIdx, Plane, OtherSideStartIdx, OnPlaneTolerance, ConnectedComponentTolerance, NegativeSpace);
+						PlaneResult.ApplyToGeo(Decomposition, PartIdx, Plane, OtherSideStartIdx, OnPlaneTolerance, ConnectedComponentTolerance, NegativeSpace, bTreatAsSolid, bSplitDisconnectedComponents);
 						UpdateProximitiesAfterSplit(PartIdx, NewPartsStartIdx, Plane, OtherSideStartIdx, OrigHullVolume);
 						PartToSplit.bMustMerge = true;
 						for (int32 NewIdx = NewPartsStartIdx; NewIdx < Decomposition.Num(); NewIdx++)
@@ -2235,16 +2435,21 @@ int32 FConvexDecomposition3::MergeBest(int32 InTargetNumParts, double MaxErrorTo
 				bool bCoversProtectedPt = false;
 				auto TestVsNegativeSpace = [](const FConvexPart& Part, const FSphereCovering& SphereCovering, const FTransform* OptionalTransform = nullptr) -> bool
 				{
-					for (int32 SphereIdx = 0; SphereIdx < SphereCovering.Num(); ++SphereIdx)
-					{
-						double Radius = SphereCovering.GetRadius(SphereIdx);
-						FVector3d Center = SphereCovering.GetCenter(SphereIdx);
-						if (ConvexPartVsSphereOverlap(Part, Center, Radius, OptionalTransform))
+					std::atomic_bool bOverlaps = false;
+					ParallelFor(SphereCovering.Num(), [&](int32 SphereIdx)
 						{
-							return true;
-						}
-					}
-					return false;
+							if (bOverlaps.load(std::memory_order_relaxed))
+							{
+								return;
+							}
+							double Radius = SphereCovering.GetRadius(SphereIdx);
+							FVector3d Center = SphereCovering.GetCenter(SphereIdx);
+							if (ConvexPartVsSphereOverlap(Part, Center, Radius, OptionalTransform))
+							{
+								bOverlaps = true;
+							}
+						});
+					return (bool)bOverlaps;
 				};
 				if (OptionalNegativeSpace)
 				{
diff --git a/Engine/Source/Runtime/GeometryCore/Public/CompGeom/ConvexDecomposition3.h b/Engine/Source/Runtime/GeometryCore/Public/CompGeom/ConvexDecomposition3.h
--- a/Engine/Source/Runtime/GeometryCore/Public/CompGeom/ConvexDecomposition3.h
+++ b/Engine/Source/Runtime/GeometryCore/Public/CompGeom/ConvexDecomposition3.h
@@ -75,18 +75,29 @@ struct FNegativeSpaceSampleSettings
 	// Attempt to keep negative space computation deterministic, at some additional runtime cost
 	bool bDeterministic = true;
 
+	// Whether to allow samples to be added inside the mesh, based on winding number. Can enabled for non-solid meshes; note the convex decomposition should then set bTreatAsSolid to false as well.
+	bool bAllowSamplesInsideMesh = false;
+
 	// @return the scale factor that has been applied by Rescale()
 	double GetAppliedScaleFactor() const
 	{
 		return AppliedScaleFactor;
 	}
 
+	UE_DEPRECATED(5.4, "Use SetResultTransform instead")
 	void Rescale(double ScaleFactor)
 	{
-		MinSpacing *= ScaleFactor;
-		ReduceRadiusMargin *= ScaleFactor;
-		MinRadius *= ScaleFactor;
-		AppliedScaleFactor *= ScaleFactor;
+		RescaleSettings(ScaleFactor);
+	}
+	
+	void SetResultTransform(FTransform InResultTransform)
+	{
+		RescaleSettings(ResultTransform.GetScale3D().X / InResultTransform.GetScale3D().X);
+		ResultTransform = InResultTransform;
+	}
+	FTransform GetResultTransform() const
+	{
+		return ResultTransform;
 	}
 
 	// Make sure the settings values are in valid ranges
@@ -103,6 +114,17 @@ private:
 
 	// Track how the settings have been rescaled
 	double AppliedScaleFactor = 1;
+
+	// Track how to transform back to the original space
+	FTransform ResultTransform = FTransform::Identity;
+	
+	void RescaleSettings(double ScaleFactor)
+	{
+		MinSpacing *= ScaleFactor;
+		ReduceRadiusMargin *= ScaleFactor;
+		MinRadius *= ScaleFactor;
+		AppliedScaleFactor *= ScaleFactor;
+	}
 };
 
 // Define a volume with a set of spheres
@@ -183,6 +205,12 @@ public:
 		}
 	}
 
+	void AddSphere(FVector3d InCenter, double InRadius)
+	{
+		Position.Add(InCenter);
+		Radius.Add(InRadius);
+	}
+
 private:
 	// Sphere centers
 	TArray<FVector3d> Position;
@@ -240,7 +268,7 @@ public:
 	 * @param Settings	Settings to use to find the negative space
 	 * @return			False on failure -- e.g., if there was no mesh available
 	 */
-	GEOMETRYCORE_API bool InitializeNegativeSpace(const FNegativeSpaceSampleSettings& Settings);
+	GEOMETRYCORE_API bool InitializeNegativeSpace(const FNegativeSpaceSampleSettings& Settings, TArrayView<const FVector3d> RequestedSamples = TArrayView<const FVector3d>());
 
 	//
 	// Settings
@@ -272,6 +300,16 @@ public:
 	// Larger values will cost more to run but can let the algorithm find a cleaner decomposition
 	int32 MaxConvexEdgePlanes = 50;
 
+	// Whether to, after each split, also detect whether a part is made up of disconnected pieces and further split them if so
+	bool bSplitDisconnectedComponents = true;
+
+	// Whether to treat the input mesh as a solid shape -- controls whether, e.g., a hole-fill is performed after splitting the mesh to close off the resulting parts
+	bool bTreatAsSolid = true;
+
+	// If greater than zero, we can 'inflate' the input shape by this amount along any degenerate axes in cases where the hull failed to construct due to the input being coplanar/colinear.
+	// Note: This should be less than the smallest tolerance used negative space spheres. Do not use this parameter to control minimum thickness of the output shapes; that can be better enforced later (after merges).
+	double ThickenAfterHullFailure = 0;
+
 
 	// If > 0, search for best merges will be restricted to a greedy local search after testing this many connections. Helpful when there are many potential merges.
 	int32 RestrictMergeSearchToLocalAfterTestNumConnections = -1;
@@ -296,7 +334,8 @@ public:
 	// Note: could return 0 if no splits were possible
 	// @param bCanSkipUnreliableGeoVolumes		if true, don't split hulls where we have questionable geometry volume results, unless there is no hull with good geometry volume results
 	// @param bOnlySplitIfNegativeSpaceCovered	if true, don't split hulls unless they overlap with some covered Negative Space (stored in the corresponding member variable)
-	GEOMETRYCORE_API int32 SplitWorst(bool bCanSkipUnreliableGeoVolumes = false, double ErrorTolerance = 0.0, bool bOnlySplitIfNegativeSpaceCovered = false);
+	// @param MinSplitSize						if > 0, don't split hulls with max bounds dimension lower than this
+	GEOMETRYCORE_API int32 SplitWorst(bool bCanSkipUnreliableGeoVolumes = false, double ErrorTolerance = 0.0, bool bOnlySplitIfNegativeSpaceCovered = false, double MinSplitSizeInWorldSpace = -1);
 
 	// Merge the pairs of convex hulls in the decomposition that will least increase the error.  Intermediate results can be used across merges, so it is best to do all merges in one call.
 	// Note: A future version of this function may replace NumOutputHulls with MaxOutputHulls, but this version keeps both parameters for compatibility / consistent behavior.
@@ -475,6 +514,7 @@ public:
 		FVector3d GeoCenter = FVector3d::ZeroVector; // Some central point of the geometry (e.g., an average of vertices)
 		FAxisAlignedBox3d Bounds;
 		bool bGeometryVolumeUnreliable = false;
+		bool bSplitFailed = false; // flag indicating that cutting has been attempted and failed, so the part should not be considered for splitting
 
 		// Flag indicating the hull should be merged into another part during a MergeBest() call
 		bool bMustMerge = false;
@@ -625,6 +665,9 @@ private:
 
 	// Negative space to attempt to protect; can be referenced by the Decomposition parts
 	FSphereCovering NegativeSpace;
+
+	// Helper to implement SplitWorst
+	bool SplitWorstHelper(bool bCanSkipUnreliableGeoVolumes, double ErrorTolerance, bool bOnlySplitIfNegativeSpaceCovered, double MinSplitSizeInWorldSpace);
 };
 
 } // end namespace UE::Geometry
