fix incorrect invalid-ID check in SplitInteriorVertexTrianglesIntoSubsets (to fix crash in bevel code)

diff --git a/Engine/Source/Runtime/GeometryCore/Private/DynamicMesh/MeshIndexUtil.cpp b/Engine/Source/Runtime/GeometryCore/Private/DynamicMesh/MeshIndexUtil.cpp
--- a/Engine/Source/Runtime/GeometryCore/Private/DynamicMesh/MeshIndexUtil.cpp
+++ b/Engine/Source/Runtime/GeometryCore/Private/DynamicMesh/MeshIndexUtil.cpp
@@ -173,7 +173,7 @@ bool UE::Geometry::SplitInteriorVertexTrianglesIntoSubsets(
 	}
 
 	FIndex2i StartTris = Mesh->GetEdgeT(SplitEdgeID1);
-	if (StartTris.B <= 0)
+	if (StartTris.B < 0)
 	{
 		return false;
 	}
