add more options to the negative space search of convex decomp, including better version of VoxelSearch
 + add function to set usually-better settings than the defaults (but don't change the defaults, to avoid changing the behavior of any existing usage)

diff --git a/Engine/Source/Runtime/GeometryCore/Private/CompGeom/ConvexDecomposition3.cpp b/Engine/Source/Runtime/GeometryCore/Private/CompGeom/ConvexDecomposition3.cpp
--- a/Engine/Source/Runtime/GeometryCore/Private/CompGeom/ConvexDecomposition3.cpp
+++ b/Engine/Source/Runtime/GeometryCore/Private/CompGeom/ConvexDecomposition3.cpp
@@ -220,6 +220,12 @@ bool FSphereCovering::AddNegativeSpace(const TFastWindingTree<FDynamicMesh3>& Sp
 					}
 					double NearDistSq;
 					int NearTID = Spatial.GetTree()->FindNearestTriangle(Pos, NearDistSq);
+					double NearDist = FMath::Sqrt(NearDistSq);
+					if (SampleSettings.OptionalObstacleSDF)
+					{
+						double ObstacleSD = SampleSettings.ObstacleDistance(Pos);
+						NearDist = FMath::Min(ObstacleSD, NearDist);
+					}
 					double R = FMath::Sqrt(NearDistSq) - SampleSettings.ReduceRadiusMargin;
 					if (R < SampleSettings.MinRadius)
 					{
@@ -233,7 +239,7 @@ bool FSphereCovering::AddNegativeSpace(const TFastWindingTree<FDynamicMesh3>& Sp
 			}
 		}
 	}
-	else if (SampleSettings.SampleMethod == FNegativeSpaceSampleSettings::ESampleMethod::VoxelSearch)
+	else if (SampleSettings.SampleMethod == FNegativeSpaceSampleSettings::ESampleMethod::VoxelSearch || SampleSettings.SampleMethod == FNegativeSpaceSampleSettings::ESampleMethod::NavigableVoxelSearch)
 	{
 		const FDynamicMesh3* Mesh = Spatial.GetTree()->GetMesh();
 
@@ -292,31 +298,86 @@ bool FSphereCovering::AddNegativeSpace(const TFastWindingTree<FDynamicMesh3>& Sp
 
 		// Compute the empty space inside the convex hull as (Convex Hull - Original Mesh)
 		FAxisAlignedBox3d HullBox = HullAABB.GetBoundingBox();
-		// expand by 1 unit (rescaled as needed) to make sure we have some empty space at the boundary
-		// TODO: Test if a smaller expand (e.g., KINDA_SMALL_NUMBER) will do just as well here
-		HullBox.Expand(1.0 * SampleSettings.GetAppliedScaleFactor());
+		HullBox.Expand(UE_DOUBLE_KINDA_SMALL_NUMBER);
 		FMarchingCubes MarchingCubes;
-		MarchingCubes.CubeSize = FMath::Clamp(SampleSettings.ReduceRadiusMargin * .5, HullBox.MaxDim() / (double)SampleSettings.MaxVoxelsPerDim, HullBox.MinDim() * .5);
+		const double TargetCubeSize = SampleSettings.ReduceRadiusMargin * SampleSettings.MarchingCubesGridScale;
+		MarchingCubes.CubeSize = FMath::Clamp(TargetCubeSize, HullBox.MaxDim() / (double)SampleSettings.MaxVoxelsPerDim, HullBox.MinDim() * .5);
 		MarchingCubes.Bounds = HullBox;
 		MarchingCubes.RootMode = ERootfindingModes::Bisection;
-		MarchingCubes.RootModeSteps = 3;
+		MarchingCubes.RootModeSteps = 5;
 		MarchingCubes.IsoValue = 0;
-		MarchingCubes.Implicit = [&HullWinding, &Spatial, WindingSign, &SampleSettings](const FVector3d& Pt) -> double
+		MarchingCubes.bParallelCompute = true;
+		if (SampleSettings.SampleMethod == FNegativeSpaceSampleSettings::ESampleMethod::VoxelSearch)
 		{
-			// Volume is inside the hull and outside the input surface
-			if ((HullWinding.FastWindingNumber(Pt) > .5) && (Spatial.FastWindingNumber(Pt) * WindingSign <= .5))
+			MarchingCubes.Implicit = [&HullWinding, &Spatial, WindingSign, &SampleSettings](const FVector3d& Pt) -> double
 			{
-				// Volume is at least ReduceRadiusMargin away from the input surface
-				double NearDistSq = FMathd::MaxReal;
-				int32 NearTri = Spatial.GetTree()->FindNearestTriangle(Pt, NearDistSq, SampleSettings.ReduceRadiusMargin);
-				if (NearTri == INDEX_NONE)
+				// Volume is inside the hull and outside the input surface
+				if ((HullWinding.FastWindingNumber(Pt) > .5) && (Spatial.FastWindingNumber(Pt) * WindingSign <= .5))
 				{
-					return 1.0;
+					// Volume is at least ReduceRadiusMargin away from the input surface
+					double NearDistSq = FMathd::MaxReal;
+					int32 NearTri = Spatial.GetTree()->FindNearestTriangle(Pt, NearDistSq, SampleSettings.ReduceRadiusMargin);
+					if (NearTri == INDEX_NONE)
+					{
+						if (SampleSettings.OptionalObstacleSDF)
+						{
+							double LocalSD = SampleSettings.ObstacleDistance(Pt);
+							if (LocalSD < SampleSettings.ReduceRadiusMargin)
+							{
+								return -1.0;
+							}
+						}
+
+						return 1.0;
+					}
+					return -1.0;
 				}
 				return -1.0;
-			}
-			return -1.0;
-		};
+			};
+		}
+		else // SampleSettings.SampleMethod == FNegativeSpaceSampleSettings::ESampleMethod::NavigableVoxelSearch
+		{
+			MarchingCubes.Bounds = HullAABB.GetBoundingBox();
+			MarchingCubes.Bounds.Expand(SampleSettings.MinRadius + SampleSettings.ReduceRadiusMargin + UE_DOUBLE_KINDA_SMALL_NUMBER);
+			const double MinRadSq = SampleSettings.MinRadius * SampleSettings.MinRadius;
+			MarchingCubes.Implicit = [&HullMeshWrap, &HullAABB, &Spatial, WindingSign, &SampleSettings, MinRadSq](const FVector3d& Pt) -> double
+			{
+				double NearDistSq = FMathd::MaxReal;
+				int32 NearTri = Spatial.GetTree()->FindNearestTriangle(Pt, NearDistSq, SampleSettings.ReduceRadiusMargin + SampleSettings.MinRadius);
+				if (NearTri != INDEX_NONE) // If we're closer than MinRadius + ReduceRadiusMargin to the surface, we can't place a MinRadius sphere here
+				{
+					return -1.0;
+				}
+				if (!SampleSettings.bOnlyConnectedToHull && Spatial.FastWindingNumber(Pt) * WindingSign > .5) // If we're not requiring 'only connected to hull', then remove the internal spaces
+				{
+					return -1.0;
+				}
+				// Test that we're also inside the (offset) convex hull
+				double HullDistSq = FMathd::MaxReal;
+				int32 TID = HullAABB.FindNearestTriangle(Pt, HullDistSq);
+				checkSlow(TID != -1);
+				if (TID != -1)
+				{
+					FTriangle3d Tri;
+					HullMeshWrap.GetTriVertices(TID, Tri.V[0], Tri.V[1], Tri.V[2]);
+					FVector3d N = Tri.Normal();
+					if (N.Dot(Pt - Tri.V[0]) > 0 && HullDistSq > MinRadSq)
+					{
+						return -1.0; // outside the convex hull
+					}
+				}
+				// If there's an optional obstacle SDF, test that we're also not too close to that
+				if (SampleSettings.OptionalObstacleSDF)
+				{
+					double LocalSD = SampleSettings.ObstacleDistance(Pt);
+					if (LocalSD < SampleSettings.MinRadius)
+					{
+						return -1.0;
+					}
+				}
+				return 1.0;
+			};
+		}
 
 		// If we can ignore internal negative space, we can use a seeded/continuation marching cubes here, starting from the convex hull.
 		// Since our negative space surface is away from the source mesh, we need to sample the hull away from its vertices as well;
@@ -339,6 +400,13 @@ bool FSphereCovering::AddNegativeSpace(const TFastWindingTree<FDynamicMesh3>& Sp
 					Mesh->GetVertex(TriInds.C));
 				ProcessTriStack.Reset();
 				ProcessTriStack.Push(HullTri);
+
+				FVector3d TriOffset = FVector3d::ZeroVector;
+				if (SampleSettings.SampleMethod == FNegativeSpaceSampleSettings::ESampleMethod::NavigableVoxelSearch)
+				{
+					TriOffset = HullTri.Normal() * SampleSettings.MinRadius;
+				}
+
 				int32 SeedsAdded = 0;
 				while (!ProcessTriStack.IsEmpty())
 				{
@@ -354,7 +422,8 @@ bool FSphereCovering::AddNegativeSpace(const TFastWindingTree<FDynamicMesh3>& Sp
 						Mids.V[SubIdx] = (Tri.V[SubIdx] + Tri.V[LastIdx]) * .5;
 						if (EdgeLensSq[SubIdx] > AddPtLenSq)
 						{
-							Seeds.Add(Mids.V[SubIdx]);
+							FVector3d ToAdd = Mids.V[SubIdx] + TriOffset;
+							Seeds.Add(ToAdd);
 						}
 						NumLongEdges += (int32)(EdgeLensSq[SubIdx] > SubDivLenSq);
 					}
@@ -381,12 +450,15 @@ bool FSphereCovering::AddNegativeSpace(const TFastWindingTree<FDynamicMesh3>& Sp
 		// Make sure the mesh is compact to simplify downsampling below
 		NegativeSpaceMesh.CompactInPlace();
 
-		auto AddSample = [this, &Mesh, &Spatial, &SampleSettings, &bAddedPoints, WindingSign](FVector3d Pos, bool bTestCover = false)
+		auto AddSample = [this, &Mesh, &Spatial, &SampleSettings, &bAddedPoints, WindingSign](FVector3d Pos, bool bTestCover = false, bool bRequireWalk = true)
 		{
-			double Winding = Spatial.FastWindingNumber(Pos) * WindingSign;
-			if (Winding > .5)
+			if (bRequireWalk || !SampleSettings.bAllowSamplesInsideMesh)
 			{
-				return;
+				double Winding = Spatial.FastWindingNumber(Pos) * WindingSign;
+				if (Winding > .5)
+				{
+					return;
+				}
 			}
 
 			double NearDistSq;
@@ -395,34 +467,38 @@ bool FSphereCovering::AddNegativeSpace(const TFastWindingTree<FDynamicMesh3>& Sp
 
 			// Walk away from the closest point until we're far enough away to create a sample
 			// (give up if we haven't found a valid sample after a few steps, or if we stepped inside the shape)
-			int32 Steps = 0;
-			while (R < SampleSettings.MinRadius && Steps++ < 3)
+			if (bRequireWalk)
 			{
-				bool bFoundValidSample = false;
-				FDistPoint3Triangle3d Query = TMeshQueries<FDynamicMesh3>::TriangleDistance(*Mesh, NearTID, Pos);
-				FVector3d Away = Pos - Query.ClosestTrianglePoint;
-				if (!Away.Normalize())
+				int32 Steps = 0;
+				while (R < SampleSettings.MinRadius && Steps++ < 3)
 				{
-					return;
-				}
-				
-				// Move away and re-test the sample
-				Pos += Away * ((SampleSettings.MinRadius - R) * 1.1);
-				if (Spatial.FastWindingNumber(Pos) * WindingSign <= .5)
-				{
-					NearTID = Spatial.GetTree()->FindNearestTriangle(Pos, NearDistSq);
-					R = FMath::Sqrt(NearDistSq) - SampleSettings.ReduceRadiusMargin;
-					if (R >= SampleSettings.MinRadius)
+					bool bFoundValidSample = false;
+					FDistPoint3Triangle3d Query = TMeshQueries<FDynamicMesh3>::TriangleDistance(*Mesh, NearTID, Pos);
+					FVector3d Away = Pos - Query.ClosestTrianglePoint;
+					if (!Away.Normalize())
 					{
-						break;
+						return;
+					}
+
+					// Move away and re-test the sample
+					Pos += Away * ((SampleSettings.MinRadius - R) * 1.1);
+					if (Spatial.FastWindingNumber(Pos) * WindingSign <= .5)
+					{
+						NearTID = Spatial.GetTree()->FindNearestTriangle(Pos, NearDistSq);
+						R = FMath::Sqrt(NearDistSq) - SampleSettings.ReduceRadiusMargin;
+						if (R >= SampleSettings.MinRadius)
+						{
+							break;
+						}
+					}
+					else  // give up if we stepped into the volume
+					{
+						return;
 					}
-				}
-				else  // give up if we stepped into the volume
-				{
-					return;
 				}
 			}
-			if (R >= SampleSettings.MinRadius)
+
+			if (!bRequireWalk || R >= SampleSettings.MinRadius)
 			{
 				if (bTestCover)
 				{
@@ -517,20 +593,25 @@ bool FSphereCovering::AddNegativeSpace(const TFastWindingTree<FDynamicMesh3>& Sp
 		int32 NumSamples = FMath::Min(VertexPositions.Num(), SampleSettings.TargetNumSamples);
 		FPriorityOrderPoints Ordering;
 		Ordering.ComputeUniformSpaced(VertexPositions, VertexAngleMetric, SampleSettings.bRequireSearchSampleCoverage ? -1 : NumSamples);
+
+		// The 'sample walk' is used to ensure samples are at least MinRadius + ReduceRadiusMargin away from the input surface. 
+		// Note this isn't needed for the NavigableVoxelSearch method; its samples are already this distance from the surface by construction
+		bool bRequireSampleWalk = SampleSettings.SampleMethod == FNegativeSpaceSampleSettings::ESampleMethod::VoxelSearch;
 		for (int32 SampleIdx = 0; SampleIdx < NumSamples; ++SampleIdx)
 		{
 			int32 VID = Ordering.Order[SampleIdx];
-			AddSample(VertexPositions[VID]);
+			AddSample(VertexPositions[VID], false /*test cover*/, bRequireSampleWalk);
 		}
 
 		if (SampleSettings.bRequireSearchSampleCoverage)
 		{
 			double SpacingThresholdSq = SampleSettings.MinSpacing * SampleSettings.MinSpacing;
+			bool bTestCover = SpacingThresholdSq > 0;
 			for (int32 SampleIdx = NumSamples; SampleIdx < Ordering.Order.Num(); ++SampleIdx)
 			{
 				int32 VID = Ordering.Order[SampleIdx];
 				FVector3d Pos = VertexPositions[VID];
-				AddSample(Pos, true /*test cover*/);
+				AddSample(Pos, bTestCover, bRequireSampleWalk);
 			}
 		}
 	}
diff --git a/Engine/Source/Runtime/GeometryCore/Public/CompGeom/ConvexDecomposition3.h b/Engine/Source/Runtime/GeometryCore/Public/CompGeom/ConvexDecomposition3.h
--- a/Engine/Source/Runtime/GeometryCore/Public/CompGeom/ConvexDecomposition3.h
+++ b/Engine/Source/Runtime/GeometryCore/Public/CompGeom/ConvexDecomposition3.h
@@ -40,14 +40,33 @@ struct FNegativeSpaceSampleSettings
 		: TargetNumSamples(TargetNumSamples), MinSpacing(MinSpacing), ReduceRadiusMargin(ReduceRadiusMargin)
 	{}
 
-	enum class ESampleMethod
+	enum class ESampleMethod : uint8
 	{
 		// Place sample spheres in a uniform grid pattern
 		Uniform,
 		// Use voxel-based subtraction and offsetting methods to specifically target concavities
-		VoxelSearch
+		VoxelSearch,
+		// Use a variant of VoxelSearch that aims to limit negative space to the space that can be accessed by a ball of radius >= MinRadius
+		NavigableVoxelSearch
 	};
 
+	// Enum to manage settings configurations for FNegativeSpaceSampleSettings
+	enum class EConfigDefaults : uint8
+	{
+		Latest
+	};
+
+	// Note: Class member default values are held fixed to avoid changing results for existing callers
+	// Use this method to request newer, recommended defaults
+	void ApplyDefaults(EConfigDefaults Settings = EConfigDefaults::Latest)
+	{
+		SampleMethod = ESampleMethod::NavigableVoxelSearch;
+		MarchingCubesGridScale = 1.0;
+		MaxVoxelsPerDim = 1024;
+		MinSpacing = 0.0;
+		TargetNumSamples = 0;
+	}
+
 	// Method used to place samples
 	ESampleMethod SampleMethod = ESampleMethod::Uniform;
 
@@ -56,6 +75,9 @@ struct FNegativeSpaceSampleSettings
 
 	// Minimum desired spacing between sampling spheres; not a strictly enforced bound.
 	double MinSpacing = 3.0;
+	
+	// Whether to allow samples w/ center inside other negative space spheres
+	bool bAllowSamplesInsideSpheres = false;
 
 	// Space to allow between spheres and actual surface
 	double ReduceRadiusMargin = 3.0;
@@ -68,16 +90,22 @@ struct FNegativeSpaceSampleSettings
 	// Whether to require that all candidate sample locations identified by Voxel Search are covered by negative space samples, up to the specified Min Sample Spacing.
 	// Note: This takes priority over TargetNumSamples if the TargetNumSamples did not achieve the required coverage.
 	bool bRequireSearchSampleCoverage = false;
-	// Whether to only consider negative space that is connected to the bounding convex hull, i.e., to ignore hollow inner pockets of negative space that cannot be reached from the outside via a larger-than-ReduceRadiusMargin-wide path
+	// Whether to only consider negative space that is connected to the bounding convex hull, i.e., to ignore hollow inner pockets of negative space that cannot be reached from the outside (for VoxelSearch and Navigable methods)
 	bool bOnlyConnectedToHull = false;
 	// Maximum number of voxels to use per dimension, when performing VoxelSearch
 	int32 MaxVoxelsPerDim = 128;
 	// Attempt to keep negative space computation deterministic, at some additional runtime cost
 	bool bDeterministic = true;
+	// Scale factor for marching cubes grid used for VoxelSearch-based methods
+	double MarchingCubesGridScale = .5;
 
 	// Whether to allow samples to be added inside the mesh, based on winding number. Can enabled for non-solid meshes; note the convex decomposition should then set bTreatAsSolid to false as well.
 	bool bAllowSamplesInsideMesh = false;
 
+	// Optional function to define an obstacle SDF which the negative space should also stay out of. Can be used for example to ignore anything below a ground plane.
+	// Note: Assumed to be in world space
+	TFunction<double(FVector Pos)> OptionalObstacleSDF;
+
 	// @return the scale factor that has been applied by Rescale()
 	double GetAppliedScaleFactor() const
 	{
@@ -109,6 +137,15 @@ struct FNegativeSpaceSampleSettings
 		MinRadius = FMath::Max(0.0, MinRadius);
 		MaxVoxelsPerDim = FMath::Clamp(MaxVoxelsPerDim, 4, 4096);
 	}
+
+	// helper to evaluate OptionalObstacleSDF for local positions
+	inline double ObstacleDistance(const FVector3d& LocalPos) const
+	{
+		checkSlow(OptionalObstacleSDF);
+		FVector3d WorldPos = ResultTransform.TransformPosition(LocalPos);
+		double WorldSD = OptionalObstacleSDF(WorldPos);
+		return WorldSD * AppliedScaleFactor;
+	}
 	
 private:
 
