add constructor to fsizedisjointset

diff --git a/Engine/Source/Runtime/GeometryCore/Public/VertexConnectedComponents.h b/Engine/Source/Runtime/GeometryCore/Public/VertexConnectedComponents.h
--- a/Engine/Source/Runtime/GeometryCore/Public/VertexConnectedComponents.h
+++ b/Engine/Source/Runtime/GeometryCore/Public/VertexConnectedComponents.h
@@ -16,6 +16,12 @@ struct FSizedDisjointSet
 {
 	TArray<int32> Parents, Sizes;
 
+	FSizedDisjointSet() = default;
+	FSizedDisjointSet(int32 NumIDs)
+	{
+		Init(NumIDs);
+	}
+
 	void Init(int32 NumIDs)
 	{
 		Parents.SetNumUninitialized(NumIDs);
