Add static helper to convex hull class to compute the convex hull volume for an array of points

diff --git a/Engine/Source/Runtime/GeometryCore/Private/CompGeom/ConvexHull3.cpp b/Engine/Source/Runtime/GeometryCore/Private/CompGeom/ConvexHull3.cpp
--- a/Engine/Source/Runtime/GeometryCore/Private/CompGeom/ConvexHull3.cpp
+++ b/Engine/Source/Runtime/GeometryCore/Private/CompGeom/ConvexHull3.cpp
@@ -775,6 +775,31 @@ struct FHullConnectivity
 };
 
 
+template<class RealType>
+double TConvexHull3<RealType>::ComputeVolume(const TArrayView<const TVector<RealType>> Points)
+{
+	TConvexHull3<RealType> Hull;
+	bool bSuccess = Hull.Solve(Points.Num(), [&Points](int32 Idx) { return Points[Idx]; });
+	if (!bSuccess)
+	{
+		return 0.0;
+	}
+	const TArray<FIndex3i>& Tris = Hull.GetTriangles();
+	double Volume = 0.0;
+	for (FIndex3i Tri : Tris)
+	{
+		FVector3d V0 = (FVector3d)Points[Tri.A], V1 = (FVector3d)Points[Tri.B], V2 = (FVector3d)Points[Tri.C];
+		FVector3d V1mV0 = V1 - V0;
+		FVector3d V2mV0 = V2 - V0;
+		FVector3d N = V2mV0.Cross(V1mV0);
+		double tmp0 = V0.X + V1.X;
+		double f1x = tmp0 + V2.X;
+		Volume += N.X * f1x;
+	}
+
+	return Volume / 6.0;
+}
+
 
 template<class RealType>
 bool TConvexHull3<RealType>::Solve(int32 NumPoints, TFunctionRef<TVector<RealType>(int32)> GetPointFunc, TFunctionRef<bool(int32)> FilterFunc)
diff --git a/Engine/Source/Runtime/GeometryCore/Public/CompGeom/ConvexHull3.h b/Engine/Source/Runtime/GeometryCore/Public/CompGeom/ConvexHull3.h
--- a/Engine/Source/Runtime/GeometryCore/Public/CompGeom/ConvexHull3.h
+++ b/Engine/Source/Runtime/GeometryCore/Public/CompGeom/ConvexHull3.h
@@ -84,6 +84,9 @@ public:
 	// Settings controlling whether and how to generate a simpler hull
 	TConvexHullSimplificationSettings<RealType> SimplificationSettings;
 
+	// Helper to compute the a convex hull and return only its volume. If the hull cannot be constructed, a volume of 0 will be returned.
+	GEOMETRYCORE_API static double ComputeVolume(const TArrayView<const TVector<RealType>> Vertices);
+
 	/**
 	 * Generate convex hull as long as input is not degenerate
 	 * If input is degenerate, this will return false, and caller can call GetDimension()
