move the convex decomposition algorithm's MergeBest function arguments into a struct, and add:
 a callback to report merges,
 a callback to prevent merges, and
 a bool to control whether the error tolerance allows merges below the target number of parts requested.

diff --git a/Engine/Source/Runtime/GeometryCore/Private/CompGeom/ConvexDecomposition3.cpp b/Engine/Source/Runtime/GeometryCore/Private/CompGeom/ConvexDecomposition3.cpp
--- a/Engine/Source/Runtime/GeometryCore/Private/CompGeom/ConvexDecomposition3.cpp
+++ b/Engine/Source/Runtime/GeometryCore/Private/CompGeom/ConvexDecomposition3.cpp
@@ -2239,7 +2239,30 @@ void FConvexDecomposition3::InitializeProximityFromDecompositionBoundingBoxOverl
 int32 FConvexDecomposition3::MergeBest(int32 InTargetNumParts, double MaxErrorTolerance, double MinThicknessToleranceWorldSpace, bool bAllowCompact, bool bRequireHullTriangles, int32 MaxOutputHulls,
 	const FSphereCovering* OptionalNegativeSpace, const FTransform* OptionalTransformIntoNegativeSpace)
 {
-	int32 TargetNumParts = FMath::Max(1, InTargetNumParts);
+	FMergeSettings Settings;
+	Settings.TargetNumParts = InTargetNumParts;
+	Settings.ErrorTolerance = MaxErrorTolerance;
+	Settings.MinThicknessTolerance = MinThicknessToleranceWorldSpace;
+	Settings.bAllowCompact = bAllowCompact;
+	Settings.bRequireHullTriangles = bRequireHullTriangles;
+	Settings.MaxOutputHulls = MaxOutputHulls;
+	Settings.OptionalNegativeSpace = OptionalNegativeSpace;
+	Settings.OptionalTransformIntoNegativeSpace = OptionalTransformIntoNegativeSpace;
+	return MergeBest(Settings);
+}
+
+int32 FConvexDecomposition3::MergeBest(const FMergeSettings& Settings)
+{
+	// expand some settings out to variables for convenience
+	int32 TargetNumParts = FMath::Max(1, Settings.TargetNumParts);
+	double MaxErrorTolerance = Settings.ErrorTolerance;
+	double MinThicknessToleranceWorldSpace = Settings.MinThicknessTolerance;
+	bool bAllowCompact = Settings.bAllowCompact;
+	bool bRequireHullTriangles = Settings.bRequireHullTriangles;
+	int32 MaxOutputHulls = Settings.MaxOutputHulls;
+	const FSphereCovering* OptionalNegativeSpace = Settings.OptionalNegativeSpace;
+	const FTransform* OptionalTransformIntoNegativeSpace = Settings.OptionalTransformIntoNegativeSpace;
+
 	const bool bHasValidMaxHulls = MaxOutputHulls > 0;
 	if (bHasValidMaxHulls)
 	{
@@ -2401,6 +2424,17 @@ int32 FConvexDecomposition3::MergeBest(int32 InTargetNumParts, double MaxErrorTo
 		}
 	}
 
+	// If we have a MergeCallback, maintain a mapping from current decomposition part indices back to original part indices
+	TArray<int32> ToOriginalPartIdx;
+	if (Settings.MergeCallback)
+	{
+		ToOriginalPartIdx.SetNumUninitialized(Decomposition.Num());
+		for (int32 PartIdx = 0; PartIdx < ToOriginalPartIdx.Num(); ++PartIdx)
+		{
+			ToOriginalPartIdx[PartIdx] = PartIdx;
+		}
+	}
+
 
 	auto CreateMergedPart = [this, &IsPartBelowSizeTolerance](FProximity& Prox) -> TUniquePtr<FConvexPart>
 	{
@@ -2437,7 +2471,7 @@ int32 FConvexDecomposition3::MergeBest(int32 InTargetNumParts, double MaxErrorTo
 	for (; MustMergeCount > 0 || VolumeTolerance > 0 || Decomposition.Num() > TargetNumParts; MergeNum++)
 	{
 		double BestKnownCost = FMathd::MaxReal;
-		bool bOnlyAllowMustMerges = VolumeTolerance == 0 && Decomposition.Num() <= TargetNumParts;
+		bool bOnlyAllowMustMerges = (VolumeTolerance == 0 || !Settings.bErrorToleranceOverridesNumParts) && Decomposition.Num() <= TargetNumParts;
 		int32 BestKnownIdx = -1;
 		int32 ConsideredCount = 0;
 		for (int32 ProxIdx = 0; ProxIdx < Proximities.Num(); ProxIdx++)
@@ -2467,6 +2501,15 @@ int32 FConvexDecomposition3::MergeBest(int32 InTargetNumParts, double MaxErrorTo
 				Decomposition[Proximities[ProxIdx].Link.A].bMustMerge &&
 				Decomposition[Proximities[ProxIdx].Link.B].bMustMerge;
 
+			// Run the custom allow-merge function (if we have one)
+			if (Settings.CustomAllowMergeParts)
+			{
+				if (!Settings.CustomAllowMergeParts(Decomposition[Proximities[ProxIdx].Link.A], Decomposition[Proximities[ProxIdx].Link.B]))
+				{
+					continue;
+				}
+			}
+
 			double MergeCost;
 			// attempt to remove parts from the proximity store if we're at the size limit
 			bool bShouldRemoveParts = ProximityComputedParts.Num() > EvictComputedPartsThreshold;
@@ -2641,10 +2684,21 @@ int32 FConvexDecomposition3::MergeBest(int32 InTargetNumParts, double MaxErrorTo
 		NewLinks.Remove(DecoToKeep);
 		NewLinks.Remove(DecoToRm);
 
+		if (Settings.MergeCallback)
+		{
+			int32 OrigKeep = ToOriginalPartIdx[DecoToKeep];
+			int32 OrigRm = ToOriginalPartIdx[DecoToRm];
+			Settings.MergeCallback(OrigKeep, OrigRm);
+		}
+
 		// Remove the unused convex, updating references to the part that was swapped back
 		int32 LastIdx = Decomposition.Num() - 1;
 		if (DecoToRm != LastIdx)
 		{
+			if (Settings.MergeCallback)
+			{
+				ToOriginalPartIdx[DecoToRm] = ToOriginalPartIdx[LastIdx];
+			}
 			// Before swap-removing DecoToRm, update references that used to point to the last element of the array to now point to DecoToRm
 			[this](int32 OrigIdx, int32 NewIdx)
 			{
diff --git a/Engine/Source/Runtime/GeometryCore/Public/CompGeom/ConvexDecomposition3.h b/Engine/Source/Runtime/GeometryCore/Public/CompGeom/ConvexDecomposition3.h
--- a/Engine/Source/Runtime/GeometryCore/Public/CompGeom/ConvexDecomposition3.h
+++ b/Engine/Source/Runtime/GeometryCore/Public/CompGeom/ConvexDecomposition3.h
@@ -373,9 +373,40 @@ public:
 	// @param bOnlySplitIfNegativeSpaceCovered	if true, don't split hulls unless they overlap with some covered Negative Space (stored in the corresponding member variable)
 	// @param MinSplitSize						if > 0, don't split hulls with max bounds dimension lower than this
 	GEOMETRYCORE_API int32 SplitWorst(bool bCanSkipUnreliableGeoVolumes = false, double ErrorTolerance = 0.0, bool bOnlySplitIfNegativeSpaceCovered = false, double MinSplitSizeInWorldSpace = -1);
+	
+	struct FConvexPart;
+
+	struct FMergeSettings
+	{
+		// Desired number of output parts
+		int32 TargetNumParts = 1;
+		// If > 0, allow merging parts with error less than this tolerance
+		double ErrorTolerance = 0;
+		// Whether ErrorTolerance is allowed to keep merging below the TargetNumParts
+		bool bErrorToleranceOverridesNumParts = true;
+		// If > 0, maximum number of output parts; overrides ErrorTolerance and TargetNumParts if they would create more parts than this
+		int32 MaxOutputHulls = -1;
+		
+		// Optionally specify a minimum thickness (in cm) for convex parts; parts below this thickness will always be merged away. Overrides TargetNumParts and ErrorTolerance when needed.
+		double MinThicknessTolerance = 0;
+		// Allow the algorithm to discard underlying geometry once it will no longer be used, resulting in a smaller representation and faster merges
+		bool bAllowCompact = true;
+		// Require all hulls to have associated triangles after MergeBest is completed. (If using InitializeFromHulls, will need to triangulate any un-merged hulls.)
+		bool bRequireHullTriangles = false;
+		
+		// Optional representation of negative space that should not be covered by merges
+		const FSphereCovering* OptionalNegativeSpace = nullptr;
+		// Optional transform from space of the convex hulls into the space of the sphere covering; if not provided, assume spheres are in the same coordinate space as the hulls
+		const FTransform* OptionalTransformIntoNegativeSpace = nullptr;
+
+		// Optional callback, to be called when two parts are merged. Takes the original (pre-merge) convex decomposition part indices as input.
+		TFunction<void(int32, int32)> MergeCallback = nullptr;
+
+		// Optional custom function to control whether parts are allowed to merge
+		TFunction<bool(const FConvexDecomposition3::FConvexPart& A, const FConvexDecomposition3::FConvexPart& B)> CustomAllowMergeParts = nullptr;
+	};
 
 	// Merge the pairs of convex hulls in the decomposition that will least increase the error.  Intermediate results can be used across merges, so it is best to do all merges in one call.
-	// Note: A future version of this function may replace NumOutputHulls with MaxOutputHulls, but this version keeps both parameters for compatibility / consistent behavior.
 	// @param TargetNumParts		The target number of parts for the decomposition; will be overriden by non-default ErrorTolerance or MinPartThickness
 	// @param ErrorTolerance		If > 0, continue to merge (if there are possible merges) until the resulting error would be greater than this value. Overrides TargetNumParts as the stopping condition.
 	//								Note: ErrorTolerance is expressed in cm, and will be cubed for volumetric error.
@@ -387,9 +418,13 @@ public:
 	// @param OptionalNegativeSpace Optional representation of negative space that should not be covered by merges
 	// @param OptionalNegativeSpaceTransform Optional transform from space of the convex hulls into the space of the sphere covering; if not provided, assume spheres are in the same coordinate space as the hulls
 	// @return						The number of merges performed
-	GEOMETRYCORE_API int32 MergeBest(int32 TargetNumParts, double ErrorTolerance = 0, double MinThicknessTolerance = 0, bool bAllowCompact = true, bool bRequireHullTriangles = false, int MaxOutputHulls = -1,
+	GEOMETRYCORE_API int32 MergeBest(int32 TargetNumParts, double ErrorTolerance = 0, double MinThicknessTolerance = 0, bool bAllowCompact = true, bool bRequireHullTriangles = false, int32 MaxOutputHulls = -1,
 		const FSphereCovering* OptionalNegativeSpace = nullptr, const FTransform* OptionalTransformIntoNegativeSpace = nullptr);
 
+	// Merge the pairs of convex hulls in the decomposition that will least increase the error.  Intermediate results can be used across merges, so it is best to do all merges in one call.
+	// @return The number of merges performed
+	GEOMETRYCORE_API int32 MergeBest(const FMergeSettings& Settings);
+
 	// simple helper to convert an error tolerance expressed in world space to a local-space volume tolerance
 	inline double ConvertDistanceToleranceToLocalVolumeTolerance(double DistTolerance) const
 	{
