remove older deprecated code from GeometryCore (FTransformSRT3's Inverse and legacy ExactPredicates methods)

diff --git a/Engine/Source/Runtime/GeometryCore/Public/CompGeom/ExactPredicates.h b/Engine/Source/Runtime/GeometryCore/Public/CompGeom/ExactPredicates.h
--- a/Engine/Source/Runtime/GeometryCore/Public/CompGeom/ExactPredicates.h
+++ b/Engine/Source/Runtime/GeometryCore/Public/CompGeom/ExactPredicates.h
@@ -57,46 +57,6 @@ float GEOMETRYCORE_API InCircle(const float* PA, const float* PB, const float* P
 
 // Note: float version of InSphere is not exposed here; instead, convert to and use the double version
 
-// Note: The below Orient2D/Orient3D functions have been deprecated in favor of the above c-array version or the further-below templated versions
-// Note: These were previously templated on VectorType, but to deprecate them and to fix ambiguity with the above double*/float* versions,
-// they have been expanded out the vector types they were used with
-
-UE_DEPRECATED(5.2, "Use Orient2<double> instead.");
-inline double Orient2D(const FVector2d& A, const FVector2d& B, const FVector2d& C)
-{
-	double PA[2]{ A.X, A.Y };
-	double PB[2]{ B.X, B.Y };
-	double PC[2]{ C.X, C.Y };
-	return Orient2D(PA, PB, PC);
-}
-UE_DEPRECATED(5.2, "Use Orient2<float> instead.");
-inline double Orient2D(const FVector2f& A, const FVector2f& B, const FVector2f& C)
-{
-	double PA[2]{ (double)A.X, (double)A.Y };
-	double PB[2]{ (double)B.X, (double)B.Y };
-	double PC[2]{ (double)C.X, (double)C.Y };
-	return Orient2D(PA, PB, PC);
-}
-UE_DEPRECATED(5.2, "Use Orient3<double> instead.")
-inline double Orient3D(const FVector3d& A, const FVector3d& B, const FVector3d& C, const FVector3d& D)
-{
-	const double PA[3]{ A.X, A.Y, A.Z };
-	const double PB[3]{ B.X, B.Y, B.Z };
-	const double PC[3]{ C.X, C.Y, C.Z };
-	const double PD[3]{ D.X, D.Y, D.Z };
-	return Orient3D(PA, PB, PC, PD);
-}
-UE_DEPRECATED(5.2, "Use Orient3<float> instead.")
-inline double Orient3D(const FVector3f& A, const FVector3f& B, const FVector3f& C, const FVector3f& D)
-{
-	const double PA[3]{ (double)A.X, (double)A.Y, (double)A.Z };
-	const double PB[3]{ (double)B.X, (double)B.Y, (double)B.Z };
-	const double PC[3]{ (double)C.X, (double)C.Y, (double)C.Z };
-	const double PD[3]{ (double)D.X, (double)D.Y, (double)D.Z };
-	return Orient3D(PA, PB, PC, PD);
-}
-
-
 /**
  * TVector2-only version that can run in float or double
  * @return value indicating which side of line AB point C is on, or 0 if ABC are collinear
diff --git a/Engine/Source/Runtime/GeometryCore/Public/TransformTypes.h b/Engine/Source/Runtime/GeometryCore/Public/TransformTypes.h
--- a/Engine/Source/Runtime/GeometryCore/Public/TransformTypes.h
+++ b/Engine/Source/Runtime/GeometryCore/Public/TransformTypes.h
@@ -183,16 +183,6 @@ public:
 			|| (TMathUtil<RealType>::Abs(Scale3D.Y - Scale3D.Z) > Tolerance);
 	}
 
-	/**
-	 * Attempts to return an inverse, but will give an incorrect result if the transform has both non-uniform scaling and rotation,
-	 * because TTransformSRT3<RealType> cannot represent the true inverse in that case.
-	 */
-	UE_DEPRECATED(5.1, "Use the InverseTransformPosition/Vector/Normal functions instead; TransformSRT3 cannot represent its own inverse.  If the exact previous behavior is needed, use InverseUnsafe.")
-	TTransformSRT3<RealType> Inverse() const
-	{
-		return InverseUnsafe();
-	}
-
 	/**
 	 * Attempts to return an inverse, but will give an incorrect result if the transform has both non-uniform scaling and rotation,
 	 * because TTransformSRT3<RealType> cannot represent the true inverse in that case.
