make convex decomposition negative space default behavior more consistent to the previous version

diff --git a/Engine/Source/Runtime/GeometryCore/Private/CompGeom/ConvexDecomposition3.cpp b/Engine/Source/Runtime/GeometryCore/Private/CompGeom/ConvexDecomposition3.cpp
--- a/Engine/Source/Runtime/GeometryCore/Private/CompGeom/ConvexDecomposition3.cpp
+++ b/Engine/Source/Runtime/GeometryCore/Private/CompGeom/ConvexDecomposition3.cpp
@@ -298,7 +298,7 @@ bool FSphereCovering::AddNegativeSpace(const TFastWindingTree<FDynamicMesh3>& Sp
 
 		// Compute the empty space inside the convex hull as (Convex Hull - Original Mesh)
 		FAxisAlignedBox3d HullBox = HullAABB.GetBoundingBox();
-		HullBox.Expand(UE_DOUBLE_KINDA_SMALL_NUMBER);
+		HullBox.Expand(SampleSettings.VoxelExpandBoundsFactor * SampleSettings.GetAppliedScaleFactor());
 		FMarchingCubes MarchingCubes;
 		const double TargetCubeSize = SampleSettings.ReduceRadiusMargin * SampleSettings.MarchingCubesGridScale;
 		MarchingCubes.CubeSize = FMath::Clamp(TargetCubeSize, HullBox.MaxDim() / (double)SampleSettings.MaxVoxelsPerDim, HullBox.MinDim() * .5);
diff --git a/Engine/Source/Runtime/GeometryCore/Public/CompGeom/ConvexDecomposition3.h b/Engine/Source/Runtime/GeometryCore/Public/CompGeom/ConvexDecomposition3.h
--- a/Engine/Source/Runtime/GeometryCore/Public/CompGeom/ConvexDecomposition3.h
+++ b/Engine/Source/Runtime/GeometryCore/Public/CompGeom/ConvexDecomposition3.h
@@ -65,6 +65,7 @@ struct FNegativeSpaceSampleSettings
 		MaxVoxelsPerDim = 1024;
 		MinSpacing = 0.0;
 		TargetNumSamples = 0;
+		VoxelExpandBoundsFactor = UE_DOUBLE_KINDA_SMALL_NUMBER;
 	}
 
 	// Method used to place samples
@@ -102,6 +103,9 @@ struct FNegativeSpaceSampleSettings
 	// Whether to allow samples to be added inside the mesh, based on winding number. Can enabled for non-solid meshes; note the convex decomposition should then set bTreatAsSolid to false as well.
 	bool bAllowSamplesInsideMesh = false;
 
+	// How much to expand the bounding box used for voxel search algorithms
+	double VoxelExpandBoundsFactor = 1.0;
+
 	// Optional function to define an obstacle SDF which the negative space should also stay out of. Can be used for example to ignore anything below a ground plane.
 	// Note: Assumed to be in world space
 	TFunction<double(FVector Pos)> OptionalObstacleSDF;
