fix unreachable implicit "++iterator" in cases where a for-loop was used to just grab the first element of a set or map

diff --git a/Engine/Plugins/Experimental/MeshModelingToolsetExp/Source/MeshModelingToolsExp/Private/MeshAttributePaintTool.cpp b/Engine/Plugins/Experimental/MeshModelingToolsetExp/Source/MeshModelingToolsExp/Private/MeshAttributePaintTool.cpp
--- a/Engine/Plugins/Experimental/MeshModelingToolsetExp/Source/MeshModelingToolsExp/Private/MeshAttributePaintTool.cpp
+++ b/Engine/Plugins/Experimental/MeshModelingToolsetExp/Source/MeshModelingToolsExp/Private/MeshAttributePaintTool.cpp
@@ -655,11 +655,8 @@ void UMeshAttributePaintTool::ApplyStamp_FloodFill(const FBrushStampData& Stamp,
 		QueueTempBuffer.Reset();
 		DoneTempBuffer.Reset();
 		InputTriROI.Reset();
-		for (int32 tid : RemainingTriangles)
-		{
-			InputTriROI.Add(tid);		// stupid way to get first set element
-			break;
-		}
+		// get a single set element via an iterator
+		InputTriROI.Add(*RemainingTriangles.CreateConstIterator());
 		FMeshConnectedComponents::GrowToConnectedTriangles(CurrentMesh, InputTriROI, OutputTriROI, &QueueTempBuffer, &DoneTempBuffer);
 		for (int32 tid : OutputTriROI)
 		{
diff --git a/Engine/Plugins/Runtime/GeometryProcessing/Source/DynamicMesh/Private/Operations/RepairOrientation.cpp b/Engine/Plugins/Runtime/GeometryProcessing/Source/DynamicMesh/Private/Operations/RepairOrientation.cpp
--- a/Engine/Plugins/Runtime/GeometryProcessing/Source/DynamicMesh/Private/Operations/RepairOrientation.cpp
+++ b/Engine/Plugins/Runtime/GeometryProcessing/Source/DynamicMesh/Private/Operations/RepairOrientation.cpp
@@ -28,15 +28,12 @@ void FMeshRepairOrientation::OrientComponents()
 	{
 		Remaining.Add(TID);
 	}
+	// note: this is only called after we've already verified there are elements in Remaining
 	auto PopOneTri = [&Remaining]()
 	{
-		for (int One : Remaining)
-		{
-			Remaining.Remove(One);
-			return One;
-		}
-		check(false);
-		return FDynamicMesh3::InvalidID;
+		int32 One = *Remaining.CreateConstIterator();
+		Remaining.Remove(One);
+		return One;
 	};
 	TArray<int> Stack;
 	while (Remaining.Num() > 0)
diff --git a/Engine/Plugins/Runtime/Metasound/Source/MetasoundEngine/Private/MetasoundSource.cpp b/Engine/Plugins/Runtime/Metasound/Source/MetasoundEngine/Private/MetasoundSource.cpp
--- a/Engine/Plugins/Runtime/Metasound/Source/MetasoundEngine/Private/MetasoundSource.cpp
+++ b/Engine/Plugins/Runtime/Metasound/Source/MetasoundEngine/Private/MetasoundSource.cpp
@@ -1713,12 +1713,7 @@ TSharedPtr<const Metasound::IGraph> UMetaSoundSource::TryGetMetaSoundPresetBaseG
 	if (ReferencedAssetClassObjects.Num() == 1)
 	{
 		// Get first element from TSet<>
-		TObjectPtr<const UObject> BaseGraph;
-		for (const TObjectPtr<UObject>& ReferencedGraph : ReferencedAssetClassObjects)
-		{
-			BaseGraph = ReferencedGraph;
-			break;
-		}
+		TObjectPtr<const UObject> BaseGraph = *ReferencedAssetClassObjects.CreateConstIterator();
 
 		// Get the reference graph as a UMetaSoundSource
 		TObjectPtr<const UMetaSoundSource> BaseMetaSoundSource = Cast<const UMetaSoundSource>(BaseGraph);
diff --git a/Engine/Source/Runtime/GeometryCore/Private/CompGeom/ConvexHull3.cpp b/Engine/Source/Runtime/GeometryCore/Private/CompGeom/ConvexHull3.cpp
--- a/Engine/Source/Runtime/GeometryCore/Private/CompGeom/ConvexHull3.cpp
+++ b/Engine/Source/Runtime/GeometryCore/Private/CompGeom/ConvexHull3.cpp
@@ -531,14 +531,14 @@ struct FHullConnectivity
 		}
 		else
 		{
-			for (int32 TriIdx : TrisWithPoints)
+			if (const auto TriIdxItr = TrisWithPoints.CreateConstIterator())
 			{
+				int32 TriIdx = *TriIdxItr;
 				checkSlow(VisiblePoints[TriIdx].Num() > 0);
 
 				FoundTriPointPair[0] = TriIdx;
 				// choose the "max point" -- the point with the largest volume when it makes a tetrahedron w/ the triangle
 				FoundTriPointPair[1] = VisiblePoints[TriIdx].MaxPt();
-				break;
 			}
 		}
 		return FoundTriPointPair;
diff --git a/Engine/Source/Runtime/GeometryCore/Private/Operations/MeshMeshCut.cpp b/Engine/Source/Runtime/GeometryCore/Private/Operations/MeshMeshCut.cpp
--- a/Engine/Source/Runtime/GeometryCore/Private/Operations/MeshMeshCut.cpp
+++ b/Engine/Source/Runtime/GeometryCore/Private/Operations/MeshMeshCut.cpp
@@ -176,12 +176,11 @@ namespace MeshCut
 			while (FaceVertices.Num() > 0)
 			{
 				int TID = -1, PtIdx = -1;
-				// hacky way to pop one element -- TODO: this seems gross!? do it better?
-				for (TPair<int, int> TIDToPtIdx : FaceVertices)
+				// use an iterator to get a single element
 				{
-					TID = TIDToPtIdx.Key;
-					PtIdx = TIDToPtIdx.Value;
-					break;
+					const auto TIDToPtIdxItr = FaceVertices.CreateConstIterator();
+					TID = TIDToPtIdxItr.Key();
+					PtIdx = TIDToPtIdxItr.Value();
 				}
 				PtIndices.Reset();
 				FaceVertices.MultiFind(TID, PtIndices);
@@ -240,12 +239,11 @@ namespace MeshCut
 			while (EdgeVertices.Num() > 0)
 			{
 				int EID = -1, PtIdx = -1;
-				// hacky way to pop one element -- TODO: this seems gross!? do it better?
-				for (TPair<int, int> EIDToPtIdx : EdgeVertices)
+				// use an iterator to get a single element
 				{
-					EID = EIDToPtIdx.Key;
-					PtIdx = EIDToPtIdx.Value;
-					break;
+					const auto EIDToPtIdxItr = EdgeVertices.CreateConstIterator();
+					EID = EIDToPtIdxItr.Key();
+					PtIdx = EIDToPtIdxItr.Value();
 				}
 				PtIndices.Reset();
 				EdgeVertices.MultiFind(EID, PtIndices);
diff --git a/Engine/Source/Runtime/RenderCore/Public/GlobalShader.h b/Engine/Source/Runtime/RenderCore/Public/GlobalShader.h
--- a/Engine/Source/Runtime/RenderCore/Public/GlobalShader.h
+++ b/Engine/Source/Runtime/RenderCore/Public/GlobalShader.h
@@ -259,9 +259,9 @@ public:
 
 	const FGlobalShaderMapSection* GetFirstSection()
 	{
-		for (auto It : SectionMap)
+		if (const auto It = SectionMap.CreateConstIterator())
 		{
-			return It.Value;
+			return It.Value();
 		}
 		return nullptr;
 	}
