make geometrycore's box types cast to/from float as well as double engine box types

diff --git a/Engine/Source/Runtime/GeometryCore/Public/BoxTypes.h b/Engine/Source/Runtime/GeometryCore/Public/BoxTypes.h
--- a/Engine/Source/Runtime/GeometryCore/Public/BoxTypes.h
+++ b/Engine/Source/Runtime/GeometryCore/Public/BoxTypes.h
@@ -327,13 +327,15 @@ struct TAxisAlignedBox3
 	}
 
 
-	explicit operator FBox() const
+	template<typename OutRealType>
+	explicit operator TBox<OutRealType>() const
 	{
-		FBox ToRet((FVector)Min, (FVector)Max);
+		TBox<OutRealType> ToRet((TVector<OutRealType>)Min, (TVector<OutRealType>)Max);
 		ToRet.IsValid = !IsEmpty();
 		return ToRet;
 	}
-	TAxisAlignedBox3(const FBox& Box)
+	template<typename InRealType>
+	TAxisAlignedBox3(const TBox<InRealType>& Box)
 	{
 		if (Box.IsValid)
 		{
@@ -647,18 +649,20 @@ struct TAxisAlignedBox2
 		this->Max = TVector2<RealType>(Center.X + HalfWidth, Center.Y + HalfWidth);
 	}
 
-	explicit operator FBox2D() const
+	template<typename OutRealType>
+	explicit operator TBox2<OutRealType>() const
 	{
-		FBox2D ToRet((FVector2D)Min, (FVector2D)Max);
+		TBox2<OutRealType> ToRet((TVector2<OutRealType>)Min, (TVector2<OutRealType>)Max);
 		ToRet.bIsValid = !IsEmpty();
 		return ToRet;
 	}
-	TAxisAlignedBox2(const FBox2D& Box)
+	template<typename InRealType>
+	TAxisAlignedBox2(const TBox2<InRealType>& Box)
 	{
 		if (Box.bIsValid)
 		{
-			Min = Box.Min;
-			Max = Box.Max;
+			Min = TVector2<RealType>(Box.Min);
+			Max = TVector2<RealType>(Box.Max);
 		}
 		else
 		{
