fix issue where point sampling would compute incorrect mesh triangle areas for non-compact meshes

diff --git a/Engine/Source/Runtime/GeometryCore/Private/Sampling/MeshSurfacePointSampling.cpp b/Engine/Source/Runtime/GeometryCore/Private/Sampling/MeshSurfacePointSampling.cpp
--- a/Engine/Source/Runtime/GeometryCore/Private/Sampling/MeshSurfacePointSampling.cpp
+++ b/Engine/Source/Runtime/GeometryCore/Private/Sampling/MeshSurfacePointSampling.cpp
@@ -29,7 +29,7 @@ struct FTriangleInfoCache
 	void InitializeForTriangleSet(const MeshType& SampleMesh)
 	{
 		TriNormals.SetNumUninitialized(SampleMesh.MaxTriangleID());
-		TriAreas.SetNumUninitialized(SampleMesh.MaxTriangleID());
+		TriAreas.SetNumZeroed(SampleMesh.MaxTriangleID()); // note: zero the memory so we can include skipped triangle IDs in the sum below
 		TriFrames.SetNumUninitialized(SampleMesh.MaxTriangleID());
 		UVTriangles.SetNumUninitialized(SampleMesh.MaxTriangleID());
 				
@@ -46,6 +46,7 @@ struct FTriangleInfoCache
 					TriFrames[tid].ToPlaneUV(A), TriFrames[tid].ToPlaneUV(B), TriFrames[tid].ToPlaneUV(C));
 			}
 		});
+		// Note: This sum includes areas of skipped triangle IDs, which are zero by initialization above
 		TotalArea = 0;
 		for (double TriArea : TriAreas)
 		{
