fix checkSlow when traversing through empty root cells to occupied children of sparse dynamic point octree

diff --git a/Engine/Source/Runtime/GeometryCore/Public/Spatial/SparseDynamicPointOctree3.h b/Engine/Source/Runtime/GeometryCore/Public/Spatial/SparseDynamicPointOctree3.h
--- a/Engine/Source/Runtime/GeometryCore/Public/Spatial/SparseDynamicPointOctree3.h
+++ b/Engine/Source/Runtime/GeometryCore/Public/Spatial/SparseDynamicPointOctree3.h
@@ -954,13 +954,16 @@ void FSparseDynamicPointOctree3::RangeQuery(
 		const FSparsePointOctreeCell* CurCell = Queue.Pop(EAllowShrinking::No);
 
 		// process elements
-		CellPointLists.Enumerate(CurCell->CellID, [&](int32 PointID)
+		if (CellPointLists.IsAllocated(CurCell->CellID))
 		{
-			if (PredicateFunc(PointID))
+			CellPointLists.Enumerate(CurCell->CellID, [&](int32 PointID)
 			{
-				PointIDs.Add(PointID);
-			}
-		});
+				if (PredicateFunc(PointID))
+				{
+					PointIDs.Add(PointID);
+				}
+			});
+		}
 
 		for (int k = 0; k < 8; ++k)
 		{
