make all the (geometry script and geometry core) mesh transform methods consistently transform normals and tangents + add a RotateMesh to geometrycore to replace the incorrect implementation in geometry script

Note we make the normal/tangent transform optional for the Scale method, because the docs intentionally called it out as not transforming normals (but w/ the default behavior changed to correctly transform them)

diff --git a/Engine/Plugins/Runtime/GeometryScripting/Source/GeometryScriptingCore/Private/MeshTransformFunctions.cpp b/Engine/Plugins/Runtime/GeometryScripting/Source/GeometryScriptingCore/Private/MeshTransformFunctions.cpp
--- a/Engine/Plugins/Runtime/GeometryScripting/Source/GeometryScriptingCore/Private/MeshTransformFunctions.cpp
+++ b/Engine/Plugins/Runtime/GeometryScripting/Source/GeometryScriptingCore/Private/MeshTransformFunctions.cpp
@@ -75,15 +75,7 @@ UDynamicMesh* UGeometryScriptLibrary_MeshTransformFunctions::RotateMesh(
 	// todo: publish correct change types
 	TargetMesh->EditMesh([&](FDynamicMesh3& EditMesh) 
 	{
-		int NumVertices = EditMesh.MaxVertexID();
-		ParallelFor(NumVertices, [&](int32 vid)
-		{
-			if (EditMesh.IsVertex(vid))
-			{
-				EditMesh.SetVertex(vid, Rotation.RotateVector( (EditMesh.GetVertex(vid) - RotationOrigin) ) + RotationOrigin);
-			}
-		});
-
+		MeshTransforms::Rotate(EditMesh, Rotation, RotationOrigin);
 	}, EDynamicMeshChangeType::GeneralEdit, EDynamicMeshAttributeChangeFlags::Unknown, false);
 
 	return TargetMesh;
diff --git a/Engine/Source/Runtime/GeometryCore/Private/DynamicMesh/MeshTransforms.cpp b/Engine/Source/Runtime/GeometryCore/Private/DynamicMesh/MeshTransforms.cpp
--- a/Engine/Source/Runtime/GeometryCore/Private/DynamicMesh/MeshTransforms.cpp
+++ b/Engine/Source/Runtime/GeometryCore/Private/DynamicMesh/MeshTransforms.cpp
@@ -21,16 +21,60 @@ void MeshTransforms::Translate(FDynamicMesh3& Mesh, const FVector3d& Translation
 }
 
 
-void MeshTransforms::Scale(FDynamicMesh3& Mesh, const FVector3d& Scale, const FVector3d& Origin, bool bReverseOrientationIfNeeded)
+void MeshTransforms::Scale(FDynamicMesh3& Mesh, const FVector3d& Scale, const FVector3d& Origin, bool bReverseOrientationIfNeeded, bool bOnlyPositions)
 {
+	bool bVertexNormals = Mesh.HasVertexNormals();
+	bool bNeedNormalTangentScaling = !bOnlyPositions && !Scale.IsUniform();
 	int NumVertices = Mesh.MaxVertexID();
+	FVector3f NormalScale = (FVector3f)Scale;
+	if (bNeedNormalTangentScaling)
+	{
+		// compute a safe inverse-scale to apply to normal vectors
+		for (int32 Idx = 0; Idx < 3; ++Idx)
+		{
+			if (!FMath::IsNearlyZero(NormalScale[Idx]))
+			{
+				NormalScale[Idx] = 1.0f / NormalScale[Idx];
+			}
+		}
+	}
 	ParallelFor(NumVertices, [&](int32 vid)
 	{
 		if (Mesh.IsVertex(vid))
 		{
 			Mesh.SetVertex(vid, (Mesh.GetVertex(vid) - Origin) * Scale + Origin);
+
+			if (bNeedNormalTangentScaling && bVertexNormals)
+			{
+				FVector3f ScaledNormal = Mesh.GetVertexNormal(vid) * NormalScale;
+				ScaledNormal.Normalize();
+				Mesh.SetVertexNormal(vid, ScaledNormal);
+			}
 		}
 	});
+	if (bNeedNormalTangentScaling && Mesh.HasAttributes())
+	{
+		for (int32 NormalTangentLayerIdx = 0; NormalTangentLayerIdx < 3; ++NormalTangentLayerIdx)
+		{
+			FVector3f ScaleBy = (FVector3f)Scale; // tangents are transformed by the scale directly
+			if (NormalTangentLayerIdx == 0) // normal is transformed by the inverse scale
+			{
+				ScaleBy = NormalScale;
+			}
+			FDynamicMeshNormalOverlay* Normals = Mesh.Attributes()->GetNormalLayer(NormalTangentLayerIdx);
+			int NumNormals = Normals->MaxElementID();
+			ParallelFor(NumNormals, [&](int ElemID)
+				{
+					if (Normals->IsElement(ElemID))
+					{
+						FVector3f ScaledNormal = Normals->GetElement(ElemID) * ScaleBy;
+						ScaledNormal.Normalize();
+						Normals->SetElement(ElemID, ScaledNormal);
+					}
+				});
+		}
+	}
+
 	if (bReverseOrientationIfNeeded && Scale.X * Scale.Y * Scale.Z < 0)
 	{
 		Mesh.ReverseOrientation(false);
@@ -61,16 +105,20 @@ void MeshTransforms::WorldToFrameCoords(FDynamicMesh3& Mesh, const FFrame3d& Fra
 
 	if (Mesh.HasAttributes())
 	{
-		FDynamicMeshNormalOverlay* Normals = Mesh.Attributes()->PrimaryNormals();
-		int NumNormals = Normals->MaxElementID();
-		ParallelFor(NumNormals, [&](int elemid)
+		for (int32 NormalTangentLayerIdx = 0; NormalTangentLayerIdx < Mesh.Attributes()->NumNormalLayers(); ++NormalTangentLayerIdx)
 		{
-			if (Normals->IsElement(elemid))
+			// Note: Because this is a rotation transform, normals and tangents transform the same way
+			FDynamicMeshNormalOverlay* Normals = Mesh.Attributes()->GetNormalLayer(NormalTangentLayerIdx);
+			int NumNormals = Normals->MaxElementID();
+			ParallelFor(NumNormals, [&](int ElemID)
 			{
-				FVector3f Normal = Normals->GetElement(elemid);
-				Normals->SetElement(elemid, (FVector3f)Frame.ToFrameVector((FVector3d)Normal));
-			}
-		});
+				if (Normals->IsElement(ElemID))
+				{
+					FVector3f Normal = Normals->GetElement(ElemID);
+					Normals->SetElement(ElemID, (FVector3f)Frame.ToFrameVector((FVector3d)Normal));
+				}
+			});
+		}
 	}
 }
 
@@ -100,47 +148,118 @@ void MeshTransforms::FrameCoordsToWorld(FDynamicMesh3& Mesh, const FFrame3d& Fra
 
 	if (Mesh.HasAttributes())
 	{
-		FDynamicMeshNormalOverlay* Normals = Mesh.Attributes()->PrimaryNormals();
-		int NumNormals = Normals->MaxElementID();
-		ParallelFor(NumNormals, [&](int elemid)
+		for (int32 NormalTangentLayerIdx = 0; NormalTangentLayerIdx < Mesh.Attributes()->NumNormalLayers(); ++NormalTangentLayerIdx)
 		{
-			if (Normals->IsElement(elemid))
+			// Note: Because this is a rotation transform, normals and tangents transform the same way
+			FDynamicMeshNormalOverlay* Normals = Mesh.Attributes()->GetNormalLayer(NormalTangentLayerIdx);
+			int NumNormals = Normals->MaxElementID();
+			ParallelFor(NumNormals, [&](int ElemID)
 			{
-				FVector3f Normal = Normals->GetElement(elemid);
-				Normals->SetElement(elemid, (FVector3f)Frame.FromFrameVector((FVector3d)Normal));
-			}
-		});
+				if (Normals->IsElement(ElemID))
+				{
+					FVector3f Normal = Normals->GetElement(ElemID);
+					Normals->SetElement(ElemID, (FVector3f)Frame.FromFrameVector((FVector3d)Normal));
+				}
+			});
+		}
 	}
 }
 
 
+void MeshTransforms::Rotate(FDynamicMesh3& Mesh, const FRotator& Rotation, const FVector3d& RotationOrigin)
+{
+	bool bVertexNormals = Mesh.HasVertexNormals();
 
+	int NumVertices = Mesh.MaxVertexID();
+	ParallelFor(NumVertices, [&](int vid)
+	{
+		if (Mesh.IsVertex(vid)) 
+		{
+			FVector3d Position = Rotation.RotateVector((Mesh.GetVertex(vid) - RotationOrigin)) + RotationOrigin;
+			Mesh.SetVertex(vid, Position);
+
+			if (bVertexNormals)
+			{
+				FVector3f Normal = Mesh.GetVertexNormal(vid);
+				Normal = (FVector3f)Rotation.RotateVector((FVector3d)Normal);
+				Mesh.SetVertexNormal(vid, Normal);
+			}
+		}
+	});
+
+	if (Mesh.HasAttributes())
+	{
+		for (int32 NormalLayerIdx = 0; NormalLayerIdx < Mesh.Attributes()->NumNormalLayers(); ++NormalLayerIdx)
+		{
+			FDynamicMeshNormalOverlay* NormalTangentLayer = Mesh.Attributes()->GetNormalLayer(NormalLayerIdx);
+			int NumElements = NormalTangentLayer->MaxElementID();
+			ParallelFor(NumElements, [&](int ElemID)
+			{
+				if (NormalTangentLayer->IsElement(ElemID))
+				{
+					FVector3f NormalTangent = NormalTangentLayer->GetElement(ElemID);
+					NormalTangent = (FVector3f)Rotation.RotateVector((FVector3d)NormalTangent);
+					NormalTangentLayer->SetElement(ElemID, NormalTangent);
+				}
+			});
+		}
+	}
+}
 
 void MeshTransforms::ApplyTransform(FDynamicMesh3& Mesh, const FTransformSRT3d& Transform, bool bReverseOrientationIfNeeded)
 {
 	bool bVertexNormals = Mesh.HasVertexNormals();
 
-	for (int vid : Mesh.VertexIndicesItr())
+	int NumVertices = Mesh.MaxVertexID();
+	ParallelFor(NumVertices, [&](int vid)
 	{
-		FVector3d Position = Mesh.GetVertex(vid);
-		Position = Transform.TransformPosition(Position);
-		Mesh.SetVertex(vid, Position);
-
-		if (bVertexNormals)
+		if (Mesh.IsVertex(vid)) 
 		{
-			FVector3f Normal = Mesh.GetVertexNormal(vid);
-			Normal = (FVector3f)Transform.TransformNormal((FVector3d)Normal);
-			Mesh.SetVertexNormal(vid, Normalized(Normal));
+			FVector3d Position = Mesh.GetVertex(vid);
+			Position = Transform.TransformPosition(Position);
+			Mesh.SetVertex(vid, Position);
+
+			if (bVertexNormals)
+			{
+				FVector3f Normal = Mesh.GetVertexNormal(vid);
+				Normal = (FVector3f)Transform.TransformNormal((FVector3d)Normal);
+				Mesh.SetVertexNormal(vid, Normal);
+			}
 		}
-	}
+	});
+
 	if (Mesh.HasAttributes())
 	{
 		FDynamicMeshNormalOverlay* Normals = Mesh.Attributes()->PrimaryNormals();
-		for (int elemid : Normals->ElementIndicesItr())
+		if (Normals)
 		{
-			FVector3f Normal = Normals->GetElement(elemid);
-			Normal = (FVector3f)Transform.TransformNormal((FVector3d)Normal);
-			Normals->SetElement(elemid, Normalized(Normal));
+			int NumNormals = Normals->MaxElementID();
+			ParallelFor(NumNormals, [&](int ElemID)
+			{
+				if (Normals->IsElement(ElemID))
+				{
+					FVector3f Normal = Normals->GetElement(ElemID);
+					Normal = (FVector3f)Transform.TransformNormal((FVector3d)Normal);
+					Normals->SetElement(ElemID, Normal);
+				}
+			});
+		}
+		if (Mesh.Attributes()->HasTangentSpace())
+		{
+			for (int32 TangentLayerIdx = 1; TangentLayerIdx < 3; ++TangentLayerIdx)
+			{
+				FDynamicMeshNormalOverlay* TangentLayer = Mesh.Attributes()->GetNormalLayer(TangentLayerIdx);
+				int NumTangents = TangentLayer->MaxElementID();
+				ParallelFor(NumTangents, [&](int ElemID)
+				{
+					if (TangentLayer->IsElement(ElemID))
+					{
+						FVector3f Tangent = TangentLayer->GetElement(ElemID);
+						Tangent = (FVector3f)Transform.TransformVector((FVector3d)Tangent);
+						TangentLayer->SetElement(ElemID, Normalized(Tangent));
+					}
+				});
+			}
 		}
 	}
 
@@ -169,7 +288,7 @@ void MeshTransforms::ApplyTransformInverse(FDynamicMesh3& Mesh, const FTransform
 			{
 				FVector3f Normal = Mesh.GetVertexNormal(vid);
 				Normal = (FVector3f)Transform.InverseTransformNormal((FVector3d)Normal);
-				Mesh.SetVertexNormal(vid, Normalized(Normal));
+				Mesh.SetVertexNormal(vid, Normal);
 			}
 		}
 	});
@@ -180,16 +299,33 @@ void MeshTransforms::ApplyTransformInverse(FDynamicMesh3& Mesh, const FTransform
 		if (Normals)
 		{
 			int NumNormals = Normals->MaxElementID();
-			ParallelFor(NumNormals, [&](int elemid)
+			ParallelFor(NumNormals, [&](int ElemID)
 			{
-				if (Normals->IsElement(elemid))
+				if (Normals->IsElement(ElemID))
 				{
-					FVector3f Normal = Normals->GetElement(elemid);
+					FVector3f Normal = Normals->GetElement(ElemID);
 					Normal = (FVector3f)Transform.InverseTransformNormal((FVector3d)Normal);
-					Normals->SetElement(elemid, Normalized(Normal));
+					Normals->SetElement(ElemID, Normal);
 				}
 			});
 		}
+		if (Mesh.Attributes()->HasTangentSpace())
+		{
+			for (int32 TangentLayerIdx = 1; TangentLayerIdx < 3; ++TangentLayerIdx)
+			{
+				FDynamicMeshNormalOverlay* TangentLayer = Mesh.Attributes()->GetNormalLayer(TangentLayerIdx);
+				int NumTangents = TangentLayer->MaxElementID();
+				ParallelFor(NumTangents, [&](int ElemID)
+				{
+					if (TangentLayer->IsElement(ElemID))
+					{
+						FVector3f Tangent = TangentLayer->GetElement(ElemID);
+						Tangent = (FVector3f)Transform.InverseTransformVector((FVector3d)Tangent);
+						TangentLayer->SetElement(ElemID, Normalized(Tangent));
+					}
+				});
+			}
+		}
 	}
 
 	if (bReverseOrientationIfNeeded && Transform.GetDeterminant() < 0)
@@ -239,15 +375,77 @@ void MeshTransforms::ApplyTransform(FDynamicMesh3& Mesh,
 		if (Normals)
 		{
 			int NumNormals = Normals->MaxElementID();
-			ParallelFor(NumNormals, [&](int elemid)
+			ParallelFor(NumNormals, [&](int ElemID)
 			{
-				if (Normals->IsElement(elemid))
+				if (Normals->IsElement(ElemID))
 				{
-					FVector3f Normal = Normals->GetElement(elemid);
+					FVector3f Normal = Normals->GetElement(ElemID);
 					Normal = NormalTransform(Normal);
-					Normals->SetElement(elemid, Normalized(Normal));
+					Normals->SetElement(ElemID, Normalized(Normal));
 				}
 			});
 		}
 	}
 }
+
+void MeshTransforms::ApplyTransform(FDynamicMesh3& Mesh,
+	TFunctionRef<FVector3d(const FVector3d&)> PositionTransform,
+	TFunctionRef<FVector3f(const FVector3f&)> NormalTransform,
+	TFunctionRef<FVector3f(const FVector3f&)> TangentTransform)
+{
+	bool bVertexNormals = Mesh.HasVertexNormals();
+
+	int NumVertices = Mesh.MaxVertexID();
+	ParallelFor(NumVertices, [&](int vid)
+	{
+		if (Mesh.IsVertex(vid))
+		{
+			FVector3d Position = Mesh.GetVertex(vid);
+			Position = PositionTransform(Position);
+			Mesh.SetVertex(vid, Position);
+
+			if (bVertexNormals)
+			{
+				FVector3f Normal = Mesh.GetVertexNormal(vid);
+				Normal = NormalTransform(Normal);
+				Mesh.SetVertexNormal(vid, Normalized(Normal));
+			}
+		}
+	});
+
+	if (Mesh.HasAttributes())
+	{
+		FDynamicMeshNormalOverlay* Normals = Mesh.Attributes()->PrimaryNormals();
+		if (Normals)
+		{
+			int NumNormals = Normals->MaxElementID();
+			ParallelFor(NumNormals, [&](int ElemID)
+			{
+				if (Normals->IsElement(ElemID))
+				{
+					FVector3f Normal = Normals->GetElement(ElemID);
+					Normal = NormalTransform(Normal);
+					Normals->SetElement(ElemID, Normalized(Normal));
+				}
+			});
+		}
+		if (Mesh.Attributes()->HasTangentSpace())
+		{
+			for (int32 TangentLayerIdx = 1; TangentLayerIdx < 3; ++TangentLayerIdx)
+			{
+				FDynamicMeshNormalOverlay* TangentLayer = Mesh.Attributes()->GetNormalLayer(TangentLayerIdx);
+				int32 NumTangents = TangentLayer->MaxElementID();
+				ParallelFor(NumTangents, [&](int ElemID)
+				{
+					if (TangentLayer->IsElement(ElemID))
+					{
+						FVector3f Tangent = TangentLayer->GetElement(ElemID);
+						Tangent = TangentTransform(Tangent);
+						TangentLayer->SetElement(ElemID, Normalized(Tangent));
+					}
+				});
+			}
+		}
+	}
+}
+
diff --git a/Engine/Source/Runtime/GeometryCore/Public/DynamicMesh/MeshTransforms.h b/Engine/Source/Runtime/GeometryCore/Public/DynamicMesh/MeshTransforms.h
--- a/Engine/Source/Runtime/GeometryCore/Public/DynamicMesh/MeshTransforms.h
+++ b/Engine/Source/Runtime/GeometryCore/Public/DynamicMesh/MeshTransforms.h
@@ -27,10 +27,16 @@ namespace MeshTransforms
 	GEOMETRYCORE_API void Translate(FDynamicMesh3& Mesh, const FVector3d& Translation);
 
 	/**
-	 * Apply Scale to vertex positions of Mesh, relative to given Origin. Does not modify any other attributes (normals/etc)
+	 * Apply Scale to vertex positions of Mesh, relative to given Origin. For non-uniform scale, also scales normals/tangents.
 	 * @param bReverseOrientationIfNeeded		If negative scaling inverts the mesh, also invert the triangle orientations
+	 * @param OnlyPositions						If true, only apply scale positions and do not affect normals/tangents
 	 */
-	GEOMETRYCORE_API void Scale(FDynamicMesh3& Mesh, const FVector3d& Scale, const FVector3d& Origin, bool bReverseOrientationIfNeeded = false);
+	GEOMETRYCORE_API void Scale(FDynamicMesh3& Mesh, const FVector3d& Scale, const FVector3d& Origin, bool bReverseOrientationIfNeeded = false, bool bOnlyPositions = false);
+
+	/**
+	 * Apply Rotation to the vertex positions, relative to the given RotationOrigin.
+	 */
+	GEOMETRYCORE_API void Rotate(FDynamicMesh3& Mesh, const FRotator& Rotation, const FVector3d& RotationOrigin);
 
 	/**
 	 * Transform Mesh into local coordinates of Frame
@@ -67,6 +73,16 @@ namespace MeshTransforms
 		TFunctionRef<FVector3d(const FVector3d&)> PositionTransform,
 		TFunctionRef<FVector3f(const FVector3f&)> NormalTransform);
 
+	/**
+	 * Apply given Transform to a Mesh.
+	 * Modifies Vertex Positions, Normals, Tangents, and any Per-Triangle Normal Overlays
+	 */
+	GEOMETRYCORE_API void ApplyTransform(FDynamicMesh3& Mesh,
+		TFunctionRef<FVector3d(const FVector3d&)> PositionTransform,
+		TFunctionRef<FVector3f(const FVector3f&)> NormalTransform,
+		TFunctionRef<FVector3f(const FVector3f&)> TangentTransform
+	);
+
 	/**
 	 * If applying Transform would invert Mesh w/ a negative scale, then invert Mesh's triangle orientations.
 	 * Note: Does not apply the transform.
