convex decomposition robustness improvements -- pick good offset directions for degenerate cases where normals are not available, and do non-solid plane cuts for parts where the plane cut gave unreliable results (created holes)

diff --git a/Engine/Source/Runtime/GeometryCore/Private/CompGeom/ConvexDecomposition3.cpp b/Engine/Source/Runtime/GeometryCore/Private/CompGeom/ConvexDecomposition3.cpp
--- a/Engine/Source/Runtime/GeometryCore/Private/CompGeom/ConvexDecomposition3.cpp
+++ b/Engine/Source/Runtime/GeometryCore/Private/CompGeom/ConvexDecomposition3.cpp
@@ -1410,20 +1410,61 @@ struct FPartialCutResult
 				// optionally re-try w/ minimal offsets in degen directions
 				if (ThickenHullAfterFailure > 0)
 				{
+					int32 Dimension = HullCompute.GetDimension();
 					const int32 OrigMaxID = Convex.InternalGeo.MaxVertexID();
 					FMeshNormals Normals(&Convex.InternalGeo);
-					Normals.ComputeVertexNormals();
+					bool bUseNormalsForOffset = Dimension >= 2;
+					if (bUseNormalsForOffset)
+					{
+						Normals.ComputeVertexNormals();
+					}
+					// Compute offsets to use based on the dimension of the failed convex hull
+					FVector3d OffsetBasis[3];
+					FVector3d FallbackOffset = FVector3d(1, 1, 1) * FMathd::InvSqrt3;
+					int32 NumOffsets = 0;
+					if (Dimension == 0)
+					{
+						OffsetBasis[0] = FVector3d(1, 0, 0);
+						OffsetBasis[1] = FVector3d(0, 1, 0);
+						OffsetBasis[2] = FVector3d(0, 0, 1);
+						NumOffsets = 3;
+					}
+					else if (Dimension == 1)
+					{
+						FLine3d Line = HullCompute.GetLine();
+						VectorUtil::MakePerpVectors(Line.Direction, OffsetBasis[0], OffsetBasis[1]);
+						NumOffsets = 2;
+					}
+					else if (Dimension == 2)
+					{
+						FPlane3d HullPlane = HullCompute.GetPlane();
+						OffsetBasis[0] = HullPlane.Normal;
+						FallbackOffset = HullPlane.Normal;
+						NumOffsets = 1;
+					}
 					const double OffsetFactor = ThickenHullAfterFailure;
 					for (int32 VID = 0; VID < OrigMaxID; ++VID)
 					{
 						if (Convex.InternalGeo.IsVertex(VID) && ForHull[Side][VID])
 						{
-							FVector3d Normal = Normals[VID];
-							if (Normal == FVector::ZeroVector)
+							
+							if (bUseNormalsForOffset)
 							{
-								Normal = FVector::OneVector * FMathd::InvSqrt3; // for degenerate normals, arbitrarily pick a diagonal offset direction
+								FVector3d Offset = Normals[VID];
+								if (Offset == FVector::ZeroVector)
+								{
+									Offset = FallbackOffset;
+								}
+								OffsetVertices[Side].Add(Convex.InternalGeo.GetVertex(VID) - Offset * OffsetFactor);
+							}
+							else
+							{
+								for (int32 OffsetIdx = 0; OffsetIdx < NumOffsets; ++OffsetIdx)
+								{
+									OffsetVertices[Side].Add(Convex.InternalGeo.GetVertex(VID) - OffsetBasis[OffsetIdx] * OffsetFactor);
+								}
 							}
-							OffsetVertices[Side].Add(Convex.InternalGeo.GetVertex(VID) - Normals[VID] * OffsetFactor);
+							
 						}
 					}
 					bOK = HullCompute.Solve(Convex.InternalGeo.MaxVertexID() + CutVertices.Num() + OffsetVertices[Side].Num(),
@@ -2028,11 +2069,13 @@ bool FConvexDecomposition3::SplitWorstHelper(bool bCanSkipUnreliableGeoVolumes,
 	double LowestError = FMathd::MaxReal;
 	int32 BestPlaneIdx = -1;
 	FPartialCutResult BestCutResult;
+	// use solid cuts only if the setting is enabled and we haven't found problems w/ this geometry in previous cuts
+	const bool bCutAsSolid = bTreatAsSolid && !Part.bGeometryVolumeUnreliable;
 
 	for (int32 PlaneIdx = 0; PlaneIdx < CandidatePlanes.Num(); PlaneIdx++)
 	{
 		const FPlane3d& Plane = CandidatePlanes[PlaneIdx];
-		FPartialCutResult PlaneResult(Part, Plane, OnPlaneTolerance, bTreatAsSolid, ThickenAfterHullFailure);
+		FPartialCutResult PlaneResult(Part, Plane, OnPlaneTolerance, bCutAsSolid, ThickenAfterHullFailure);
 		if (!PlaneResult.bSuccess)
 		{
 			continue;
@@ -2058,7 +2101,7 @@ bool FConvexDecomposition3::SplitWorstHelper(bool bCanSkipUnreliableGeoVolumes,
 
 	int32 NewPartsStartIdx = Decomposition.Num();
 	int32 OtherSideStartIdx = -1;
-	BestCutResult.ApplyToGeo(Decomposition, WorstIdx, CandidatePlanes[BestPlaneIdx], OtherSideStartIdx, OnPlaneTolerance, ConnectedComponentTolerance, NegativeSpace, bTreatAsSolid, bSplitDisconnectedComponents);
+	BestCutResult.ApplyToGeo(Decomposition, WorstIdx, CandidatePlanes[BestPlaneIdx], OtherSideStartIdx, OnPlaneTolerance, ConnectedComponentTolerance, NegativeSpace, bCutAsSolid, bSplitDisconnectedComponents);
 
 	UpdateProximitiesAfterSplit(WorstIdx, NewPartsStartIdx, CandidatePlanes[BestPlaneIdx], OtherSideStartIdx, OrigHullVolume);
 
