Geometry Script.  Fix a bug in MeshTransforms that caused a crash if the mesh had normals attributes but no tangents.

diff --git a/Engine/Source/Runtime/GeometryCore/Private/DynamicMesh/MeshTransforms.cpp b/Engine/Source/Runtime/GeometryCore/Private/DynamicMesh/MeshTransforms.cpp
--- a/Engine/Source/Runtime/GeometryCore/Private/DynamicMesh/MeshTransforms.cpp
+++ b/Engine/Source/Runtime/GeometryCore/Private/DynamicMesh/MeshTransforms.cpp
@@ -54,7 +54,8 @@ void MeshTransforms::Scale(FDynamicMesh3& Mesh, const FVector3d& Scale, const FV
 	});
 	if (bNeedNormalTangentScaling && Mesh.HasAttributes())
 	{
-		for (int32 NormalTangentLayerIdx = 0; NormalTangentLayerIdx < 3; ++NormalTangentLayerIdx)
+		const int32 NumNormalLayers = FMath::Min(Mesh.Attributes()->NumNormalLayers(), 3);
+		for (int32 NormalTangentLayerIdx = 0; NormalTangentLayerIdx < NumNormalLayers; ++NormalTangentLayerIdx)
 		{
 			FVector3f ScaleBy = (FVector3f)Scale; // tangents are transformed by the scale directly
 			if (NormalTangentLayerIdx == 0) // normal is transformed by the inverse scale
