Add an Enumerate method to PointHashGrid3 to reduce the need to build temp arrays, and support optional early out
+ Make the vertex connected components dist-threshold method use the new enumerate method, and early out if there is only one component

diff --git a/Engine/Source/Runtime/GeometryCore/Public/Spatial/PointHashGrid3.h b/Engine/Source/Runtime/GeometryCore/Public/Spatial/PointHashGrid3.h
--- a/Engine/Source/Runtime/GeometryCore/Public/Spatial/PointHashGrid3.h
+++ b/Engine/Source/Runtime/GeometryCore/Public/Spatial/PointHashGrid3.h
@@ -296,6 +296,59 @@ public:
 		return ResultsOut.Num() - InitialNum;
 	}
 
+	/**
+	 * Call ProcessPointFunc on all points in grid within a given sphere, or until the function returns false
+	 * Note: Not thread-safe to update, remove or insert points during enumeration.
+	 * @param QueryPoint the center of the query sphere
+	 * @param Radius the radius of the query sphere
+	 * @param DistanceSqFunc Function you provide which measures the squared distance between QueryPoint and a Value
+	 * @param ProcessPointFunc Function you provide to process each found point. Takes point data and the distance squared; returns true to continue iteration, or false to stop.
+	 * @param IgnoreFunc optional Function you may provide which will result in a Value being ignored if IgnoreFunc(Value) returns true
+	 */
+	void EnumeratePointsInBall(
+		const TVector<RealType>& QueryPoint, RealType Radius,
+		TFunctionRef<RealType(const PointDataType&)> DistanceSqFunc,
+		TFunctionRef<bool(PointDataType, double)> ProcessPointFunc,
+		TFunctionRef<bool(const PointDataType&)> IgnoreFunc = [](const PointDataType& data) { return false; }) const
+	{
+		if (!Hash.Num())
+		{
+			return;
+		}
+
+		FVector3i min_idx = Indexer.ToGrid(QueryPoint - Radius * TVector<RealType>::One());
+		FVector3i max_idx = Indexer.ToGrid(QueryPoint + Radius * TVector<RealType>::One());
+
+		RealType RadiusSquared = Radius * Radius;
+
+		for (int zi = min_idx.Z; zi <= max_idx.Z; zi++)
+		{
+			for (int yi = min_idx.Y; yi <= max_idx.Y; yi++)
+			{
+				for (int xi = min_idx.X; xi <= max_idx.X; xi++)
+				{
+					FVector3i idx(xi, yi, zi);
+					for (typename TMultiMap<FVector3i, PointDataType>::TConstKeyIterator It = Hash.CreateConstKeyIterator(idx); It; ++It)
+					{
+						const PointDataType& Value = It.Value();
+						if (IgnoreFunc(Value))
+						{
+							continue;
+						}
+						RealType DistSq = DistanceSqFunc(Value);
+						if (DistSq < RadiusSquared)
+						{
+							if (!ProcessPointFunc(Value, DistSq))
+							{
+								return;
+							}
+						}
+					}
+				}
+			}
+		}
+	}
+
 private:
 	template<bool bEarlyOut = false>
 	TPair<PointDataType, RealType> FindInRadiusHelper(
diff --git a/Engine/Source/Runtime/GeometryCore/Public/VertexConnectedComponents.h b/Engine/Source/Runtime/GeometryCore/Public/VertexConnectedComponents.h
--- a/Engine/Source/Runtime/GeometryCore/Public/VertexConnectedComponents.h
+++ b/Engine/Source/Runtime/GeometryCore/Public/VertexConnectedComponents.h
@@ -206,7 +206,6 @@ public:
 	void ConnectCloseVertices(const TriangleMeshType& Mesh, double CloseVertexThreshold, int32 KeepSizeThreshold = 0)
 	{
 		TPointHashGrid3d<int32> VertexHash(CloseVertexThreshold * 3, -1);
-		TArray<int32> Neighbors;
 		for (int32 VID = 0; VID < Mesh.MaxVertexID(); VID++)
 		{
 			if (!Mesh.IsVertex(VID))
@@ -219,17 +218,16 @@ public:
 				continue;
 			}
 
+			int32 MaxSetSize = Mesh.VertexCount();
 			FVector3d Pt = Mesh.GetVertex(VID);
-			Neighbors.Reset();
-			VertexHash.FindPointsInBall(Pt, CloseVertexThreshold, [&Mesh, Pt](int32 OtherVID)
+			VertexHash.EnumeratePointsInBall(Pt, CloseVertexThreshold, [&Mesh, Pt](int32 OtherVID)
 				{
-					// TODO: could return MaxReal if OtherVID is already in the same component so we skip union-ing it?
 					return DistanceSquared(Pt, Mesh.GetVertex(OtherVID));
-				}, Neighbors);
-			for (int32 NbrVID : Neighbors)
-			{
-				DisjointSet.Union(SetID, NbrVID);
-			}
+				}, [this, SetID, MaxSetSize](const int32& NbrVID, double DistSq) 
+				{
+					int32 UnionSetID = DisjointSet.Union(SetID, NbrVID);
+					return DisjointSet.Sizes[UnionSetID] < MaxSetSize; // stop iterating if all vertices are in the same component
+				});
 			VertexHash.InsertPointUnsafe(VID, Pt);
 		}
 	}
