Specify the ConvexDecomposition3 thicken hull after failure parameter in world space, not local space, to match other parameters on that class

diff --git a/Engine/Source/Runtime/GeometryCore/Private/CompGeom/ConvexDecomposition3.cpp b/Engine/Source/Runtime/GeometryCore/Private/CompGeom/ConvexDecomposition3.cpp
--- a/Engine/Source/Runtime/GeometryCore/Private/CompGeom/ConvexDecomposition3.cpp
+++ b/Engine/Source/Runtime/GeometryCore/Private/CompGeom/ConvexDecomposition3.cpp
@@ -1828,6 +1828,7 @@ bool FConvexDecomposition3::SplitWorstHelper(bool bCanSkipUnreliableGeoVolumes,
 	}
 
 	double MinSplitSize = MinSplitSizeInWorldSpace <= 0 ? MinSplitSizeInWorldSpace : ConvertDistanceToleranceToLocalSpace(MinSplitSizeInWorldSpace);
+	double ThickenAfterHullFailureLocalSpace = ConvertDistanceToleranceToLocalSpace(ThickenAfterHullFailure);
 
 	double VolumeTolerance = ConvertDistanceToleranceToLocalVolumeTolerance(ErrorTolerance);
 
@@ -2075,7 +2076,7 @@ bool FConvexDecomposition3::SplitWorstHelper(bool bCanSkipUnreliableGeoVolumes,
 	for (int32 PlaneIdx = 0; PlaneIdx < CandidatePlanes.Num(); PlaneIdx++)
 	{
 		const FPlane3d& Plane = CandidatePlanes[PlaneIdx];
-		FPartialCutResult PlaneResult(Part, Plane, OnPlaneTolerance, bCutAsSolid, ThickenAfterHullFailure);
+		FPartialCutResult PlaneResult(Part, Plane, OnPlaneTolerance, bCutAsSolid, ThickenAfterHullFailureLocalSpace);
 		if (!PlaneResult.bSuccess)
 		{
 			continue;
@@ -2347,6 +2348,7 @@ int32 FConvexDecomposition3::MergeBest(const FMergeSettings& Settings)
 	// Support having a max error tolerance
 	double VolumeTolerance = ConvertDistanceToleranceToLocalVolumeTolerance(MaxErrorTolerance);
 	double MinThicknessTolerance = ConvertDistanceToleranceToLocalSpace(MinThicknessToleranceWorldSpace);
+	double ThickenAfterHullFailureLocalSpace = ConvertDistanceToleranceToLocalSpace(ThickenAfterHullFailure);
 
 	int32 MergeNum = 0;
 
@@ -2454,7 +2456,7 @@ int32 FConvexDecomposition3::MergeBest(const FMergeSettings& Settings)
 						{
 							continue;
 						}
-						FPartialCutResult PlaneResult(PartToSplit, Plane, OnPlaneTolerance, bTreatAsSolid, ThickenAfterHullFailure);
+						FPartialCutResult PlaneResult(PartToSplit, Plane, OnPlaneTolerance, bTreatAsSolid, ThickenAfterHullFailureLocalSpace);
 						if (!PlaneResult.bSuccess)
 						{
 							continue;
