Improvements to capsule generator (and corresponding add capsule tool and geometry script node):
 + Add a cylinder-section subdivision parameter, to control tesselation
 + Make UVs consistent across different subdivision / hemisphere step amounts

diff --git a/Engine/Plugins/Runtime/GeometryScripting/Source/GeometryScriptingCore/Private/MeshPrimitiveFunctions.cpp b/Engine/Plugins/Runtime/GeometryScripting/Source/GeometryScriptingCore/Private/MeshPrimitiveFunctions.cpp
--- a/Engine/Plugins/Runtime/GeometryScripting/Source/GeometryScriptingCore/Private/MeshPrimitiveFunctions.cpp
+++ b/Engine/Plugins/Runtime/GeometryScripting/Source/GeometryScriptingCore/Private/MeshPrimitiveFunctions.cpp
@@ -341,6 +341,7 @@ UDynamicMesh* UGeometryScriptLibrary_MeshPrimitiveFunctions::AppendCapsule(
 	float LineLength,
 	int32 HemisphereSteps,
 	int32 CircleSteps,
+	int32 SegmentSteps,
 	EGeometryScriptPrimitiveOriginMode Origin,
 	UGeometryScriptDebug* Debug)
 {
@@ -355,6 +356,7 @@ UDynamicMesh* UGeometryScriptLibrary_MeshPrimitiveFunctions::AppendCapsule(
 	CapsuleGenerator.SegmentLength = FMath::Max(FMathf::ZeroTolerance, LineLength);
 	CapsuleGenerator.NumHemisphereArcSteps = FMath::Max(2, HemisphereSteps);
 	CapsuleGenerator.NumCircleSteps = FMath::Max(3, CircleSteps);
+	CapsuleGenerator.NumSegmentSteps = FMath::Max(0, SegmentSteps);
 	CapsuleGenerator.bPolygroupPerQuad = (PrimitiveOptions.PolygroupMode == EGeometryScriptPrimitivePolygroupMode::PerQuad);
 	CapsuleGenerator.Generate();
 
diff --git a/Engine/Plugins/Runtime/GeometryScripting/Source/GeometryScriptingCore/Public/GeometryScript/MeshPrimitiveFunctions.h b/Engine/Plugins/Runtime/GeometryScripting/Source/GeometryScriptingCore/Public/GeometryScript/MeshPrimitiveFunctions.h
--- a/Engine/Plugins/Runtime/GeometryScripting/Source/GeometryScriptingCore/Public/GeometryScript/MeshPrimitiveFunctions.h
+++ b/Engine/Plugins/Runtime/GeometryScripting/Source/GeometryScriptingCore/Public/GeometryScript/MeshPrimitiveFunctions.h
@@ -232,9 +232,26 @@ public:
 		float LineLength = 75,
 		int32 HemisphereSteps = 5,
 		int32 CircleSteps = 8,
+		int32 SegmentSteps = 0,
 		EGeometryScriptPrimitiveOriginMode Origin = EGeometryScriptPrimitiveOriginMode::Base,
 		UGeometryScriptDebug* Debug = nullptr);
 
+	// Version of append capsule without SegmentSteps parameter
+	UE_DEPRECATED(5.5, "Use AppendCapsule with the SegmentSteps parameter, instead")
+	static UDynamicMesh* AppendCapsule(
+		UDynamicMesh* TargetMesh,
+		FGeometryScriptPrimitiveOptions PrimitiveOptions,
+		FTransform Transform,
+		float Radius = 30,
+		float LineLength = 75,
+		int32 HemisphereSteps = 5,
+		int32 CircleSteps = 8,
+		EGeometryScriptPrimitiveOriginMode Origin = EGeometryScriptPrimitiveOriginMode::Base,
+		UGeometryScriptDebug* Debug = nullptr)
+	{
+		return AppendCapsule(TargetMesh, PrimitiveOptions, Transform, Radius, LineLength, HemisphereSteps, CircleSteps, 0, Origin, Debug);
+	}
+
 	/**
 	* Appends a 3D Cylinder (with optional end caps) to the Target Mesh.
 	*/
diff --git a/Engine/Plugins/Runtime/MeshModelingToolset/Source/MeshModelingTools/Private/AddPrimitiveTool.cpp b/Engine/Plugins/Runtime/MeshModelingToolset/Source/MeshModelingTools/Private/AddPrimitiveTool.cpp
--- a/Engine/Plugins/Runtime/MeshModelingToolset/Source/MeshModelingTools/Private/AddPrimitiveTool.cpp
+++ b/Engine/Plugins/Runtime/MeshModelingToolset/Source/MeshModelingTools/Private/AddPrimitiveTool.cpp
@@ -717,6 +717,7 @@ void UAddCapsulePrimitiveTool::GenerateMesh(FDynamicMesh3* OutMesh) const
 	CapGen.SegmentLength = CapsuleSettings->CylinderLength;
 	CapGen.NumHemisphereArcSteps = CapsuleSettings->HemisphereSlices;
 	CapGen.NumCircleSteps = CapsuleSettings->CylinderSlices;
+	CapGen.NumSegmentSteps = CapsuleSettings->CylinderSubdivisions;
 	CapGen.bPolygroupPerQuad = ShapeSettings->PolygroupMode == EMakeMeshPolygroupMode::PerQuad;
 	CapGen.Generate();
 	OutMesh->Copy(&CapGen);
diff --git a/Engine/Plugins/Runtime/MeshModelingToolset/Source/MeshModelingTools/Public/AddPrimitiveTool.h b/Engine/Plugins/Runtime/MeshModelingToolset/Source/MeshModelingTools/Public/AddPrimitiveTool.h
--- a/Engine/Plugins/Runtime/MeshModelingToolset/Source/MeshModelingTools/Public/AddPrimitiveTool.h
+++ b/Engine/Plugins/Runtime/MeshModelingToolset/Source/MeshModelingTools/Public/AddPrimitiveTool.h
@@ -394,6 +394,10 @@ public:
 	/** Number of radial slices of the cylindrical section. */
 	UPROPERTY(EditAnywhere, Category = Shape, meta = (UIMin = "3", UIMax = "100", ClampMin = "3", ClampMax = "500", ProceduralShapeSetting))
 	int CylinderSlices = 16;
+
+	/** Number of lengthwise subdivisions along cylindrical section */
+	UPROPERTY(EditAnywhere, Category = Shape, meta = (UIMin = "0", UIMax = "100", ClampMin = "0", ClampMax = "500", ProceduralShapeSetting))
+	int CylinderSubdivisions = 1;
 };
 
 UENUM()
diff --git a/Engine/Source/Runtime/GeometryCore/Public/Generators/CapsuleGenerator.h b/Engine/Source/Runtime/GeometryCore/Public/Generators/CapsuleGenerator.h
--- a/Engine/Source/Runtime/GeometryCore/Public/Generators/CapsuleGenerator.h
+++ b/Engine/Source/Runtime/GeometryCore/Public/Generators/CapsuleGenerator.h
@@ -38,6 +38,9 @@ public:
 	/** Number of vertices along each circle */
 	int NumCircleSteps = 3;
 
+	/** Number of subdivisions lengthwise along the cylindrical section */
+	int NumSegmentSteps = 0;
+
 	/** If true, each quad gets a separate polygroup, otherwise the entire mesh is a single polygroup */
 	bool bPolygroupPerQuad = false;
 
@@ -81,6 +84,18 @@ private:
 				}
 			}
 
+			// add intermediate loops along the cylindrical section
+			double SegStepSize = 1.0 / (double)(NumSegmentSteps + 1.0);
+			double SegAlong = SegStepSize;
+			for (int32 SegStep = 0; SegStep < NumSegmentSteps; ++SegStep, SegAlong += SegStepSize)
+			{
+				for (t = 0, Theta = 0; t < NumCircleSteps; ++t, ++VtxIdx, Theta += Dtheta)
+				{
+					FVector3d Normal(FMath::Cos(Theta), FMath::Sin(Theta), 0.0);
+					SetVertex(VtxIdx, Normal * Radius + Offset * (1 - SegAlong), FVector3f(Normal));
+				}
+			}
+
 			// add points for second arc section
 			for (p = 1, Phi = FMathd::HalfPi; p < NumHemisphereArcSteps; ++p, Phi += Dphi) // NB: this skips the poles.
 			{
@@ -99,26 +114,46 @@ private:
 
 	void GenerateUVVertices()
 	{
-		// generate the UV's
-		int32 NumPhi = (2 * NumHemisphereArcSteps);
-		const float DUVphi = 1.0f / float(NumPhi - 1);
+		// generate the UVs
+		int32 NumPhi = (2 * NumHemisphereArcSteps + NumSegmentSteps);
 		const float DUVtheta = -1.0f / float(NumCircleSteps);
-
 		int32 UVIdx = 0;
-		int32 p,t;
-		float UVPhi, UVTheta;
-		for ( p = 1, UVPhi = DUVphi; p < NumPhi - 1; ++p, UVPhi += DUVphi)
+
+		// helper to add UVs for a given range in Phi
+		auto AddUVSpan = [this, &UVIdx, DUVtheta](int32 StepStart, int32 NumSteps, float PhiStart, float PhiStepSize)
 		{
-			for (t = 0, UVTheta = 1; t < NumCircleSteps; ++t, ++UVIdx, UVTheta += DUVtheta)
+			float UVPhi = PhiStart;
+			int32 PIdx = StepStart;
+			for (; PIdx < StepStart + NumSteps; ++PIdx, UVPhi += PhiStepSize)
 			{
+				float UVTheta = 1;
+				for (int32 t = 0; t < NumCircleSteps; ++t, ++UVIdx, UVTheta += DUVtheta)
+				{
+					UVs[UVIdx] = FVector2f(UVTheta, UVPhi);
+					UVParentVertex[UVIdx] = PIdx * NumCircleSteps + t;
+				}
 				UVs[UVIdx] = FVector2f(UVTheta, UVPhi);
-				UVParentVertex[UVIdx] = (p - 1) * NumCircleSteps + t;
+				UVParentVertex[UVIdx] = PIdx * NumCircleSteps; // Wrap around
+				++UVIdx;
 			}
-			UVs[UVIdx] = FVector2f(UVTheta, UVPhi);
-			UVParentVertex[UVIdx] = (p - 1) * NumCircleSteps; // Wrap around
-			++UVIdx;
-		}
-		int32 NorthPoleVtxIdx = (NumPhi - 2) * NumCircleSteps;
+			return PIdx;
+		};
+
+		const float PhiSpan = static_cast<float>(2 * Radius + SegmentLength);
+		const float HemisphereStepSize = static_cast<float>(Radius) / (PhiSpan * static_cast<float>(NumHemisphereArcSteps - 1));
+		// add the first hemisphere cap (except the pole)
+		int32 PIdx = AddUVSpan(1, NumHemisphereArcSteps - 1, HemisphereStepSize, HemisphereStepSize);
+
+		// add the cylindrical section
+		const float SegmentStepSize = static_cast<float>(SegmentLength) / (PhiSpan * static_cast<float>(NumSegmentSteps + 1));
+		PIdx = AddUVSpan(PIdx, NumSegmentSteps, static_cast<float>(Radius) / PhiSpan + SegmentStepSize, SegmentStepSize);
+
+		// add the closing hemisphere end-cap (except the pole)
+		AddUVSpan(PIdx, NumHemisphereArcSteps - 1, static_cast<float>(Radius + SegmentLength) / PhiSpan, HemisphereStepSize);
+
+		int32 NorthPoleVtxIdx = (NumPhi - 2 + NumSegmentSteps) * NumCircleSteps;
+		float UVTheta;
+		int32 t;
 		for (t = 0, UVTheta = 1 + DUVtheta; t < NumCircleSteps; ++t, ++UVIdx, UVTheta += DUVtheta)
 		{
 			UVs[UVIdx] = FVector2f(UVTheta, 0.0);
@@ -143,7 +178,7 @@ private:
 
 	void OutputEquatorialTriangles()
 	{
-		int32 NumPhi = (2 * NumHemisphereArcSteps);
+		int32 NumPhi = (2 * NumHemisphereArcSteps + NumSegmentSteps);
 		int32 TriIdx = 0, PolyIdx = 0;
 
 		// Generate equatorial triangles
@@ -184,7 +219,7 @@ private:
 
 	void OutputPolarTriangles()
 	{
-		int32 NumPhi = (2 * NumHemisphereArcSteps);
+		int32 NumPhi = (2 * NumHemisphereArcSteps + NumSegmentSteps);
 		const int32 NumEquatorialVtx = (NumPhi - 2) * NumCircleSteps;
 		const int32 NumEquatorialUVVtx = (NumPhi - 2) * (NumCircleSteps + 1);
 		const int32 NorthPoleVtxIdx = NumEquatorialVtx;
@@ -230,9 +265,9 @@ public:
 		NumHemisphereArcSteps = FMath::Max(NumHemisphereArcSteps, 2);
 		int32 NumPhi = (2 * NumHemisphereArcSteps);
 		NumCircleSteps = FMath::Max(NumCircleSteps, 3);
-		const int32 NumVertices = (NumPhi - 2) * NumCircleSteps + 2;
-		const int32 NumUVs = (NumPhi - 2) * (NumCircleSteps + 1) + (2 * NumCircleSteps);
-		const int32 NumTris = (NumPhi - 2) * NumCircleSteps * 2;
+		const int32 NumVertices = (NumPhi - 2 + NumSegmentSteps) * NumCircleSteps + 2;
+		const int32 NumUVs = (NumPhi - 2 + NumSegmentSteps) * (NumCircleSteps + 1) + (2 * NumCircleSteps);
+		const int32 NumTris = (NumPhi - 2 + NumSegmentSteps) * NumCircleSteps * 2;
 		SetBufferSizes(NumVertices, NumTris, NumUVs, NumVertices);
 
 		GenerateVertices();
