remove unreferenced local variables to fix geometry processing tests build

diff --git a/Engine/Source/Runtime/GeometryCore/Public/MeshQueries.h b/Engine/Source/Runtime/GeometryCore/Public/MeshQueries.h
--- a/Engine/Source/Runtime/GeometryCore/Public/MeshQueries.h
+++ b/Engine/Source/Runtime/GeometryCore/Public/MeshQueries.h
@@ -429,7 +429,6 @@ public:
 			{
 				continue;
 			}
-			FVector3d a, b, c;
 			FTriangle3d Tri1;
 			Mesh1.GetTriVertices(TI, Tri1.V[0], Tri1.V[1], Tri1.V[2]);
 			for (int TJ = 0; TJ < Mesh2.MaxTriangleID(); TJ++)
