Set a default value for dynamic mesh CanCollapse edge parameter + note that it doesn't affect the result

diff --git a/Engine/Source/Runtime/GeometryCore/Public/DynamicMesh/DynamicMesh3.h b/Engine/Source/Runtime/GeometryCore/Public/DynamicMesh/DynamicMesh3.h
--- a/Engine/Source/Runtime/GeometryCore/Public/DynamicMesh/DynamicMesh3.h
+++ b/Engine/Source/Runtime/GeometryCore/Public/DynamicMesh/DynamicMesh3.h
@@ -1274,10 +1274,10 @@ public:
 	 * Tests whether collapsing the specified edge using the CollapseEdge function would succeed
 	 * @param KeepVertID index of the vertex that should be kept
 	 * @param RemoveVertID index of the vertex that should be removed
-	 * @param EdgeParameterT vKeep is moved to Lerp(KeepPos, RemovePos, EdgeParameterT)
+	 * @param EdgeParameterT vKeep is moved to Lerp(KeepPos, RemovePos, EdgeParameterT). Note: Does not currently affect whether the edge is collapsable.
 	 * @return Ok if the edge can be collapsed, or enum value indicating why the operation cannot be applied
 	 */
-	GEOMETRYCORE_API virtual EMeshResult CanCollapseEdge(int vKeep, int vRemove, double collapse_t) const;
+	GEOMETRYCORE_API virtual EMeshResult CanCollapseEdge(int vKeep, int vRemove, double EdgeParameterT = 0) const;
 
 	/**
 	 * Collapse the edge between the two vertices, if topologically possible.
