Make the dynamic mesh copy-from-shape-generator build faster, by building attributes after appending the base triangles + inlining AddEdgeInternal

diff --git a/Engine/Source/Runtime/GeometryCore/Private/DynamicMesh/DynamicMesh3.cpp b/Engine/Source/Runtime/GeometryCore/Private/DynamicMesh/DynamicMesh3.cpp
--- a/Engine/Source/Runtime/GeometryCore/Private/DynamicMesh/DynamicMesh3.cpp
+++ b/Engine/Source/Runtime/GeometryCore/Private/DynamicMesh/DynamicMesh3.cpp
@@ -131,11 +131,6 @@ bool FDynamicMesh3::Copy(const FMeshShapeGenerator* Generator)
 
 	EnableTriangleGroups();
 
-	if (Generator->HasAttributes())
-	{
-		EnableAttributes();
-	}
-
 	int NumVerts = Generator->Vertices.Num();
 	for (int i = 0; i < NumVerts; ++i)
 	{
@@ -145,27 +140,36 @@ bool FDynamicMesh3::Copy(const FMeshShapeGenerator* Generator)
 	int NumTris = Generator->Triangles.Num();
 	if (Generator->HasAttributes())
 	{
-		FDynamicMeshUVOverlay* UVOverlay = Attributes()->PrimaryUV();
-		FDynamicMeshNormalOverlay* NormalOverlay = Attributes()->PrimaryNormals();
-		int NumUVs = Generator->UVs.Num();
-		for (int i = 0; i < NumUVs; ++i)
+		bool bSuccess = true;
+		// First append all triangles w/out attributes enabled
+		for (int32 TID = 0; TID < NumTris; ++TID)
 		{
-			UVOverlay->AppendElement(Generator->UVs[i]);
+			int PolyID = Generator->TrianglePolygonIDs.Num() > 0 ? 1 + Generator->TrianglePolygonIDs[TID] : 0;
+			int AppendedTID = AppendTriangle(Generator->Triangles[TID], PolyID);
+			bSuccess &= bool(TID == AppendedTID);
 		}
-		int NumNormals = Generator->Normals.Num();
-		for (int i = 0; i < NumNormals; ++i)
+		// If they were successfully appended, enable and set attributes
+		// (doing this as a post-process is faster)
+		if (ensure(bSuccess))
 		{
-			NormalOverlay->AppendElement(Generator->Normals[i]);
-		}
+			EnableAttributes();
+			FDynamicMeshUVOverlay* UVOverlay = Attributes()->PrimaryUV();
+			FDynamicMeshNormalOverlay* NormalOverlay = Attributes()->PrimaryNormals();
+			int NumUVs = Generator->UVs.Num();
+			for (int i = 0; i < NumUVs; ++i)
+			{
+				UVOverlay->AppendElement(Generator->UVs[i]);
+			}
+			int NumNormals = Generator->Normals.Num();
+			for (int i = 0; i < NumNormals; ++i)
+			{
+				NormalOverlay->AppendElement(Generator->Normals[i]);
+			}
 
-		for (int i = 0; i < NumTris; ++i)
-		{
-			int PolyID = Generator->TrianglePolygonIDs.Num() > 0 ? 1 + Generator->TrianglePolygonIDs[i] : 0;
-			int tid = AppendTriangle(Generator->Triangles[i], PolyID);
-			if (ensure(tid == i))
+			for (int i = 0; i < NumTris; ++i)
 			{
-				UVOverlay->SetTriangle(tid, Generator->TriangleUVs[i]);
-				NormalOverlay->SetTriangle(tid, Generator->TriangleNormals[i]);
+				UVOverlay->SetTriangle(i, Generator->TriangleUVs[i]);
+				NormalOverlay->SetTriangle(i, Generator->TriangleNormals[i]);
 			}
 		}
 	}
@@ -1148,24 +1152,6 @@ bool FDynamicMesh3::CheckValidity(FValidityOptions ValidityOptions, EValidityChe
 
 
 
-
-
-
-
-
-int FDynamicMesh3::AddEdgeInternal(int vA, int vB, int tA, int tB)
-{
-	if (vB < vA) {
-		int t = vB; vB = vA; vA = t;
-	}
-	int eid = EdgeRefCounts.Allocate();
-	Edges.InsertAt(FEdge{{vA, vB},{tA, tB}}, eid);
-	VertexEdgeLists.Insert(vA, eid);
-	VertexEdgeLists.Insert(vB, eid);
-	return eid;
-}
-
-
 int FDynamicMesh3::AddTriangleInternal(int a, int b, int c, int e0, int e1, int e2)
 {
 	int tid = TriangleRefCounts.Allocate();
diff --git a/Engine/Source/Runtime/GeometryCore/Public/DynamicMesh/DynamicMesh3.h b/Engine/Source/Runtime/GeometryCore/Public/DynamicMesh/DynamicMesh3.h
--- a/Engine/Source/Runtime/GeometryCore/Public/DynamicMesh/DynamicMesh3.h
+++ b/Engine/Source/Runtime/GeometryCore/Public/DynamicMesh/DynamicMesh3.h
@@ -1408,7 +1408,18 @@ protected:
 		TriangleEdges[TriangleID] = FIndex3i(e0, e1, e2);
 	}
 
-	GEOMETRYCORE_API int AddEdgeInternal(int vA, int vB, int tA, int tB = InvalidID);
+	inline int AddEdgeInternal(int vA, int vB, int tA, int tB = InvalidID)
+	{
+		if (vB < vA) {
+			int t = vB; vB = vA; vA = t;
+		}
+		int eid = EdgeRefCounts.Allocate();
+		Edges.InsertAt(FEdge{ {vA, vB},{tA, tB} }, eid);
+		VertexEdgeLists.Insert(vA, eid);
+		VertexEdgeLists.Insert(vB, eid);
+		return eid;
+
+	}
 	GEOMETRYCORE_API int AddTriangleInternal(int a, int b, int c, int e0, int e1, int e2);
 
 	inline int ReplaceTriangleVertex(int TriangleID, int vOld, int vNew)
