Inline some geometrycore smalllistset methods and change them to use a template function, to reduce overhead of tfunction calls in dynamic mesh processing

diff --git a/Engine/Source/Runtime/GeometryCore/Private/Util/SmallListSet.cpp b/Engine/Source/Runtime/GeometryCore/Private/Util/SmallListSet.cpp
--- a/Engine/Source/Runtime/GeometryCore/Private/Util/SmallListSet.cpp
+++ b/Engine/Source/Runtime/GeometryCore/Private/Util/SmallListSet.cpp
@@ -235,139 +235,18 @@ bool FSmallListSet::Contains(int32 ListIndex, int32 Value) const
 
 
 
-int32 FSmallListSet::Find(int32 ListIndex, const TFunction<bool(int32)>& PredicateFunc, int32 InvalidValue) const
-{
-	int32 block_ptr = ListHeads[ListIndex];
-	if (block_ptr != NullValue)
-	{
-		int32 N = ListBlocks[block_ptr];
-		if (N < BLOCKSIZE)
-		{
-			int32 iEnd = block_ptr + N;
-			for (int32 i = block_ptr + 1; i <= iEnd; ++i)
-			{
-				int32 Value = ListBlocks[i];
-				if (PredicateFunc(Value))
-				{
-					return Value;
-				}
-			}
-		}
-		else
-		{
-			// we spilled to linked list, have to iterate through it as well
-			int32 iEnd = block_ptr + BLOCKSIZE;
-			for (int32 i = block_ptr + 1; i <= iEnd; ++i)
-			{
-				int32 Value = ListBlocks[i];
-				if (PredicateFunc(Value))
-				{
-					return Value;
-				}
-			}
-			int32 cur_ptr = ListBlocks[block_ptr + BLOCK_LIST_OFFSET];
-			while (cur_ptr != NullValue)
-			{
-				int32 Value = LinkedListElements[cur_ptr];
-				if (PredicateFunc(Value))
-				{
-					return Value;
-				}
-				cur_ptr = LinkedListElements[cur_ptr + 1];
-			}
-		}
-	}
-	return InvalidValue;
-}
 
 
 
 
 
 
-bool FSmallListSet::Replace(int32 ListIndex, const TFunction<bool(int32)>& PredicateFunc, int32 NewValue)
-{
-	int32 block_ptr = ListHeads[ListIndex];
-	if (block_ptr != NullValue)
-	{
-		int32 N = ListBlocks[block_ptr];
-		if (N < BLOCKSIZE)
-		{
-			int32 iEnd = block_ptr + N;
-			for (int32 i = block_ptr + 1; i <= iEnd; ++i)
-			{
-				int32 Value = ListBlocks[i];
-				if (PredicateFunc(Value))
-				{
-					ListBlocks[i] = NewValue;
-					return true;
-				}
-			}
-		}
-		else
-		{
-			// we spilled to linked list, have to iterate through it as well
-			int32 iEnd = block_ptr + BLOCKSIZE;
-			for (int32 i = block_ptr + 1; i <= iEnd; ++i)
-			{
-				int32 Value = ListBlocks[i];
-				if (PredicateFunc(Value))
-				{
-					ListBlocks[i] = NewValue;
-					return true;
-				}
-			}
-			int32 cur_ptr = ListBlocks[block_ptr + BLOCK_LIST_OFFSET];
-			while (cur_ptr != NullValue)
-			{
-				int32 Value = LinkedListElements[cur_ptr];
-				if (PredicateFunc(Value))
-				{
-					LinkedListElements[cur_ptr] = NewValue;
-					return true;
-				}
-				cur_ptr = LinkedListElements[cur_ptr + 1];
-			}
-		}
-	}
-	return false;
-}
 
 
 
 
 
-void FSmallListSet::Enumerate(int32 ListIndex, TFunctionRef<void(int32)> ApplyFunc) const
-{
-	int32 block_ptr = ListHeads[ListIndex];
-	if (block_ptr != NullValue)
-	{
-		int32 N = ListBlocks[block_ptr];
-		if (N < BLOCKSIZE)
-		{
-			int32 iEnd = block_ptr + N;
-			for (int32 i = block_ptr + 1; i <= iEnd; ++i)
-			{
-				ApplyFunc(ListBlocks[i]);
-			}
-		}
-		else
-		{
-			// we spilled to linked list, have to iterate through it as well
-			int32 iEnd = block_ptr + BLOCKSIZE;
-			for (int32 i = block_ptr + 1; i <= iEnd; ++i)
-			{
-				ApplyFunc(ListBlocks[i]);
-			}
-			int32 cur_ptr = ListBlocks[block_ptr + BLOCK_LIST_OFFSET];
-			while (cur_ptr != NullValue)
-			{
-				ApplyFunc(LinkedListElements[cur_ptr]);
-				cur_ptr = LinkedListElements[cur_ptr + 1];
-			}
-		}
-	}
-}
+
 
 bool FSmallListSet::EnumerateEarlyOut(int32 ListIndex, TFunctionRef<bool(int32)> ApplyFunc) const
 {
diff --git a/Engine/Source/Runtime/GeometryCore/Public/Util/SmallListSet.h b/Engine/Source/Runtime/GeometryCore/Public/Util/SmallListSet.h
--- a/Engine/Source/Runtime/GeometryCore/Public/Util/SmallListSet.h
+++ b/Engine/Source/Runtime/GeometryCore/Public/Util/SmallListSet.h
@@ -172,7 +172,51 @@ public:
 	 * Search the list at ListIndex for a value where PredicateFunc(value) returns true
 	 * @return the found value, or the InvalidValue argument if not found
 	 */
-	GEOMETRYCORE_API int32 Find(int32 ListIndex, const TFunction<bool(int32)>& PredicateFunc, int32 InvalidValue = -1) const;
+	template<typename IntToBoolFunc>
+	inline int32 Find(int32 ListIndex, const IntToBoolFunc& PredicateFunc, int32 InvalidValue = -1) const
+	{
+		int32 block_ptr = ListHeads[ListIndex];
+		if (block_ptr != NullValue)
+		{
+			int32 N = ListBlocks[block_ptr];
+			if (N < BLOCKSIZE)
+			{
+				int32 iEnd = block_ptr + N;
+				for (int32 i = block_ptr + 1; i <= iEnd; ++i)
+				{
+					int32 Value = ListBlocks[i];
+					if (PredicateFunc(Value))
+					{
+						return Value;
+					}
+				}
+			}
+			else
+			{
+				// we spilled to linked list, have to iterate through it as well
+				int32 iEnd = block_ptr + BLOCKSIZE;
+				for (int32 i = block_ptr + 1; i <= iEnd; ++i)
+				{
+					int32 Value = ListBlocks[i];
+					if (PredicateFunc(Value))
+					{
+						return Value;
+					}
+				}
+				int32 cur_ptr = ListBlocks[block_ptr + BLOCK_LIST_OFFSET];
+				while (cur_ptr != NullValue)
+				{
+					int32 Value = LinkedListElements[cur_ptr];
+					if (PredicateFunc(Value))
+					{
+						return Value;
+					}
+					cur_ptr = LinkedListElements[cur_ptr + 1];
+				}
+			}
+		}
+		return InvalidValue;
+	}
 
 
 
@@ -180,13 +224,92 @@ public:
 	 * Search the list at ListIndex for a value where PredicateFunc(value) returns true, and replace it with NewValue
 	 * @return true if the value was found and replaced
 	 */
-	GEOMETRYCORE_API bool Replace(int32 ListIndex, const TFunction<bool(int32)>& PredicateFunc, int32 NewValue);
+	template<typename IntToBoolFunc>
+	inline bool Replace(int32 ListIndex, const IntToBoolFunc& PredicateFunc, int32 NewValue)
+	{
+		int32 block_ptr = ListHeads[ListIndex];
+		if (block_ptr != NullValue)
+		{
+			int32 N = ListBlocks[block_ptr];
+			if (N < BLOCKSIZE)
+			{
+				int32 iEnd = block_ptr + N;
+				for (int32 i = block_ptr + 1; i <= iEnd; ++i)
+				{
+					int32 Value = ListBlocks[i];
+					if (PredicateFunc(Value))
+					{
+						ListBlocks[i] = NewValue;
+						return true;
+					}
+				}
+			}
+			else
+			{
+				// we spilled to linked list, have to iterate through it as well
+				int32 iEnd = block_ptr + BLOCKSIZE;
+				for (int32 i = block_ptr + 1; i <= iEnd; ++i)
+				{
+					int32 Value = ListBlocks[i];
+					if (PredicateFunc(Value))
+					{
+						ListBlocks[i] = NewValue;
+						return true;
+					}
+				}
+				int32 cur_ptr = ListBlocks[block_ptr + BLOCK_LIST_OFFSET];
+				while (cur_ptr != NullValue)
+				{
+					int32 Value = LinkedListElements[cur_ptr];
+					if (PredicateFunc(Value))
+					{
+						LinkedListElements[cur_ptr] = NewValue;
+						return true;
+					}
+					cur_ptr = LinkedListElements[cur_ptr + 1];
+				}
+			}
+		}
+		return false;
+	}
+
 
 
 	/**
 	 * Call ApplyFunc on each element of the list at ListIndex
 	 */
-	GEOMETRYCORE_API void Enumerate(int32 ListIndex, TFunctionRef<void(int32)> ApplyFunc) const;
+	template<typename IntToVoidFunc>
+	inline void Enumerate(int32 ListIndex, const IntToVoidFunc& ApplyFunc) const
+	{
+		int32 block_ptr = ListHeads[ListIndex];
+		if (block_ptr != NullValue)
+		{
+			int32 N = ListBlocks[block_ptr];
+			if (N < BLOCKSIZE)
+			{
+				int32 iEnd = block_ptr + N;
+				for (int32 i = block_ptr + 1; i <= iEnd; ++i)
+				{
+					ApplyFunc(ListBlocks[i]);
+				}
+			}
+			else
+			{
+				// we spilled to linked list, have to iterate through it as well
+				int32 iEnd = block_ptr + BLOCKSIZE;
+				for (int32 i = block_ptr + 1; i <= iEnd; ++i)
+				{
+					ApplyFunc(ListBlocks[i]);
+				}
+				int32 cur_ptr = ListBlocks[block_ptr + BLOCK_LIST_OFFSET];
+				while (cur_ptr != NullValue)
+				{
+					ApplyFunc(LinkedListElements[cur_ptr]);
+					cur_ptr = LinkedListElements[cur_ptr + 1];
+				}
+			}
+		}
+	}
 
 	/**
 	 * Call ApplyFunc on each element of the list at ListIndex, until ApplyFunc returns false
