fix comment for the Initialize method on the dynamic mesh triangle and vertex attribute sets + make the vertex init not resize to 0 just to fill

diff --git a/Engine/Source/Runtime/GeometryCore/Public/DynamicMesh/DynamicMeshTriangleAttribute.h b/Engine/Source/Runtime/GeometryCore/Public/DynamicMesh/DynamicMeshTriangleAttribute.h
--- a/Engine/Source/Runtime/GeometryCore/Public/DynamicMesh/DynamicMeshTriangleAttribute.h
+++ b/Engine/Source/Runtime/GeometryCore/Public/DynamicMesh/DynamicMeshTriangleAttribute.h
@@ -159,7 +159,7 @@ public:
 	}
 
 
-	/** Initialize the attribute values to the given max triangle ID */
+	/** Initialize the attribute values with InitialValue, and resize to the parent mesh's max triangle ID */
 	void Initialize(AttribValueType InitialValue = (AttribValueType)0)
 	{
 		check(ParentMesh != nullptr);
diff --git a/Engine/Source/Runtime/GeometryCore/Public/DynamicMesh/DynamicVertexAttribute.h b/Engine/Source/Runtime/GeometryCore/Public/DynamicMesh/DynamicVertexAttribute.h
--- a/Engine/Source/Runtime/GeometryCore/Public/DynamicMesh/DynamicVertexAttribute.h
+++ b/Engine/Source/Runtime/GeometryCore/Public/DynamicMesh/DynamicVertexAttribute.h
@@ -148,12 +148,12 @@ public:
 		}
 	}
 
-	/** Initialize the attribute values to the given max triangle ID */
+	/** Initialize the attribute values with InitialValue, and resize to the parent mesh's max vertex ID */
 	void Initialize(AttribValueType InitialValue = (AttribValueType)0)
 	{
 		check(Parent != nullptr);
-		AttribValues.Resize(0);
-		AttribValues.Resize( Parent->MaxVertexID() * AttribDimension, InitialValue );
+		AttribValues.Resize(Parent->MaxVertexID() * AttribDimension);
+		AttribValues.Fill(InitialValue);
 	}
 
 	void SetNewValue(int NewVertexID, const AttribValueType* Data)
