GeometryCore:  Fix bug when updating the pos of an array of vertices in the  DynamicVerticesOctree3.

diff --git a/Engine/Source/Runtime/GeometryCore/Public/DynamicMesh/DynamicVerticesOctree3.h b/Engine/Source/Runtime/GeometryCore/Public/DynamicMesh/DynamicVerticesOctree3.h
--- a/Engine/Source/Runtime/GeometryCore/Public/DynamicMesh/DynamicVerticesOctree3.h
+++ b/Engine/Source/Runtime/GeometryCore/Public/DynamicMesh/DynamicVerticesOctree3.h
@@ -166,7 +166,7 @@ public:
 		{
 			FVector3d Point = VertexSource->GetVertex(VertexID);
 			ModifiedBounds.Contain(Point);
-			FSparseDynamicPointOctree3::ReinsertPoint(VertexID, ModifiedBounds);
+			FSparseDynamicPointOctree3::ReinsertPoint(VertexID, Point);
 		}
 	}
 
