Replaced TEMPLATE_REQUIRES usage.

diff --git a/Engine/Source/Runtime/GeometryCore/Public/TransformTypes.h b/Engine/Source/Runtime/GeometryCore/Public/TransformTypes.h
--- a/Engine/Source/Runtime/GeometryCore/Public/TransformTypes.h
+++ b/Engine/Source/Runtime/GeometryCore/Public/TransformTypes.h
@@ -202,30 +202,62 @@ public:
 		return Scale3D.AllComponentsEqual(Tolerance) || Rotation.IsIdentity(Tolerance);
 	}
 
+
+	// The following templates perform vector-type-conversion variants.
+	// This allows applying a float transform to double vector and vice-versa.
+	// Whether this should be allowed is debatable. However in practice it is extremely rare to convert an
+	// entire float transform to a double transform in order to apply to a double vector, which is the only
+	// case where this conversion is an issue
+
+
 	/**
 	 * @return input point with QST transformation applied, ie QST(P) = Rotate(Scale*P) + Translate
 	 */
-	TVector<RealType> TransformPosition(const TVector<RealType>& P) const
+	template<typename RealType2>
+	TVector<RealType2> TransformPosition(const TVector<RealType2>& P) const
 	{
-		//Transform using QST is following
-		//QST(P) = Q.Rotate(S*P) + T where Q = quaternion, S = scale, T = translation
-		return Rotation * (Scale3D*P) + Translation;
+		if constexpr (std::is_same_v<RealType, RealType2>)
+		{
+			//Transform using QST is following
+			//QST(P) = Q.Rotate(S*P) + T where Q = quaternion, S = scale, T = translation
+			return Rotation * (Scale3D*P) + Translation;
+		}
+		else
+		{
+			return TVector<RealType2>(TransformPosition(TVector<RealType>(P)));
+		}
 	}
 
 	/**
 	 * @return input vector with QS transformation applied, ie QS(V) = Rotate(Scale*V)
 	 */
-	TVector<RealType> TransformVector(const UE::Math::TVector<RealType>& V) const
+	template<typename RealType2>
+	TVector<RealType2> TransformVector(const UE::Math::TVector<RealType2>& V) const
 	{
-		return Rotation * (Scale3D*V);
+		if constexpr (std::is_same_v<RealType, RealType2>)
+		{
+			return Rotation * (Scale3D*V);
+		}
+		else
+		{
+			return TVector<RealType2>(TransformVector(TVector<RealType>(V)));
+		}
 	}
 
 	/**
 	 * @return input vector with Q transformation applied, ie Q(V) = Rotate(V)
 	 */
-	TVector<RealType> TransformVectorNoScale(const UE::Math::TVector<RealType>& V) const
+	template<typename RealType2>
+	TVector<RealType2> TransformVectorNoScale(const UE::Math::TVector<RealType2>& V) const
 	{
-		return Rotation * V;
+		if constexpr (std::is_same_v<RealType, RealType2>)
+		{
+			return Rotation * V;
+		}
+		else
+		{
+			return TVector<RealType2>(TransformVectorNoScale(TVector<RealType>(V)));
+		}
 	}
 
 	/**
@@ -233,39 +265,71 @@ public:
 	 * However 1/Scale requires special handling in case any component is near-zero.
 	 * @return input surface normal with transform applied.
 	 */
-	TVector<RealType> TransformNormal(const UE::Math::TVector<RealType>& Normal) const
+	template<typename RealType2>
+	TVector<RealType2> TransformNormal(const UE::Math::TVector<RealType2>& Normal) const
 	{
-		// transform normal by a safe inverse scale + normalize, and a standard rotation
-		const TVector<RealType>& S = Scale3D;
-		RealType DetSign = FMathd::SignNonZero(S.X * S.Y * S.Z); // we only need to multiply by the sign of the determinant, rather than divide by it, since we normalize later anyway
-		TVector<RealType> SafeInvS(S.Y*S.Z*DetSign, S.X*S.Z*DetSign, S.X*S.Y*DetSign);
-		return TransformVectorNoScale( Normalized(SafeInvS*Normal) );
+		if constexpr (std::is_same_v<RealType, RealType2>)
+		{
+			// transform normal by a safe inverse scale + normalize, and a standard rotation
+			const TVector<RealType>& S = Scale3D;
+			RealType DetSign = FMathd::SignNonZero(S.X * S.Y * S.Z); // we only need to multiply by the sign of the determinant, rather than divide by it, since we normalize later anyway
+			TVector<RealType> SafeInvS(S.Y*S.Z*DetSign, S.X*S.Z*DetSign, S.X*S.Y*DetSign);
+			return TransformVectorNoScale( Normalized(SafeInvS*Normal) );
+		}
+		else
+		{
+			return TVector<RealType2>(TransformNormal(TVector<RealType>(Normal)));
+		}
 	}
 
 
 	/**
 	 * @return input vector with inverse-QST transformation applied, ie QSTinv(P) = InverseScale(InverseRotate(P - Translate))
 	 */
-	TVector<RealType> InverseTransformPosition(const UE::Math::TVector<RealType> &P) const
+	template<typename RealType2>
+	TVector<RealType2> InverseTransformPosition(const UE::Math::TVector<RealType2>& P) const
 	{
-		return GetSafeScaleReciprocal(Scale3D) * Rotation.InverseMultiply(P - Translation);
+		if constexpr (std::is_same_v<RealType, RealType2>)
+		{
+			return GetSafeScaleReciprocal(Scale3D) * Rotation.InverseMultiply(P - Translation);
+		}
+		else
+		{
+			return TVector<RealType2>(InverseTransformPosition(TVector<RealType>(P)));
+		}
 	}
 
 	/**
 	 * @return input vector with inverse-QS transformation applied, ie QSinv(V) = InverseScale(InverseRotate(V))
 	 */
-	TVector<RealType> InverseTransformVector(const UE::Math::TVector<RealType> &V) const
+	template<typename RealType2>
+	TVector<RealType2> InverseTransformVector(const UE::Math::TVector<RealType2>& V) const
 	{
-		return GetSafeScaleReciprocal(Scale3D) * Rotation.InverseMultiply(V);
+		if constexpr (std::is_same_v<RealType, RealType2>)
+		{
+			return GetSafeScaleReciprocal(Scale3D) * Rotation.InverseMultiply(V);
+		}
+		else
+		{
+			return TVector<RealType2>(InverseTransformVector(TVector<RealType>(V)));
+		}
 	}
 
 
 	/**
 	 * @return input vector with inverse-Q transformation applied, ie Qinv(V) = InverseRotate(V)
 	 */
-	TVector<RealType> InverseTransformVectorNoScale(const UE::Math::TVector<RealType> &V) const
+	template<typename RealType2>
+	TVector<RealType2> InverseTransformVectorNoScale(const UE::Math::TVector<RealType2>& V) const
 	{
-		return Rotation.InverseMultiply(V);
+		if constexpr (std::is_same_v<RealType, RealType2>)
+		{
+			return Rotation.InverseMultiply(V);
+		}
+		else
+		{
+			return TVector<RealType2>(InverseTransformVectorNoScale(TVector<RealType>(V)));
+		}
 	}
 
 
@@ -273,9 +337,17 @@ public:
 	 * Surface Normals are special, their inverse transform is InverseRotate( Normalize(Scale * Normal) ) )
 	 * @return input surface normal with inverse transform applied.
 	 */
-	TVector<RealType> InverseTransformNormal(const UE::Math::TVector<RealType>& Normal) const
+	template<typename RealType2>
+	TVector<RealType2> InverseTransformNormal(const UE::Math::TVector<RealType2>& Normal) const
 	{
-		return Normalized( Scale3D * InverseTransformVectorNoScale(Normal) );
+		if constexpr (std::is_same_v<RealType, RealType2>)
+		{
+			return Normalized( Scale3D * InverseTransformVectorNoScale(Normal) );
+		}
+		else
+		{
+			return TVector<RealType2>(InverseTransformNormal(TVector<RealType>(Normal)));
+		}
 	}
 
 
@@ -345,56 +417,6 @@ public:
 
 		return SafeReciprocalScale;
 	}
-
-
-
-
-	// vector-type-conversion variants. This allows applying a float transform to double vector and vice-versa.
-	// Whether this should be allowed is debatable. However in practice it is extremely rare to convert an
-	// entire float transform to a double transform in order to apply to a double vector, which is the only
-	// case where this conversion is an issue
-
-	template<typename RealType2, TEMPLATE_REQUIRES(std::is_same<RealType, RealType2>::value == false)>
-	TVector<RealType2> TransformPosition(const UE::Math::TVector<RealType2>& P) const
-	{
-		return TVector<RealType2>(TransformPosition(TVector<RealType>(P)));
-	}
-	template<typename RealType2, TEMPLATE_REQUIRES(std::is_same<RealType, RealType2>::value == false)>
-	TVector<RealType2> TransformVector(const UE::Math::TVector<RealType2>& V) const
-	{
-		return TVector<RealType2>(TransformVector(TVector<RealType>(V)));
-	}
-	template<typename RealType2, TEMPLATE_REQUIRES(std::is_same<RealType, RealType2>::value == false)>
-	TVector<RealType2> TransformVectorNoScale(const UE::Math::TVector<RealType2>& V) const
-	{
-		return TVector<RealType2>(TransformVectorNoScale(TVector<RealType>(V)));
-	}
-	template<typename RealType2, TEMPLATE_REQUIRES(std::is_same<RealType, RealType2>::value == false)>
-	TVector<RealType2> TransformNormal(const UE::Math::TVector<RealType2>& V) const
-	{
-		return TVector<RealType2>(TransformNormal(TVector<RealType>(V)));
-	}
-	template<typename RealType2, TEMPLATE_REQUIRES(std::is_same<RealType, RealType2>::value == false)>
-	TVector<RealType2> InverseTransformPosition(const UE::Math::TVector<RealType2>& P) const
-	{
-		return TVector<RealType2>(InverseTransformPosition(TVector<RealType>(P)));
-	}
-	template<typename RealType2, TEMPLATE_REQUIRES(std::is_same<RealType, RealType2>::value == false)>
-	TVector<RealType2> InverseTransformVector(const UE::Math::TVector<RealType2>& V) const
-	{
-		return TVector<RealType2>(InverseTransformVector(TVector<RealType>(V)));
-	}
-	template<typename RealType2, TEMPLATE_REQUIRES(std::is_same<RealType, RealType2>::value == false)>
-	TVector<RealType2> InverseTransformVectorNoScale(const UE::Math::TVector<RealType2>& V) const
-	{
-		return TVector<RealType2>(InverseTransformVectorNoScale(TVector<RealType>(V)));
-	}
-	template<typename RealType2, TEMPLATE_REQUIRES(std::is_same<RealType, RealType2>::value == false)>
-	TVector<RealType2> InverseTransformNormal(const UE::Math::TVector<RealType2>& V) const
-	{
-		return TVector<RealType2>(InverseTransformNormal(TVector<RealType>(V)));
-	}
-
 };
 typedef TTransformSRT3<float> FTransformSRT3f;
 typedef TTransformSRT3<double> FTransformSRT3d;
