GridIndexing :
- add way to retrieve the box of a cell
- new methods to get a rectangle of cells overlaping a box (2D or 3D)

diff --git a/Engine/Source/Runtime/GeometryCore/Public/Util/GridIndexing2.h b/Engine/Source/Runtime/GeometryCore/Public/Util/GridIndexing2.h
--- a/Engine/Source/Runtime/GeometryCore/Public/Util/GridIndexing2.h
+++ b/Engine/Source/Runtime/GeometryCore/Public/Util/GridIndexing2.h
@@ -6,6 +6,7 @@
 
 #include "VectorTypes.h"
 #include "IntVectorTypes.h"
+#include "Math/IntRect.h"
 
 namespace UE
 {
@@ -77,6 +78,22 @@ struct TShiftGridIndexer2
 			(int)TMathUtil<RealType>::Floor((Point.Y - Origin.Y) / CellSize));
 	}
 
+	/** Convert real-valued 2D box to a rectangle using integer grid coordinates */
+	inline FIntRect ToGrid(const TBox2<RealType>& Box) const
+	{
+		const FVector2i Min = ToGrid(Box.Min);
+		const FVector2i Max = ToGrid(Box.Max);
+		return FIntRect(Min.X, Min.Y, Max.X, Max.Y);
+	}
+
+	/** Convert real-valued 3D box to a rectangle using integer grid coordinates */
+	inline FIntRect ToGrid(const TBox<RealType>& Box) const
+	{
+		const FVector2i Min = ToGrid(TVector2<RealType>(Box.Min));
+		const FVector2i Max = ToGrid(TVector2<RealType>(Box.Max));
+		return FIntRect(Min.X, Min.Y, Max.X, Max.Y);
+	}
+
 	/** Convert integer grid coordinates to real-valued point */
 	inline TVector2<RealType> FromGrid(const FVector2i& GridPoint) const
 	{
@@ -92,6 +109,14 @@ struct TShiftGridIndexer2
 			((RealType)RealGridPoint.X * CellSize) + Origin.X,
 			((RealType)RealGridPoint.Y * CellSize) + Origin.Y);
 	}
+
+	/** Compute the real-valued 2D box of an integer grid coordinates */
+	inline TBox2<RealType> BoxFromGrid(const FVector2i& GridPoint) const
+	{
+		const TVector2<RealType> Min = FromGrid(GridPoint);
+		const TVector2<RealType> Max = Min + TVector2<RealType>(CellSize);
+		return TBox2<RealType>(Min, Max);
+	}
 };
 typedef TShiftGridIndexer2<float> FShiftGridIndexer2f;
 typedef TShiftGridIndexer2<double> FShiftGridIndexer2d;
