add a max index utility function to match the min index one

diff --git a/Engine/Source/Runtime/GeometryCore/Public/VectorUtil.h b/Engine/Source/Runtime/GeometryCore/Public/VectorUtil.h
--- a/Engine/Source/Runtime/GeometryCore/Public/VectorUtil.h
+++ b/Engine/Source/Runtime/GeometryCore/Public/VectorUtil.h
@@ -194,6 +194,20 @@ namespace VectorUtil
 		}
 	}
 
+	/** @return 0/1/2 index of largest value in Vector3 */
+	template <typename ValueVecType>
+	inline int Max3Index(const ValueVecType& Vector3)
+	{
+		if (Vector3[0] >= Vector3[1])
+		{
+			return Vector3[0] >= Vector3[2] ? 0 : 2;
+		}
+		else
+		{
+			return (Vector3[1] >= Vector3[2]) ? 1 : 2;
+		}
+	}
+
 
 	/**
 	 * Calculates two vectors perpendicular to input Normal, as efficiently as possible.
