fix quickhull initialization to never double-assign points, and remove the duplicate point filtering from the subsequent point insertion steps

diff --git a/Engine/Source/Runtime/GeometryCore/Private/CompGeom/ConvexHull3.cpp b/Engine/Source/Runtime/GeometryCore/Private/CompGeom/ConvexHull3.cpp
--- a/Engine/Source/Runtime/GeometryCore/Private/CompGeom/ConvexHull3.cpp
+++ b/Engine/Source/Runtime/GeometryCore/Private/CompGeom/ConvexHull3.cpp
@@ -269,9 +269,6 @@ struct FHullConnectivity
 	TArray<FIndex3i> TriNeighbors;
 	TArray<FVisiblePoints> VisiblePoints;
 	TSet<int32> TrisWithPoints;
-	
-	TArray<uint16> PointMemberships; // Used for tracking set membership for point indices
-	uint16 MembershipNumber = 0;
 
 	// If positive, this threshold additionally filters which points are considered 'visible' as only points at least this far from the plane
 	double VisibleDistanceThreshold = -FMathd::MaxReal;
@@ -330,6 +327,18 @@ struct FHullConnectivity
 	{
 		VisiblePoints.SetNum(Triangles.Num());
 
+		// Track which points are already assigned so never assign one point to multiple tris
+		TArray<bool> PointAssigned;
+		PointAssigned.SetNumZeroed(NumPoints);
+		for (int32 TriIdx = 0; TriIdx < Triangles.Num(); TriIdx++)
+		{
+			FIndex3i Tri = Triangles[TriIdx];
+			for (int32 SubIdx = 0; SubIdx < 3; ++SubIdx)
+			{
+				PointAssigned[Tri[SubIdx]] = true;
+			}
+		}
+
 		TVector<RealType> TriPts[3];
 		TVector<RealType> Pt;
 		for (int32 TriIdx = 0; TriIdx < Triangles.Num(); TriIdx++)
@@ -339,7 +348,7 @@ struct FHullConnectivity
 
 			for (int32 PtIdx = 0; PtIdx < NumPoints; PtIdx++)
 			{
-				if (!FilterFunc(PtIdx))
+				if (!FilterFunc(PtIdx) || PointAssigned[PtIdx] == true)
 				{
 					continue;
 				}
@@ -353,6 +362,7 @@ struct FHullConnectivity
 						TrisWithPoints.Add(TriIdx);
 					}
 					VisiblePoints[TriIdx].AddPtByValue(PtIdx, Distance);
+					PointAssigned[PtIdx] = true;
 				}
 			}
 		}
@@ -682,33 +692,6 @@ struct FHullConnectivity
 		int32 NumAdd = ToAdd.Num();
 		TVector<RealType> TriPts[3];
 
-		// Remove duplicates from the unclaimed list (unless the list is small)
-		if (NewlyUnclaimed.Num() > 10)
-		{
-			// Use  PointMemberships to track if we've already seen the point
-			if (PointMemberships.Num() != NumPoints || MembershipNumber == MAX_uint16)
-			{
-				MembershipNumber = 1;
-				PointMemberships.Reset();
-				PointMemberships.SetNumZeroed(NumPoints);
-			}
-			else
-			{
-				MembershipNumber++;
-			}
-			for (int32 Idx = 0; Idx < NewlyUnclaimed.Num(); ++Idx)
-			{
-				int32 UnclaimedIdx = NewlyUnclaimed[Idx];
-				if (PointMemberships[UnclaimedIdx] == MembershipNumber)
-				{
-					NewlyUnclaimed.RemoveAtSwap(Idx, EAllowShrinking::No);
-				}
-				else
-				{
-					PointMemberships[UnclaimedIdx] = MembershipNumber;
-				}
-			}
-		}
 
 		for (int32 AddIdx = 0; AddIdx < NumAdd; AddIdx++)
 		{
