Add mitering support to generalized cylinder generator

diff --git a/Engine/Source/Runtime/GeometryCore/Public/Curve/CurveUtil.h b/Engine/Source/Runtime/GeometryCore/Public/Curve/CurveUtil.h
--- a/Engine/Source/Runtime/GeometryCore/Public/Curve/CurveUtil.h
+++ b/Engine/Source/Runtime/GeometryCore/Public/Curve/CurveUtil.h
@@ -61,6 +61,23 @@ using namespace UE::Math;
 			Normalize(OutToPrev);
 		}
 	}
+
+	/**
+	 * Get (by reference) vectors pointing toward the given vertex index, from its surrounding vertices
+	 * If bLoop is false and the Prev or Next vertex would be out of bounds, a zero vector is used instead
+	 */
+	template<typename RealType, typename VectorType>
+	inline void GetVectorsToPrevNext(const TArrayView<const VectorType>& Vertices, int32 VertexIndex, VectorType& OutToPrev, VectorType& OutToNext, bool bNormalize, bool bLoop)
+	{
+		if (bLoop)
+		{
+			GetVectorsToPrevNext<RealType, VectorType, true>(Vertices, VertexIndex, OutToPrev, OutToNext, bNormalize);
+		}
+		else
+		{
+			GetVectorsToPrevNext<RealType, VectorType, false>(Vertices, VertexIndex, OutToPrev, OutToNext, bNormalize);
+		}
+	}
 	
 	/**
 	 * @return the tangent direction (normalized) of the path at Idx, using the surrounding vertices (as found by GetPrevNext())
diff --git a/Engine/Source/Runtime/GeometryCore/Public/Generators/SweepGenerator.h b/Engine/Source/Runtime/GeometryCore/Public/Generators/SweepGenerator.h
--- a/Engine/Source/Runtime/GeometryCore/Public/Generators/SweepGenerator.h
+++ b/Engine/Source/Runtime/GeometryCore/Public/Generators/SweepGenerator.h
@@ -717,7 +717,6 @@ public:
  * 
  * TODO: 
  *  - a custom variant for toruses specifically (would be faster)
- *  - Mitering cross sections support?
  */
 class /*GEOMETRYCORE_API*/ FGeneralizedCylinderGenerator : public FSweepGeneratorBase
 {
@@ -739,6 +738,20 @@ public:
 	double StartScale = 1.0;
 	double EndScale = 1.0;
 
+	// Maximum mitering scale to apply at turns; only applied if greater than 1
+	// Note: If PathFrames are specified, then bAlignFramesToSampledTangents must be true for mitering to be performed.
+	double MaxMiterScale = 1.0;
+
+	// Whether to align the frames to the path tangents. Only relevant if PathFrames is specified.
+	bool bAlignFramesToSampledTangents = false;
+
+	// Configure settings to support 'mitering' -- i.e., scaling the cross sections as needed to maintain consistent cross section size through sharp corners, up to the specified scale limit
+	void EnableMitering(double InMaxMiterScale = 10)
+	{
+		MaxMiterScale = InMaxMiterScale;
+		bAlignFramesToSampledTangents = true;
+	}
+
 	// When true, the generator attempts to scale UV's in a way that preserves scaling across different mesh
 	// results, aiming for 1.0 in UV space to be equal to UnitUVInWorldCoordinates in world space. This in
 	// practice means adjusting the U scale relative to the CrossSection curve length and V scale relative
@@ -819,19 +832,46 @@ public:
 		{
 			FVector3d C = Path[PathIdx];
 			FVector3d X, Y;
-			if (bHaveExplicitFrames == false)
+			FMatrix2d MiterPosScale = FMatrix2d::Identity();
+			if (bHaveExplicitFrames)
 			{
-				FVector3d Tangent = UE::Geometry::CurveUtil::Tangent<double, FVector3d>(Path, PathIdx, bLoop);
-				CrossSectionFrame.AlignAxis(2, Tangent);
-				X = CrossSectionFrame.X();
-				Y = CrossSectionFrame.Y();
+				CrossSectionFrame = PathFrames[PathIdx];
+				C = PathFrames[PathIdx].Origin;
 			}
-			else
+			if (!bHaveExplicitFrames || bAlignFramesToSampledTangents)
 			{
-				C = PathFrames[PathIdx].Origin;
-				X = PathFrames[PathIdx].X();
-				Y = PathFrames[PathIdx].Y();
+				FVector3d ToPrev, ToNext;
+				UE::Geometry::CurveUtil::GetVectorsToPrevNext<double, FVector3d>(Path, PathIdx, ToPrev, ToNext, true, bLoop);
+
+				FVector3d Tangent = Normalized(-ToPrev + ToNext);
+				CrossSectionFrame.AlignAxis(2, Tangent);
+				if (MaxMiterScale > 1)
+				{
+					// only miter if neither neighbor vector is exactly zero
+					if (!ToPrev.IsZero() && !ToNext.IsZero())
+					{
+						FVector3d TurnDir = ToPrev + ToNext;
+						FVector2d ScaleAlong = (FVector2d)CrossSectionFrame.ToFrameVector(TurnDir);
+						double ScaleAmount = 1.0;
+						if (ScaleAlong.Normalize())
+						{
+							double CosTwoTheta = ToPrev.Dot(ToNext);
+							double SqrScale = 2.0 / (1.0 - CosTwoTheta);
+							if (!FMath::IsFinite(SqrScale) || SqrScale > MaxMiterScale * MaxMiterScale)
+							{
+								ScaleAmount = MaxMiterScale;
+							}
+							else
+							{
+								ScaleAmount = FMath::Sqrt(SqrScale);
+							}
+							MiterPosScale = FMatrix2d::AxisScale(ScaleAlong, ScaleAmount, false /* axis was already normalized */);
+						}
+					}
+				}
 			}
+			X = CrossSectionFrame.X();
+			Y = CrossSectionFrame.Y();
 
 			double T = FMathd::Clamp((AccumArcLength / TotalPathArcLength), 0.0, 1.0);
 			double UniformScale = (bApplyScaling) ? FMathd::Lerp(StartScale, EndScale, T) : 1.0;
@@ -839,8 +879,11 @@ public:
 
 			for (int SubIdx = 0; SubIdx < XNum; SubIdx++)
 			{
-				FVector2d XP = UniformScale * PathScaling * CrossSection[SubIdx];
+				FVector2d XP = MiterPosScale * (UniformScale * PathScaling * CrossSection[SubIdx]);
 				FVector2d XN = XNormals[SubIdx];
+				// Note: Arguably, one could apply an inverse mitering scale to the normals, and then re-normalize, here ...
+				// However, it's better to think about these normals an average from the two unscaled cross sections that are meeting at the cross section plane,
+				// and in that model they won't be affected by the miter scale
 				Vertices[SubIdx + PathIdx * XNum] = C + X * XP.X + Y * XP.Y;
 				Normals[SubIdx + PathIdx * XNum] = (FVector3f)(X * XN.X + Y * XN.Y);
 			}
diff --git a/Engine/Source/Runtime/GeometryCore/Public/MatrixTypes.h b/Engine/Source/Runtime/GeometryCore/Public/MatrixTypes.h
--- a/Engine/Source/Runtime/GeometryCore/Public/MatrixTypes.h
+++ b/Engine/Source/Runtime/GeometryCore/Public/MatrixTypes.h
@@ -448,6 +448,24 @@ struct TMatrix2
 		return RotationRad(AngleDeg * TMathUtil<RealType>::DegToRad);
 	}
 
+	// Create a matrix that scales along the specified axis
+	// @param Axis	Axis to scale along
+	// @param Scale	Amount to scale
+	// @param bNormalizeAxis Whether to normalize Axis. If false, assumes Axis was already normalized.
+	static TMatrix2<RealType> AxisScale(TVector2<RealType> Axis, RealType Scale, bool bNormalizeAxis = true)
+	{
+		if (bNormalizeAxis)
+		{
+			Axis.Normalize();
+		}
+		RealType X2 = Axis.X * Axis.X;
+		RealType Y2 = Axis.Y * Axis.Y;
+		RealType XY = Axis.X * Axis.Y;
+		return TMatrix2<RealType>(
+			Scale * X2 + Y2, Scale * XY - XY,
+			Scale * XY - XY, Scale * Y2 + X2);
+	}
+
 	/**
 	 * Assumes we have a rotation matrix (uniform scale ok)
 	 */
