PR #11705: Fixed redundant UV mesh spatial computation in modeling mode bake tools

diff --git a/Engine/Plugins/Runtime/GeometryProcessing/Source/DynamicMesh/Private/Sampling/MeshMapBaker.cpp b/Engine/Plugins/Runtime/GeometryProcessing/Source/DynamicMesh/Private/Sampling/MeshMapBaker.cpp
--- a/Engine/Plugins/Runtime/GeometryProcessing/Source/DynamicMesh/Private/Sampling/MeshMapBaker.cpp
+++ b/Engine/Plugins/Runtime/GeometryProcessing/Source/DynamicMesh/Private/Sampling/MeshMapBaker.cpp
@@ -293,6 +293,9 @@ void FMeshMapBaker::Bake()
 		}
 	}
 
+	// Generate UV space mesh Spatial
+	TMeshAABBTree3 FlatSpatialCache(&FlatMesh, true);
+
 	ECorrespondenceStrategy UseStrategy = this->CorrespondenceStrategy;
 	bool bIsIdentity = true;
 	int NumDetailMeshes = 0;
@@ -444,7 +447,7 @@ void FMeshMapBaker::Bake()
 	};
 
 	FMeshMapBakerQueue OutputQueue(NumTiles);
-	ParallelFor(NumTiles, [this, &Tiles, &BorderTexelsPerTile, &GutterTexelsPerTile, &OutputQueue, &WriteToOutputBuffer, &WriteToOutputBufferQueued, &ComputeCorrespondenceSample](int32 TileIdx)
+	ParallelFor(NumTiles, [this, &Tiles, &BorderTexelsPerTile, &GutterTexelsPerTile, &OutputQueue, &WriteToOutputBuffer, &WriteToOutputBufferQueued, &ComputeCorrespondenceSample, &FlatSpatialCache](int32 TileIdx)
 	{
 		TRACE_CPUPROFILER_EVENT_SCOPE(FMeshMapBaker::Bake_EvalTile);
 
@@ -461,7 +464,7 @@ void FMeshMapBaker::Bake()
 		OccupancyMap.GutterSize = GutterSize;
 		OccupancyMap.Initialize(Dimensions, PaddedTile, SamplesPerPixel);
 		const auto GetTriangleIDFunc = [this](int32 TriangleID) { return FlatMesh.GetTriangleGroup(TriangleID); };
-		OccupancyMap.ClassifySamplesFromUVSpaceMesh(FlatMesh, GetTriangleIDFunc, TargetMeshUVCharts);
+		OccupancyMap.ClassifySamplesFromUVSpaceMesh(FlatMesh, FlatSpatialCache, GetTriangleIDFunc, TargetMeshUVCharts);
 		ComputeGutterTexelsUsingFilterKernelCoverage(OccupancyMap, Tile, Dimensions, FilterKernelSize, IsInFilterRegionEval);
 		BorderTexelsPerTile[TileIdx] = OccupancyMap.BorderTexels;
 		GutterTexelsPerTile[TileIdx] = OccupancyMap.GutterTexels;
diff --git a/Engine/Source/Runtime/GeometryCore/Public/Image/ImageOccupancyMap.h b/Engine/Source/Runtime/GeometryCore/Public/Image/ImageOccupancyMap.h
--- a/Engine/Source/Runtime/GeometryCore/Public/Image/ImageOccupancyMap.h
+++ b/Engine/Source/Runtime/GeometryCore/Public/Image/ImageOccupancyMap.h
@@ -133,14 +133,25 @@ public:
 
 	template<typename MeshType, typename GetTriangleIDFuncType>
 	bool ClassifySamplesFromUVSpaceMesh(
-		const MeshType& UVSpaceMesh, 
+		const MeshType& UVSpaceMesh,
 		GetTriangleIDFuncType GetTriangleIDFunc = [](int32 TriangleID) { return TriangleID; },
 		const TArray<int32>* UVSpaceMeshTriCharts = nullptr)
 	{
-		TRACE_CPUPROFILER_EVENT_SCOPE(FImageOccupancyMap::ClassifySamples);
-
+		TRACE_CPUPROFILER_EVENT_SCOPE(FImageOccupancyMap::ClassifySamples_NoSpatial);
+		
 		// make flat mesh
-		TMeshAABBTree3<MeshType> FlatSpatial(&UVSpaceMesh, true);
+		TMeshAABBTree3<MeshType> UVSpaceMeshSpatial(&UVSpaceMesh, true);
+		return ClassifySamplesFromUVSpaceMesh(UVSpaceMesh, UVSpaceMeshSpatial, GetTriangleIDFunc, UVSpaceMeshTriCharts);
+	}
+
+	template<typename MeshType, typename GetTriangleIDFuncType>
+	bool ClassifySamplesFromUVSpaceMesh(
+		const MeshType& UVSpaceMesh,
+		const TMeshAABBTree3<MeshType>& UVSpaceMeshSpatial,
+		GetTriangleIDFuncType GetTriangleIDFunc = [](int32 TriangleID) { return TriangleID; },
+		const TArray<int32>* UVSpaceMeshTriCharts = nullptr)
+	{
+		TRACE_CPUPROFILER_EVENT_SCOPE(FImageOccupancyMap::ClassifySamples);
 
 		const int64 LinearImageSize = Tile.Num();
 		TexelInteriorSamples.Init(0, LinearImageSize);
@@ -161,7 +172,7 @@ public:
 
 		// Classify texel samples
 		ParallelFor(Tile.GetHeight(),
-					[this, &UVSpaceMesh, &GetTriangleIDFunc, &FlatSpatial, TexelDiag, &QueryOptions, &TexelSize, UVSpaceMeshTriCharts, NumUVSpaceMeshTriCharts]
+					[this, &UVSpaceMesh, &GetTriangleIDFunc, &UVSpaceMeshSpatial, TexelDiag, &QueryOptions, &TexelSize, UVSpaceMeshTriCharts, NumUVSpaceMeshTriCharts]
 		(int32 ImgY)
 		{
 			for (int32 ImgX = 0; ImgX < Tile.GetWidth(); ++ImgX)
@@ -179,7 +190,7 @@ public:
 					const FVector3d UVPoint3d(UVPoint.X, UVPoint.Y, 0);
 
 					double NearDistSqr;
-					const int32 NearestTriID = FlatSpatial.FindNearestTriangle(UVPoint3d, NearDistSqr, QueryOptions);
+					const int32 NearestTriID = UVSpaceMeshSpatial.FindNearestTriangle(UVPoint3d, NearDistSqr, QueryOptions);
 					if (NearestTriID >= 0)
 					{
 						FVector3d A, B, C;
