fix issue with dynamic mesh edge merges where we could create a non-manifold mesh by trying to merge edge AB to edge CD when edge AD also existed in the mesh; resulting in the tris originally connected to edges AB, CD and AD all having edge AB after the merge.
With this change, we now refuse to merge the edges in this case.

diff --git a/Engine/Source/Runtime/GeometryCore/Private/DynamicMesh/DynamicMesh3_Edits.cpp b/Engine/Source/Runtime/GeometryCore/Private/DynamicMesh/DynamicMesh3_Edits.cpp
--- a/Engine/Source/Runtime/GeometryCore/Private/DynamicMesh/DynamicMesh3_Edits.cpp
+++ b/Engine/Source/Runtime/GeometryCore/Private/DynamicMesh/DynamicMesh3_Edits.cpp
@@ -1729,6 +1729,15 @@ EMeshResult FDynamicMesh3::MergeEdges(int eKeep, int eDiscard, FMergeEdgesInfo&
 	{
 		return EMeshResult::Failed_InvalidNeighbourhood;
 	}
+	// the un-matched edge vertices, a/d and b/c, should also not be directly connected, and should not be the same vertex
+	if (a == d || FindEdge(a, d) != InvalidID)
+	{
+		return EMeshResult::Failed_InvalidNeighbourhood;
+	}
+	if (b == c || FindEdge(b, c) != InvalidID)
+	{
+		return EMeshResult::Failed_InvalidNeighbourhood;
+	}
 
 	// if vertices at either end already share a common neighbour vertex, and we
 	// do the merge, that would create duplicate edges. This is something like the
