better fix for dynamic mesh merge edges, to not prevent edge merges for edge pairs that share a vertex

diff --git a/Engine/Source/Runtime/GeometryCore/Private/DynamicMesh/DynamicMesh3_Edits.cpp b/Engine/Source/Runtime/GeometryCore/Private/DynamicMesh/DynamicMesh3_Edits.cpp
--- a/Engine/Source/Runtime/GeometryCore/Private/DynamicMesh/DynamicMesh3_Edits.cpp
+++ b/Engine/Source/Runtime/GeometryCore/Private/DynamicMesh/DynamicMesh3_Edits.cpp
@@ -1729,12 +1729,13 @@ EMeshResult FDynamicMesh3::MergeEdges(int eKeep, int eDiscard, FMergeEdgesInfo&
 	{
 		return EMeshResult::Failed_InvalidNeighbourhood;
 	}
-	// the un-matched edge vertices, a/d and b/c, should also not be directly connected, and should not be the same vertex
-	if (a == d || FindEdge(a, d) != InvalidID)
+	// the un-matched edge vertices, a/d and b/c, should also not be directly connected
+	// (unless the edges share a vertex, in which case they are always connected by one of the two merge edges)
+	if (a != c && b != d && FindEdge(a, d) != InvalidID)
 	{
 		return EMeshResult::Failed_InvalidNeighbourhood;
 	}
-	if (b == c || FindEdge(b, c) != InvalidID)
+	if (a != c && b != d && FindEdge(b, c) != InvalidID)
 	{
 		return EMeshResult::Failed_InvalidNeighbourhood;
 	}
