Add method to get the Unit value from a Distance on a Segment

diff --git a/Engine/Source/Runtime/GeometryCore/Public/SegmentTypes.h b/Engine/Source/Runtime/GeometryCore/Public/SegmentTypes.h
--- a/Engine/Source/Runtime/GeometryCore/Public/SegmentTypes.h
+++ b/Engine/Source/Runtime/GeometryCore/Public/SegmentTypes.h
@@ -117,6 +117,15 @@ public:
 		return Center + ((T)2 * UnitParameter - (T)1) * Extent * Direction;
 	}
 
+	/**
+	 * @param DistanceParameter distance on the segment from the segment Origin in range [-Extent, Extent]
+	 * @return unit value in range [0,1]
+	 */
+	inline T ConvertToUnitRange(T DistanceParameter) const
+	{
+		return T(0.5) * (T(1) + DistanceParameter / Extent); 
+	}
+
 	/**
 	 * @return minimum squared distance from Point to segment
 	 */
@@ -454,6 +463,15 @@ public:
 		return Center + ((T)2 * UnitParameter - (T)1) * Extent * Direction;
 	}
 
+	/**
+	 * @param DistanceParameter distance on the segment from the segment Origin in range [-Extent, Extent]
+	 * @return unit value in range [0,1]
+	 */
+	inline T ConvertToUnitRange(T DistanceParameter) const
+	{
+		return T(0.5) * (T(1) + DistanceParameter / Extent); 
+	}
+
 	/**
 	 * @return minimum squared distance from Point to segment
 	 */
