make geometrycore grid interpolant support templated RealType rather than always running in double precision
 + add support for clamped value queries (+ fix old comment about clamping values)
 + fix bounds of interpolant to be the actual useful bounds -- i.e., don't include a cell's width of space that will just return the Outside value when queried
 + add a SafeLargeValue to geometrycore's math constants (to give a good default Outside value for float interpolant grid)

diff --git a/Engine/Source/Runtime/GeometryCore/Public/Implicit/GridInterpolant.h b/Engine/Source/Runtime/GeometryCore/Public/Implicit/GridInterpolant.h
--- a/Engine/Source/Runtime/GeometryCore/Public/Implicit/GridInterpolant.h
+++ b/Engine/Source/Runtime/GeometryCore/Public/Implicit/GridInterpolant.h
@@ -22,34 +22,34 @@ namespace Geometry
  *
  * GridType must have a GetValue() that returns a value to interpolate at a given FVector3i location -- (w/ locations ranging from [0,0,0] to Dimensions (exclusive))
  */
-template <class GridType>
+template <class GridType, typename RealType = double>
 class TTriLinearGridInterpolant /*: public BoundedImplicitFunction3d (TODO: consider add ImplicitFunction3d interface concept once more implicit functions are available*/
 {
 public:
 	GridType* Grid;
-	FVector3d GridOrigin;
-	double CellSize;
+	TVector<RealType> GridOrigin;
+	RealType CellSize;
 	FVector3i Dimensions;
 
 	// value to return if query point is outside Grid (in an SDF
 	// outside is usually positive). Need to do math with this value,
-	// and cast this value to/from float; use FMathf::MaxReal to avoid overflow
-	double Outside = FMathf::MaxReal;
+	// and cast this value to/from float; use TMathUtil<RealType>::SafeLargeValue to avoid overflow
+	RealType Outside = TMathUtil<RealType>::SafeLargeValue;
 
-	TTriLinearGridInterpolant(GridType* Grid, FVector3d GridOrigin, double CellSize, FVector3i Dimensions) : Grid(Grid), GridOrigin(GridOrigin), CellSize(CellSize), Dimensions(Dimensions)
+	TTriLinearGridInterpolant(GridType* Grid, TVector<RealType> GridOrigin, RealType CellSize, FVector3i Dimensions) : Grid(Grid), GridOrigin(GridOrigin), CellSize(CellSize), Dimensions(Dimensions)
 	{
 	}
 
-	FAxisAlignedBox3d Bounds() const
+	TAxisAlignedBox3<RealType> Bounds() const
 	{
-		return FAxisAlignedBox3d(
+		return TAxisAlignedBox3<RealType>(
 			{ GridOrigin.X, GridOrigin.Y, GridOrigin.Z },
-			{ GridOrigin.X + CellSize * Dimensions.X,
-			  GridOrigin.Y + CellSize * Dimensions.Y,
-			  GridOrigin.Z + CellSize * Dimensions.Z });
+			{ GridOrigin.X + CellSize * (Dimensions.X - 1),
+			  GridOrigin.Y + CellSize * (Dimensions.Y - 1),
+			  GridOrigin.Z + CellSize * (Dimensions.Z - 1)});
 	}
 
-	FVector3i Cell(const FVector3d& Pt) const
+	FVector3i Cell(const TVector<RealType>& Pt) const
 	{
 		// compute integer coordinates
 		FVector3i CellCoords;
@@ -60,43 +60,51 @@ public:
 		return CellCoords;
 	}
 
-	double Value(const FVector3d& Pt) const
+	template<bool bClamped = false>
+	RealType Value(const TVector<RealType>& Pt) const
 	{
-		FVector3d gridPt(
+		TVector<RealType> gridPt(
 			((Pt.X - GridOrigin.X) / CellSize),
 			((Pt.Y - GridOrigin.Y) / CellSize),
 			((Pt.Z - GridOrigin.Z) / CellSize));
 
+		if constexpr (bClamped)
+		{
+			gridPt.X = FMath::Clamp(gridPt.X, 0, Dimensions.X - (1 + FMathd::Epsilon));
+			gridPt.Y = FMath::Clamp(gridPt.Y, 0, Dimensions.Y - (1 + FMathd::Epsilon));
+			gridPt.Z = FMath::Clamp(gridPt.Z, 0, Dimensions.Z - (1 + FMathd::Epsilon));
+		}
+
 		// compute integer coordinates
 		int X0 = (int)gridPt.X;
 		int Y0 = (int)gridPt.Y, Y1 = Y0 + 1;
 		int Z0 = (int)gridPt.Z, Z1 = Z0 + 1;
 
-		// clamp to Grid
-		if (X0 < 0 || (X0 + 1) >= Dimensions.X ||
-			Y0 < 0 || Y1 >= Dimensions.Y ||
-			Z0 < 0 || Z1 >= Dimensions.Z)
+		// return Outside if not in Grid
+		if constexpr (!bClamped)
 		{
-			return Outside;
+			if (X0 < 0 || (X0 + 1) >= Dimensions.X ||
+				Y0 < 0 || Y1 >= Dimensions.Y ||
+				Z0 < 0 || Z1 >= Dimensions.Z)
+			{
+				return Outside;
+			}
 		}
 
 		// convert double coords to [0,1] range
-		double fAx = gridPt.X - (double)X0;
-		double fAy = gridPt.Y - (double)Y0;
-		double fAz = gridPt.Z - (double)Z0;
-		double OneMinusfAx = 1.0 - fAx;
+		RealType fAx = gridPt.X - (RealType)X0;
+		RealType fAy = gridPt.Y - (RealType)Y0;
+		RealType fAz = gridPt.Z - (RealType)Z0;
+		RealType OneMinusfAx = 1.0 - fAx;
 
 		// compute trilinear interpolant. The code below tries to do this with the fewest 
 		// number of variables, in hopes that optimizer will be clever about re-using registers, etc.
 		// Commented code at bottom is fully-expanded version.
-		// [TODO] it is possible to implement lerps here as A+(B-A)*t, saving a multiply and a variable.
-		//   This is numerically worse, but since the Grid values are floats and
-		//   we are computing in doubles, does it matter?
-		double xa, xb;
+		RealType xa, xb;
 
 		get_value_pair(X0, Y0, Z0, xa, xb);
-		double yz = (1 - fAy) * (1 - fAz);
-		double sum = (OneMinusfAx * xa + fAx * xb) * yz;
+		RealType yz = (1 - fAy) * (1 - fAz);
+		RealType sum = (OneMinusfAx * xa + fAx * xb) * yz;
 
 		get_value_pair(X0, Y0, Z1, xa, xb);
 		yz = (1 - fAy) * (fAz);
@@ -126,17 +134,17 @@ public:
 
 
 protected:
-	void get_value_pair(int I, int J, int K, double& A, double& B) const
+	void get_value_pair(int I, int J, int K, RealType& A, RealType& B) const
 	{
-		A = (double)Grid->GetValue(FVector3i(I, J, K));
-		B = (double)Grid->GetValue(FVector3i(I + 1, J, K));
+		A = (RealType)Grid->GetValue(FVector3i(I, J, K));
+		B = (RealType)Grid->GetValue(FVector3i(I + 1, J, K));
 	}
 
 
 public:
-	FVector3d Gradient(const FVector3d& Pt) const
+	TVector<RealType> Gradient(const TVector<RealType>& Pt) const
 	{
-		FVector3d gridPt = FVector3d(
+		TVector<RealType> gridPt = TVector<RealType>(
 			((Pt.X - GridOrigin.X) / CellSize),
 			((Pt.Y - GridOrigin.Y) / CellSize),
 			((Pt.Z - GridOrigin.Z) / CellSize));
@@ -146,7 +154,7 @@ public:
 			gridPt.Y < 0 || gridPt.Y >= Dimensions.Y - 1 ||
 			gridPt.Z < 0 || gridPt.Z >= Dimensions.Z - 1)
 		{
-			return FVector3d::Zero();
+			return TVector<RealType>::Zero();
 		}
 
 		// compute integer coordinates
@@ -154,22 +162,22 @@ public:
 		int Y0 = (int)gridPt.Y, Y1 = Y0 + 1;
 		int Z0 = (int)gridPt.Z, Z1 = Z0 + 1;
 
-		// convert double coords to [0,1] range
-		double fAx = gridPt.X - (double)X0;
-		double fAy = gridPt.Y - (double)Y0;
-		double fAz = gridPt.Z - (double)Z0;
+		// convert RealType coords to [0,1] range
+		RealType fAx = gridPt.X - (RealType)X0;
+		RealType fAy = gridPt.Y - (RealType)Y0;
+		RealType fAz = gridPt.Z - (RealType)Z0;
 
-		double fV000, fV100;
+		RealType fV000, fV100;
 		get_value_pair(X0, Y0, Z0, fV000, fV100);
-		double fV010, fV110;
+		RealType fV010, fV110;
 		get_value_pair(X0, Y1, Z0, fV010, fV110);
-		double fV001, fV101;
+		RealType fV001, fV101;
 		get_value_pair(X0, Y0, Z1, fV001, fV101);
-		double fV011, fV111;
+		RealType fV011, fV111;
 		get_value_pair(X0, Y1, Z1, fV011, fV111);
 
 		// [TODO] can re-order this to vastly reduce number of ops!
-		double gradX =
+		RealType gradX =
 			-fV000 * (1 - fAy) * (1 - fAz) +
 			-fV001 * (1 - fAy) * (fAz)+
 			-fV010 * (fAy) * (1 - fAz) +
@@ -179,7 +187,7 @@ public:
 			fV110 * (fAy) * (1 - fAz) +
 			fV111 * (fAy) * (fAz);
 
-		double gradY =
+		RealType gradY =
 			-fV000 * (1 - fAx) * (1 - fAz) +
 			-fV001 * (1 - fAx) * (fAz)+
 			fV010 * (1 - fAx) * (1 - fAz) +
@@ -189,7 +197,7 @@ public:
 			fV110 * (fAx) * (1 - fAz) +
 			fV111 * (fAx) * (fAz);
 
-		double gradZ =
+		RealType gradZ =
 			-fV000 * (1 - fAx) * (1 - fAy) +
 			fV001 * (1 - fAx) * (1 - fAy) +
 			-fV010 * (1 - fAx) * (fAy)+
@@ -199,7 +207,7 @@ public:
 			-fV110 * (fAx) * (fAy)+
 			fV111 * (fAx) * (fAy);
 
-		return FVector3d(gradX, gradY, gradZ);
+		return TVector<RealType>(gradX, gradY, gradZ);
 	}
 };
 
diff --git a/Engine/Source/Runtime/GeometryCore/Public/Intersection/ContainmentQueries3.h b/Engine/Source/Runtime/GeometryCore/Public/Intersection/ContainmentQueries3.h
--- a/Engine/Source/Runtime/GeometryCore/Public/Intersection/ContainmentQueries3.h
+++ b/Engine/Source/Runtime/GeometryCore/Public/Intersection/ContainmentQueries3.h
@@ -24,7 +24,7 @@ namespace UE
 
 	namespace Geometry
 	{
-		template <class GridType> class TTriLinearGridInterpolant;
+		template <class GridType, typename RealType> class TTriLinearGridInterpolant;
 		template <typename T> struct TCapsule3;
 		template <typename T> struct THalfspace3;
 		template <typename T> struct TSphere3;
diff --git a/Engine/Source/Runtime/GeometryCore/Public/MathUtil.h b/Engine/Source/Runtime/GeometryCore/Public/MathUtil.h
--- a/Engine/Source/Runtime/GeometryCore/Public/MathUtil.h
+++ b/Engine/Source/Runtime/GeometryCore/Public/MathUtil.h
@@ -4,6 +4,7 @@
 
 #include "GeometryBase.h"
 #include "HAL/Platform.h"
+#include "EngineDefines.h"
 
 #include <cmath>
 #include <cfloat>
@@ -25,6 +26,9 @@ struct TMathUtilConstants<float>
 	/** largest possible number for type */
 	static constexpr float MaxReal = FLT_MAX;
 
+	/** a very large value, but not too close to the max possible for the type */
+	static constexpr float SafeLargeValue = UE_LARGE_WORLD_MAX;
+
 	/** 3.14159... */
 	static constexpr float Pi = 3.1415926535897932384626433832795f;
 	static constexpr float FourPi = 4.0f * Pi;
@@ -63,6 +67,9 @@ struct TMathUtilConstants<double>
 	/** largest possible number for type */
 	static constexpr double MaxReal = DBL_MAX;
 
+	/** a very large value, but not too close to the max possible for the type */
+	static constexpr double SafeLargeValue = (double)FLT_MAX;
+
 	/** 3.14159... */
 	static constexpr double Pi = 3.1415926535897932384626433832795;
 	static constexpr double FourPi = 4.0 * Pi;
@@ -97,7 +104,8 @@ struct TMathUtilConstants<int32>
 {
 	static constexpr int32 Epsilon = 0;
 	static constexpr int32 ZeroTolerance = 0;
-	static constexpr int32 MaxReal = ((int32)0x7fffffff);
+	static constexpr int32 MaxReal = MAX_int32;
+	static constexpr int32 SafeLargeValue = MAX_int32/2;
 	static constexpr int32 Pi = 3;
 	static constexpr int32 FourPi = 4 * Pi;
 	static constexpr int32 TwoPi = 2 * Pi;
@@ -119,7 +127,8 @@ struct TMathUtilConstants<int64>
 {
 	static constexpr int64 Epsilon = 0;
 	static constexpr int64 ZeroTolerance = 0;
-	static constexpr int64 MaxReal = ((int64)0x7fffffffffffffff);
+	static constexpr int64 MaxReal = MAX_int64;
+	static constexpr int64 SafeLargeValue = MAX_int64/2;
 	static constexpr int64 Pi = 3;
 	static constexpr int64 FourPi = 4 * Pi;
 	static constexpr int64 TwoPi = 2 * Pi;
