Compact small list sets during dynamic mesh CompactInPlace

diff --git a/Engine/Source/Runtime/GeometryCore/Private/DynamicMesh/DynamicMesh3_Edits.cpp b/Engine/Source/Runtime/GeometryCore/Private/DynamicMesh/DynamicMesh3_Edits.cpp
--- a/Engine/Source/Runtime/GeometryCore/Private/DynamicMesh/DynamicMesh3_Edits.cpp
+++ b/Engine/Source/Runtime/GeometryCore/Private/DynamicMesh/DynamicMesh3_Edits.cpp
@@ -431,7 +431,7 @@ void FDynamicMesh3::CompactInPlace(FCompactMaps* CompactInfo)
 		VertexUVs->Resize(VertexCount() * 2);
 	}
 
-	// [TODO] VertexEdgeLists!!!
+	VertexEdgeLists.Compact(VertexCount());
 
 	/** shift triangles **/
 
diff --git a/Engine/Source/Runtime/GeometryCore/Private/Util/SmallListSet.cpp b/Engine/Source/Runtime/GeometryCore/Private/Util/SmallListSet.cpp
--- a/Engine/Source/Runtime/GeometryCore/Private/Util/SmallListSet.cpp
+++ b/Engine/Source/Runtime/GeometryCore/Private/Util/SmallListSet.cpp
@@ -42,7 +42,62 @@ void FSmallListSet::AllocateAt(int32 ListIndex)
 	}
 }
 
+void FSmallListSet::Compact(int32 MaxListIndex)
+{
+	checkSlow(MaxListIndex >= 0);
+	int32 CurSize = (int32)ListHeads.GetLength();
+	if (MaxListIndex < CurSize)
+	{
+		// We just resize w/out book-keeping what we cleared, since we rebuild the blocks/etc below
+		ListHeads.Resize(MaxListIndex);
+	}
 
+	AllocatedCount = 0;
+	TDynamicVector<int32> NewBlocks{};
+	TDynamicVector<int32> NewLinkedListElements{};
+	for (int32 Idx = 0, Num = (int32)ListHeads.GetLength(), CurBlockIdx = 0; Idx < Num; ++Idx, CurBlockIdx += BLOCK_LIST_OFFSET + 1)
+	{
+		int32 OrigHead = ListHeads[Idx];
+		if (OrigHead == NullValue)
+		{
+			continue;
+		}
+		AllocatedCount++;
+		ListHeads[Idx] = CurBlockIdx;
+		NewBlocks.InsertAt(NullValue, CurBlockIdx + BLOCK_LIST_OFFSET);
+		for (int32 SubIdx = 0; SubIdx < BLOCK_LIST_OFFSET; ++SubIdx)
+		{
+			NewBlocks[CurBlockIdx + SubIdx] = ListBlocks[OrigHead + SubIdx];
+		}
+		
+		int32 OrigLinkStart = ListBlocks[OrigHead + BLOCK_LIST_OFFSET];
+		if (OrigLinkStart == NullValue)
+		{
+			NewBlocks[CurBlockIdx + BLOCK_LIST_OFFSET] = NullValue;
+		}
+		else
+		{
+			int32 CurPtr = OrigLinkStart;
+			NewBlocks[CurBlockIdx + BLOCK_LIST_OFFSET] = NewLinkedListElements.GetLength();
+			while (true)
+			{
+				NewLinkedListElements.Add(LinkedListElements[CurPtr]);
+				CurPtr = LinkedListElements[CurPtr + 1];
+				if (CurPtr == NullValue)
+				{
+					break;
+				}
+				NewLinkedListElements.Add(NewLinkedListElements.GetLength() + 1);
+			}
+			NewLinkedListElements.Add(NullValue);
+		}
+	}
+
+	ListBlocks = MoveTemp(NewBlocks);
+	LinkedListElements = MoveTemp(NewLinkedListElements);
+	FreeHeadIndex = NullValue;
+	FreeBlocks.Clear();
+}
 
 
 void FSmallListSet::Insert(int32 ListIndex, int32 Value)
diff --git a/Engine/Source/Runtime/GeometryCore/Public/Util/SmallListSet.h b/Engine/Source/Runtime/GeometryCore/Public/Util/SmallListSet.h
--- a/Engine/Source/Runtime/GeometryCore/Public/Util/SmallListSet.h
+++ b/Engine/Source/Runtime/GeometryCore/Public/Util/SmallListSet.h
@@ -95,6 +95,11 @@ public:
 		FreeHeadIndex = NullValue;
 	}
 
+	/**
+	 * Clearing any lists at or after MaxListIndex and compact the ListBlocks and LinkedListElements so there are no free blocks or free linked list elements
+	 */
+	void Compact(int32 MaxListIndex);
+
 
 	/**
 	 * @return true if a list has been allocated at the given ListIndex
