add clamp method to axis aligned box2/3

diff --git a/Engine/Source/Runtime/GeometryCore/Public/BoxTypes.h b/Engine/Source/Runtime/GeometryCore/Public/BoxTypes.h
--- a/Engine/Source/Runtime/GeometryCore/Public/BoxTypes.h
+++ b/Engine/Source/Runtime/GeometryCore/Public/BoxTypes.h
@@ -481,6 +481,15 @@ struct TAxisAlignedBox3
 		return Contains(Box.Min) && Contains(Box.Max);
 	}
 
+	TVector<RealType> Clamp(const TVector<RealType>& V)
+	{
+		return TVector<RealType>(
+			FMath::Clamp(V.X, Min.X, Max.X),
+			FMath::Clamp(V.Y, Min.Y, Max.Y),
+			FMath::Clamp(V.Z, Min.Z, Max.Z)
+		);
+	}
+
 	TAxisAlignedBox3<RealType> Intersect(const TAxisAlignedBox3<RealType>& Box) const
 	{
 		TAxisAlignedBox3<RealType> Intersection(
@@ -756,6 +765,14 @@ struct TAxisAlignedBox2
 		return Contains(Box.Min) && Contains(Box.Max);
 	}
 
+	TVector<RealType> Clamp(const TVector2<RealType>& V)
+	{
+		return TVector2<RealType>(
+			FMath::Clamp(V.X, Min.X, Max.X),
+			FMath::Clamp(V.Y, Min.Y, Max.Y)
+		);
+	}
+
 	bool Intersects(const TAxisAlignedBox2<RealType>& Box) const
 	{
 		return !((Box.Max.X < Min.X) || (Box.Min.X > Max.X) || (Box.Max.Y < Min.Y) || (Box.Min.Y > Max.Y));
