support optional vector cellsize in sdf calculation and interpolation

diff --git a/Engine/Source/Runtime/GeometryCore/Public/Implicit/GridInterpolant.h b/Engine/Source/Runtime/GeometryCore/Public/Implicit/GridInterpolant.h
--- a/Engine/Source/Runtime/GeometryCore/Public/Implicit/GridInterpolant.h
+++ b/Engine/Source/Runtime/GeometryCore/Public/Implicit/GridInterpolant.h
@@ -9,6 +9,8 @@
 #include "VectorTypes.h"
 #include "IntVectorTypes.h"
 
+#include <type_traits>
+
 namespace UE
 {
 namespace Geometry
@@ -22,13 +24,18 @@ namespace Geometry
  *
  * GridType must have a GetValue() that returns a value to interpolate at a given FVector3i location -- (w/ locations ranging from [0,0,0] to Dimensions (exclusive))
  */
-template <class GridType, typename RealType = double>
+template <class GridType, typename RealType = double, bool bScalarCellSize = true>
 class TTriLinearGridInterpolant /*: public BoundedImplicitFunction3d (TODO: consider add ImplicitFunction3d interface concept once more implicit functions are available*/
 {
+
 public:
+
+	// Type to use for CellSize
+	using CellSizeType = std::conditional_t<bScalarCellSize, RealType, TVector<RealType>>;
+
 	GridType* Grid;
 	TVector<RealType> GridOrigin;
-	RealType CellSize;
+	CellSizeType CellSize;
 	FVector3i Dimensions;
 
 	// value to return if query point is outside Grid (in an SDF
@@ -36,7 +43,7 @@ public:
 	// and cast this value to/from float; use TMathUtil<RealType>::SafeLargeValue to avoid overflow
 	RealType Outside = TMathUtil<RealType>::SafeLargeValue;
 
-	TTriLinearGridInterpolant(GridType* Grid, TVector<RealType> GridOrigin, RealType CellSize, FVector3i Dimensions) : Grid(Grid), GridOrigin(GridOrigin), CellSize(CellSize), Dimensions(Dimensions)
+	TTriLinearGridInterpolant(GridType* Grid, TVector<RealType> GridOrigin, CellSizeType CellSize, FVector3i Dimensions) : Grid(Grid), GridOrigin(GridOrigin), CellSize(CellSize), Dimensions(Dimensions)
 	{
 	}
 
@@ -44,18 +51,18 @@ public:
 	{
 		return TAxisAlignedBox3<RealType>(
 			{ GridOrigin.X, GridOrigin.Y, GridOrigin.Z },
-			{ GridOrigin.X + CellSize * (Dimensions.X - 1),
-			  GridOrigin.Y + CellSize * (Dimensions.Y - 1),
-			  GridOrigin.Z + CellSize * (Dimensions.Z - 1)});
+			{ GridOrigin.X + GetDim(CellSize, 0) * (Dimensions.X - 1),
+			  GridOrigin.Y + GetDim(CellSize, 1) * (Dimensions.Y - 1),
+			  GridOrigin.Z + GetDim(CellSize, 2) * (Dimensions.Z - 1)});
 	}
 
 	FVector3i Cell(const TVector<RealType>& Pt) const
 	{
 		// compute integer coordinates
 		FVector3i CellCoords;
-		CellCoords.X = (int)((Pt.X - GridOrigin.X) / CellSize);
-		CellCoords.Y = (int)((Pt.Y - GridOrigin.Y) / CellSize);
-		CellCoords.Z = (int)((Pt.Z - GridOrigin.Z) / CellSize);
+		CellCoords.X = (int)((Pt.X - GridOrigin.X) / GetDim(CellSize, 0));
+		CellCoords.Y = (int)((Pt.Y - GridOrigin.Y) / GetDim(CellSize, 1));
+		CellCoords.Z = (int)((Pt.Z - GridOrigin.Z) / GetDim(CellSize, 2));
 
 		return CellCoords;
 	}
@@ -64,9 +71,9 @@ public:
 	RealType Value(const TVector<RealType>& Pt) const
 	{
 		TVector<RealType> gridPt(
-			((Pt.X - GridOrigin.X) / CellSize),
-			((Pt.Y - GridOrigin.Y) / CellSize),
-			((Pt.Z - GridOrigin.Z) / CellSize));
+			((Pt.X - GridOrigin.X) / GetDim(CellSize, 0)),
+			((Pt.Y - GridOrigin.Y) / GetDim(CellSize, 1)),
+			((Pt.Z - GridOrigin.Z) / GetDim(CellSize, 2)));
 
 		if constexpr (bClamped)
 		{
@@ -95,7 +102,7 @@ public:
 		RealType fAx = gridPt.X - (RealType)X0;
 		RealType fAy = gridPt.Y - (RealType)Y0;
 		RealType fAz = gridPt.Z - (RealType)Z0;
-		RealType OneMinusfAx = 1.0 - fAx;
+		RealType OneMinusfAx = (RealType)(1.0) - fAx;
 
 		// compute trilinear interpolant. The code below tries to do this with the fewest 
 		// number of variables, in hopes that optimizer will be clever about re-using registers, etc.
@@ -145,9 +152,9 @@ public:
 	TVector<RealType> Gradient(const TVector<RealType>& Pt) const
 	{
 		TVector<RealType> gridPt = TVector<RealType>(
-			((Pt.X - GridOrigin.X) / CellSize),
-			((Pt.Y - GridOrigin.Y) / CellSize),
-			((Pt.Z - GridOrigin.Z) / CellSize));
+			((Pt.X - GridOrigin.X) /  GetDim(CellSize, 0)),
+			((Pt.Y - GridOrigin.Y) /  GetDim(CellSize, 1)),
+			((Pt.Z - GridOrigin.Z) /  GetDim(CellSize, 2)));
 
 		// clamp to Grid
 		if (gridPt.X < 0 || gridPt.X >= Dimensions.X - 1 ||
@@ -209,6 +216,19 @@ public:
 
 		return TVector<RealType>(gradX, gradY, gradZ);
 	}
+
+private:
+	FORCEINLINE static RealType GetDim(CellSizeType CellSize, int32 Axis)
+	{
+		if constexpr (bScalarCellSize)
+		{
+			return CellSize;
+		}
+		else
+		{
+			return CellSize[Axis];
+		}
+	}
 };
 
 
diff --git a/Engine/Source/Runtime/GeometryCore/Public/Implicit/SweepingMeshSDF.h b/Engine/Source/Runtime/GeometryCore/Public/Implicit/SweepingMeshSDF.h
--- a/Engine/Source/Runtime/GeometryCore/Public/Implicit/SweepingMeshSDF.h
+++ b/Engine/Source/Runtime/GeometryCore/Public/Implicit/SweepingMeshSDF.h
@@ -41,16 +41,28 @@ namespace Geometry
  * 
  * This code is based on the implementation found at [REDACTED_PROFILE_URL]/SDFGen
  */
-template <class TriangleMeshType>
+template <class TriangleMeshType, bool bScalarCellSize = true>
 class TSweepingMeshSDF
 {
 public:
 
+	// Type to use for CellSize
+	using CellSizeType = std::conditional_t<bScalarCellSize, float, FVector3f>;
+	// Type to use for CellSize, but in double precision
+	using CellSizeTyped = std::conditional_t<bScalarCellSize, double, FVector3d>;
+	// Unit cube, in CellSizeType
+	static inline CellSizeType UnitCellSize() { return CellSizeType(1); };
+
 	// INPUTS
 
 	const TriangleMeshType* Mesh;
 	TMeshAABBTree3<TriangleMeshType>* Spatial;
-	float CellSize;
+	CellSizeType CellSize;
+
+	void SetCellSize(float InCellSize)
+	{
+		CellSize = CellSizeType(InCellSize);
+	}
 
 	// Width of the band around triangles for which exact Distances are computed
 	// (In fact this is conservative, the band is often larger locally)
@@ -129,10 +141,10 @@ public:
 
 	/**
 	 * @param Mesh Triangle mesh to build an SDF around
-	 * @param CellSize Spacing between Grid points
+	 * @param InCellSize Spacing between Grid points
 	 * @param Spatial Optional AABB tree; note it *must* be provided if ComputeMode is set to NarrowBand_SpatialFloodFill
 	 */
-	TSweepingMeshSDF(const TriangleMeshType* Mesh = nullptr, float CellSize = 1, TMeshAABBTree3<TriangleMeshType>* Spatial = nullptr) : Mesh(Mesh), Spatial(Spatial), CellSize(CellSize)
+	TSweepingMeshSDF(const TriangleMeshType* Mesh = nullptr, CellSizeType InCellSize = UnitCellSize(), TMeshAABBTree3<TriangleMeshType>* Spatial = nullptr) : Mesh(Mesh), Spatial(Spatial), CellSize(InCellSize), GridOrigin(0)
 	{
 	}
 
@@ -146,9 +158,12 @@ public:
 		IntersectionsGrid.Resize(0, 0, 0, EAllowShrinking::Yes);
 	}
 
-	TTriLinearGridInterpolant<TSweepingMeshSDF> MakeInterpolant()
+	// Make an interpolating sampler of the sdf data. Note the lifetime is dependent on this class.
+	template<typename InterpolantRealType = double>
+	TTriLinearGridInterpolant<FDenseGrid3f, InterpolantRealType, bScalarCellSize> MakeInterpolant()
 	{
-		return TTriLinearGridInterpolant<TSweepingMeshSDF>(this, (FVector3d)GridOrigin, CellSize, Dimensions());
+		using InterpVecType = TVector<InterpolantRealType>;
+		return TTriLinearGridInterpolant<FDenseGrid3f, InterpolantRealType, bScalarCellSize>(&Grid, (InterpVecType)GridOrigin, (InterpVecType)CellSize, Dimensions());
 	}
 
 	/**
@@ -173,7 +188,7 @@ public:
 		{
 			return false;
 		}
-		if (CellSize <= 0 || !FMath::IsFinite(CellSize))
+		if (!IsValid(CellSize))
 		{
 			return false;
 		}
@@ -185,7 +200,7 @@ public:
 	}
 
 	/**
-	 * Compute the SDF
+	 * Compute the SDF enclosing the given bounds, with a 'safe' buffer zone
 	 * @param Bounds Bounding box for the mesh data (passed as param it is usually already available, depending on TriangleMeshType and whether an AABB tree was provided)
 	 * @return false if cancelled or failed (e.g. due to invalid arguments); true otherwise
 	 */
@@ -201,26 +216,63 @@ public:
 			ExactBandWidth = FMath::Max(ApproxMaxCellsPerDimension/2-1, 1);
 		}
 
-		double MaxDim = MaxElement(Bounds.Max - Bounds.Min + ExpandBounds * 2.0);
-		if (!ensureMsgf(MaxDim / CellSize <= ApproxMaxCellsPerDimension - 2 * ExactBandWidth, TEXT("SDF resolution clamped to avoid excessive memory use")))
+		FVector3d ExpandedDiagonal = (Bounds.Diagonal() + ExpandBounds * 2.0);
+		FVector3d ApproxVoxelsPerDim = (FVector3d)(ExpandedDiagonal / CellSize);
+		int32 MaxVoxelsDim = VectorUtil::Max3Index(ApproxVoxelsPerDim);
+		if (!ensureMsgf(ApproxVoxelsPerDim[MaxVoxelsDim] <= ApproxMaxCellsPerDimension - 2 * ExactBandWidth, TEXT("SDF resolution clamped to avoid excessive memory use")))
 		{
-			CellSize = float( MaxDim / (ApproxMaxCellsPerDimension - 2 * ExactBandWidth) );
-			if (!ensure(CellSize > 0 && FMath::IsFinite(CellSize)))
+			float CellSizeFactor = float(ExpandedDiagonal[MaxVoxelsDim] / (ApproxMaxCellsPerDimension - 2 * ExactBandWidth)) / GetDim(CellSize, MaxVoxelsDim);
+			CellSize *= CellSizeFactor;
+			if (!ensure(IsValid(CellSize)))
 			{
 				return false;
 			}
 		}
 
-		float fBufferWidth = float(ExactBandWidth) * CellSize;
+		CellSizeType fBufferWidth = float(ExactBandWidth) * CellSize;
 		if (ComputeMode == EComputeModes::NarrowBand_SpatialFloodFill)
 		{
-			fBufferWidth = (float)FMath::Max(fBufferWidth, float(NarrowBandMaxDistance));
+			if constexpr (bScalarCellSize)
+			{
+				fBufferWidth = (float)FMath::Max(fBufferWidth, float(NarrowBandMaxDistance));
+			}
+			else
+			{
+				for (int32 Idx = 0; Idx < 3; ++Idx)
+				{
+					fBufferWidth = FMath::Max(fBufferWidth, float(NarrowBandMaxDistance));
+				}
+			}
 		}
-		GridOrigin = (FVector3f)Bounds.Min - fBufferWidth * FVector3f::One() - (FVector3f)ExpandBounds;
+		FVector3f UseOrigin = (FVector3f)Bounds.Min - fBufferWidth * FVector3f::One() - (FVector3f)ExpandBounds;
 		FVector3f max = (FVector3f)Bounds.Max + fBufferWidth * FVector3f::One() + (FVector3f)ExpandBounds;
-		int NI = (int)((max.X - GridOrigin.X) / CellSize) + 1;
-		int NJ = (int)((max.Y - GridOrigin.Y) / CellSize) + 1;
-		int NK = (int)((max.Z - GridOrigin.Z) / CellSize) + 1;
+		FVector3i UseDims(
+			(int)((max.X - UseOrigin.X) / GetDim(CellSize, 0)) + 1,
+			(int)((max.Y - UseOrigin.Y) / GetDim(CellSize, 1)) + 1,
+			(int)((max.Z - UseOrigin.Z) / GetDim(CellSize, 2)) + 1
+		);
+
+		return Compute(UseOrigin, UseDims);
+	}
+	
+	/**
+	 * Compute the SDF with exactly specified grid origin and dimensions
+	 * @param InGridOrigin Origin of the output grid
+	 * @param InDimensions Dimensions of the output grid
+	 * @return false if cancelled or failed (e.g. due to invalid arguments); true otherwise
+	 */
+	bool Compute(FVector3f InGridOrigin, FVector3i InDimensions)
+	{
+		GridOrigin = InGridOrigin;
+		int NI = InDimensions.X;
+		int NJ = InDimensions.Y;
+		int NK = InDimensions.Z;
+
+		// dimensions must be positive
+		if (NI < 0 || NJ < 0 || NK < 0)
+		{
+			return false;
+		}
 
 		if (ComputeMode == EComputeModes::NarrowBand_SpatialFloodFill)
 		{
@@ -296,14 +348,14 @@ public:
 private:
 	FVector3f cell_center(FVector3i IJK) const
 	{
-		return FVector3f((float)IJK.X * CellSize + GridOrigin[0],
-			(float)IJK.Y * CellSize + GridOrigin[1],
-			(float)IJK.Z * CellSize + GridOrigin[2]);
+		return FVector3f((float)IJK.X * GetDim(CellSize, 0) + GridOrigin[0],
+			(float)IJK.Y * GetDim(CellSize, 1) + GridOrigin[1],
+			(float)IJK.Z * GetDim(CellSize, 2) + GridOrigin[2]);
 	}
 
 	float upper_bound(const FDenseGrid3f& GridIn) const
 	{
-		return (float(GridIn.GetDimensions().X + GridIn.GetDimensions().Y + GridIn.GetDimensions().Z) * CellSize);
+		return float(GridIn.GetDimensions().X) * GetDim(CellSize, 0) + float(GridIn.GetDimensions().Y) * GetDim(CellSize, 1) + float(GridIn.GetDimensions().Z) * GetDim(CellSize, 2);
 	}
 
 	float cell_tri_dist(const FVector3i& Idx, int TID) const
@@ -315,7 +367,7 @@ private:
 	}
 
 
-	void make_level_set3(FVector3f Origin, float DX, int NI, int NJ, int NK, FDenseGrid3f& Distances, int ExactBand)
+	void make_level_set3(FVector3f Origin, CellSizeType DX, int NI, int NJ, int NK, FDenseGrid3f& Distances, int ExactBand)
 	{
 		Distances.Resize(NI, NJ, NK);
 		Distances.Assign(upper_bound(Distances)); // upper bound on distance
@@ -333,7 +385,7 @@ private:
 		// Compute narrow-band Distances. For each triangle, we find its Grid-coord-bbox,
 		// and compute exact Distances within that box. The intersection_count Grid
 		// is also filled in this computation
-		double ddx = (double)DX;
+		CellSizeTyped ddx = (CellSizeTyped)DX;
 		double ox = (double)Origin[0], oy = (double)Origin[1], oz = (double)Origin[2];
 		FVector3d xp, xq, xr;
 		for (int TID = 0; TID < Mesh->MaxTriangleID(); TID++)
@@ -349,9 +401,9 @@ private:
 			Mesh->GetTriVertices(TID, xp, xq, xr);
 
 			// real IJK coordinates of xp/xq/xr
-			double fip = (xp[0] - ox) / ddx, fjp = (xp[1] - oy) / ddx, fkp = (xp[2] - oz) / ddx;
-			double fiq = (xq[0] - ox) / ddx, fjq = (xq[1] - oy) / ddx, fkq = (xq[2] - oz) / ddx;
-			double fir = (xr[0] - ox) / ddx, fjr = (xr[1] - oy) / ddx, fkr = (xr[2] - oz) / ddx;
+			double fip = (xp[0] - ox) / GetDim(ddx, 0), fjp = (xp[1] - oy) / GetDim(ddx, 1), fkp = (xp[2] - oz) / GetDim(ddx, 2);
+			double fiq = (xq[0] - ox) / GetDim(ddx, 0), fjq = (xq[1] - oy) / GetDim(ddx, 1), fkq = (xq[2] - oz) / GetDim(ddx, 2);
+			double fir = (xr[0] - ox) / GetDim(ddx, 0), fjr = (xr[1] - oy) / GetDim(ddx, 1), fkr = (xr[2] - oz) / GetDim(ddx, 2);
 
 			// clamped integer bounding box of triangle plus exact-band
 			int i0 = FMath::Clamp(((int)FMath::Min3(fip, fiq, fir)) - ExactBand, 0, NI - 1);
@@ -366,7 +418,7 @@ private:
 			for (int K = k0; K <= k1; ++K) {
 				for (int J = j0; J <= j1; ++J) {
 					for (int I = i0; I <= i1; ++I) {
-						FVector3d gx((float)I * DX + Origin[0], (float)J * DX + Origin[1], (float)K * DX + Origin[2]);
+						FVector3d gx((float)I * GetDim(DX, 0) + Origin[0], (float)J * GetDim(DX, 1) + Origin[1], (float)K * GetDim(DX, 2) + Origin[2]);
 						float d = (float)PointTriangleDistance(gx, xp, xq, xr);
 						if (d < Distances.At(I, J, K)) {
 							Distances.At(I, J, K) = d;
@@ -434,7 +486,7 @@ private:
 
 
 
-	void make_level_set3_parallel(FVector3f Origin, float DX, int NI, int NJ, int NK, FDenseGrid3f& Distances, int ExactBand)
+	void make_level_set3_parallel(FVector3f Origin, CellSizeType DX, int NI, int NJ, int NK, FDenseGrid3f& Distances, int ExactBand)
 	{
 		Distances.Resize(NI, NJ, NK);
 		Distances.Assign(upper_bound(Grid)); // upper bound on distance
@@ -450,7 +502,7 @@ private:
 		IntersectionsGrid.Assign(0);
 
 		double ox = (double)Origin[0], oy = (double)Origin[1], oz = (double)Origin[2];
-		double invdx = 1.0 / DX;
+		CellSizeTyped invdx = CellSizeTyped(UnitCellSize()) / (CellSizeTyped)DX;
 
 		// Compute narrow-band Distances. For each triangle, we find its Grid-coord-bbox,
 		// and compute exact Distances within that box.
@@ -481,9 +533,9 @@ private:
 			Mesh->GetTriVertices(TID, xp, xq, xr);
 
 			// real IJK coordinates of xp/xq/xr
-			double fip = (xp[0] - ox) * invdx, fjp = (xp[1] - oy) * invdx, fkp = (xp[2] - oz) * invdx;
-			double fiq = (xq[0] - ox) * invdx, fjq = (xq[1] - oy) * invdx, fkq = (xq[2] - oz) * invdx;
-			double fir = (xr[0] - ox) * invdx, fjr = (xr[1] - oy) * invdx, fkr = (xr[2] - oz) * invdx;
+			double fip = (xp[0] - ox) * GetDim(invdx, 0), fjp = (xp[1] - oy) * GetDim(invdx, 1), fkp = (xp[2] - oz) * GetDim(invdx, 2);
+			double fiq = (xq[0] - ox) * GetDim(invdx, 0), fjq = (xq[1] - oy) * GetDim(invdx, 1), fkq = (xq[2] - oz) * GetDim(invdx, 2);
+			double fir = (xr[0] - ox) * GetDim(invdx, 0), fjr = (xr[1] - oy) * GetDim(invdx, 1), fkr = (xr[2] - oz) * GetDim(invdx, 2);
 
 			// clamped integer bounding box of triangle plus exact-band
 			int i0 = FMath::Clamp(((int)FMath::Min3(fip, fiq, fir)) - ExactBand, 0, NI - 1);
@@ -500,7 +552,7 @@ private:
 					int base_idx = ((J < wj) ? 0 : 1) | ((K < wk) ? 0 : 2);    // construct index into spinlocks array
 
 					for (int I = i0; I <= i1; ++I) {
-						FVector3d gx((float)I * DX + Origin[0], (float)J * DX + Origin[1], (float)K * DX + Origin[2]);
+						FVector3d gx((float)I * GetDim(DX, 0) + Origin[0], (float)J * GetDim(DX, 1) + Origin[1], (float)K * GetDim(DX, 2) + Origin[2]);
 						float d = (float)PointTriangleDistance(gx, xp, xq, xr);
 						if (d < Distances.At(I, J, K)) {
 							int lock_idx = base_idx | ((I < wi) ? 0 : 4);
@@ -566,7 +618,7 @@ private:
 
 
 
-	void make_level_set3_parallel_spatial(FVector3f Origin, float DX, int NI, int NJ, int NK, FDenseGrid3f& Distances, int ExactBand)
+	void make_level_set3_parallel_spatial(FVector3f Origin, CellSizeType DX, int NI, int NJ, int NK, FDenseGrid3f& Distances, int ExactBand)
 	{
 		Distances.Resize(NI, NJ, NK);
 		float upper_bound = this->upper_bound(Distances);
@@ -583,7 +635,7 @@ private:
 		IntersectionsGrid.Assign(0);
 
 		double ox = (double)Origin[0], oy = (double)Origin[1], oz = (double)Origin[2];
-		double invdx = 1.0 / DX;
+		CellSizeTyped invdx = CellSizeTyped(UnitCellSize()) / (CellSizeTyped)DX;
 
 		// Compute narrow-band Distances. For each triangle, we find its Grid-coord-bbox,
 		// and compute exact Distances within that box.
@@ -614,9 +666,9 @@ private:
 			Mesh->GetTriVertices(TID, xp, xq, xr);
 
 			// real IJK coordinates of xp/xq/xr
-			double fip = (xp[0] - ox) * invdx, fjp = (xp[1] - oy) * invdx, fkp = (xp[2] - oz) * invdx;
-			double fiq = (xq[0] - ox) * invdx, fjq = (xq[1] - oy) * invdx, fkq = (xq[2] - oz) * invdx;
-			double fir = (xr[0] - ox) * invdx, fjr = (xr[1] - oy) * invdx, fkr = (xr[2] - oz) * invdx;
+			double fip = (xp[0] - ox) * GetDim(invdx, 0), fjp = (xp[1] - oy) * GetDim(invdx, 1), fkp = (xp[2] - oz) * GetDim(invdx, 2);
+			double fiq = (xq[0] - ox) * GetDim(invdx, 0), fjq = (xq[1] - oy) * GetDim(invdx, 1), fkq = (xq[2] - oz) * GetDim(invdx, 2);
+			double fir = (xr[0] - ox) * GetDim(invdx, 0), fjr = (xr[1] - oy) * GetDim(invdx, 1), fkr = (xr[2] - oz) * GetDim(invdx, 2);
 
 			// clamped integer bounding box of triangle plus exact-band
 			int i0 = FMath::Clamp(((int)FMath::Min3(fip, fiq, fir)) - ExactBand, 0, NI - 1);
@@ -637,13 +689,13 @@ private:
 			}
 		});
 
-		double max_dist = ExactBand * (DX * FMathd::Sqrt2);
+		double max_dist = ExactBand * (double)GetDiagLength(DX);
 		ParallelFor(Grid.Size(), [this, &Origin, DX, NI, NJ, NK, &Distances, max_dist, upper_bound, &closest_tri](int LinearIdx)
 		{
 			FVector3i Idx = Grid.ToIndex(LinearIdx);
 			if (Distances[Idx] == 1) {
 				int I = Idx.X, J = Idx.Y, K = Idx.Z;
-				FVector3d p((float)I * DX + Origin[0], (float)J * DX + Origin[1], (float)K * DX + Origin[2]);
+				FVector3d p((float)I * GetDim(DX, 0) + Origin[0], (float)J * GetDim(DX, 1) + Origin[1], (float)K * GetDim(DX, 2) + Origin[2]);
 				double dsqr;
 				int near_tid = Spatial->FindNearestTriangle(p, dsqr, max_dist);
 				if (near_tid == IndexConstants::InvalidID) {
@@ -717,7 +769,7 @@ private:
 
 
 
-	void make_level_set3_parallel_floodfill(FVector3f Origin, float DX, int NI, int NJ, int NK, FDenseGrid3f& Distances)
+	void make_level_set3_parallel_floodfill(FVector3f Origin, CellSizeType DX, int NI, int NJ, int NK, FDenseGrid3f& Distances)
 	{
 		Distances.Resize(NI, NJ, NK);
 		float upper_bound = this->upper_bound(Distances);
@@ -734,7 +786,7 @@ private:
 		IntersectionsGrid.Assign(0);
 
 		double ox = (double)Origin[0], oy = (double)Origin[1], oz = (double)Origin[2];
-		double invdx = 1.0 / DX;
+		CellSizeTyped invdx = CellSizeTyped(UnitCellSize()) / (CellSizeTyped)DX;
 
 		// the steps below that populate the grid will potentially touch the same cells at the
 		// same time, so we need to lock them. However locking the entire grid for each cell 
@@ -806,7 +858,7 @@ private:
 
 				FVector3d v = Mesh->GetVertex(vid);
 				// real IJK coordinates of v
-				double fi = (v.X - ox) * invdx, fj = (v.Y - oy) * invdx, fk = (v.Z - oz) * invdx;
+				double fi = (v.X - ox) * GetDim(invdx, 0), fj = (v.Y - oy) * GetDim(invdx, 1), fk = (v.Z - oz) * GetDim(invdx, 2);
 				FVector3i Idx(
 					FMath::Clamp((int)fi, 0, NI - 1),
 					FMath::Clamp((int)fj, 0, NJ - 1),
@@ -842,7 +894,7 @@ private:
 		// we could do this parallel w/ some kind of producer-consumer...
 		FAxisAlignedBox3i Bounds = Distances.BoundsInclusive();
 		double max_dist = NarrowBandMaxDistance;
-		double max_query_dist = max_dist + (2.0 * DX * FMathd::Sqrt2);
+		double max_query_dist = max_dist + GetDiagLength(DX);
 		{
 			TRACE_CPUPROFILER_EVENT_SCOPE(Geometry_SweepingMeshSDF_FloodFill);
 			while (PendingCellQueue.Num() > 0)
@@ -963,7 +1015,7 @@ private:
 
 
 	// sweep through Grid in different directions, Distances and closest tris
-	void sweep_pass(FVector3f Origin, float DX, FDenseGrid3f& Distances, FDenseGrid3i& closest_tri)
+	void sweep_pass(FVector3f Origin, CellSizeType DX, FDenseGrid3f& Distances, FDenseGrid3i& closest_tri)
 	{
 		sweep(Distances, closest_tri, Origin, DX, +1, +1, +1);
 		if (CancelF()) return;
@@ -984,7 +1036,7 @@ private:
 
 
 	// single sweep pass
-	void sweep(FDenseGrid3f& phi, FDenseGrid3i& closest_tri, FVector3f Origin, float DX, int di, int dj, int dk)
+	void sweep(FDenseGrid3f& phi, FDenseGrid3i& closest_tri, FVector3f Origin, CellSizeType DX, int di, int dj, int dk)
 	{
 		int i0, i1;
 		if (di > 0) { i0 = 1; i1 = phi.GetDimensions().X; }
@@ -1005,7 +1057,7 @@ private:
 			{
 				for (int I = i0; I != i1; I += di)
 				{
-					FVector3d gx(float(I) * DX + Origin[0], float(J) * DX + Origin[1], float(K) * DX + Origin[2]);
+					FVector3d gx(float(I) * GetDim(DX, 0) + Origin[0], float(J) * GetDim(DX, 1) + Origin[1], float(K) * GetDim(DX, 2) + Origin[2]);
 					check_neighbour(phi, closest_tri, gx, I, J, K, I - di, J, K);
 					check_neighbour(phi, closest_tri, gx, I, J, K, I, J - dj, K);
 					check_neighbour(phi, closest_tri, gx, I, J, K, I - di, J - dj, K);
@@ -1038,17 +1090,17 @@ private:
 
 
 	// fill the intersection Grid w/ number of intersections in each cell
-	void compute_intersections(FVector3f Origin, float DX, int32 NI, int32 NJ, int32 NK, FDenseGrid3i& IntersectionCount)
+	void compute_intersections(FVector3f Origin, CellSizeType DX, int32 NI, int32 NJ, int32 NK, FDenseGrid3i& IntersectionCount)
 	{
 		double ox = (double)Origin[0], oy = (double)Origin[1], oz = (double)Origin[2];
-		double invdx = 1.0 / DX;
+		CellSizeTyped invdx = CellSizeTyped(UnitCellSize()) / (CellSizeTyped)DX;
 
 		bool cancelled = false;
 
 		// this is what we will do for each triangle. There are no Grid-reads, only Grid-writes, 
 		// since we use atomic_increment, it is always thread-safe
 		ParallelFor(Mesh->MaxTriangleID(),
-			[this, &Origin, &DX, &NI, &NJ, &NK,
+			[this, &Origin, &NI, &NJ, &NK,
 			&IntersectionCount, &ox, &oy, &oz, &invdx, &cancelled](int32 TID)
 			{
 				if (!Mesh->IsTriangle(TID))
@@ -1076,9 +1128,9 @@ private:
 				}
 
 				// real IJK coordinates of xp/xq/xr
-				double fip = (xp[0] - ox) * invdx, fjp = (xp[1] - oy) * invdx, fkp = (xp[2] - oz) * invdx;
-				double fiq = (xq[0] - ox) * invdx, fjq = (xq[1] - oy) * invdx, fkq = (xq[2] - oz) * invdx;
-				double fir = (xr[0] - ox) * invdx, fjr = (xr[1] - oy) * invdx, fkr = (xr[2] - oz) * invdx;
+				double fip = (xp[0] - ox) * GetDim(invdx, 0), fjp = (xp[1] - oy) * GetDim(invdx, 1), fkp = (xp[2] - oz) * GetDim(invdx, 2);
+				double fiq = (xq[0] - ox) * GetDim(invdx, 0), fjq = (xq[1] - oy) * GetDim(invdx, 1), fkq = (xq[2] - oz) * GetDim(invdx, 2);
+				double fir = (xr[0] - ox) * GetDim(invdx, 0), fjr = (xr[1] - oy) * GetDim(invdx, 1), fkr = (xr[2] - oz) * GetDim(invdx, 2);
 
 				// recompute J/K integer bounds of triangle w/o exact band
 				int32 j0 = FMath::Clamp(FMath::CeilToInt32(FMath::Min3(fjp, fjq, fjr)), 0, NJ - 1);
@@ -1339,7 +1391,51 @@ public:
 		}
 	}
 
-
+	// Helper methods to access CellSize in either scalar or vector mode
+	FORCEINLINE static double GetDim(CellSizeTyped CellSize, int32 Axis)
+	{
+		if constexpr (bScalarCellSize)
+		{
+			return (double)CellSize;
+		}
+		else
+		{
+			return (double)CellSize[Axis];
+		}
+	}
+	FORCEINLINE static float GetDim(CellSizeType CellSize, int32 Axis)
+	{
+		if constexpr (bScalarCellSize)
+		{
+			return (float)CellSize;
+		}
+		else
+		{
+			return (float)CellSize[Axis];
+		}
+	}
+	FORCEINLINE static float GetDiagLength(CellSizeType CellSize)
+	{
+		if constexpr (bScalarCellSize)
+		{
+			return CellSize * FMathf::Sqrt3;
+		}
+		else
+		{
+			return CellSize.Length();
+		}
+	}
+	FORCEINLINE static bool IsValid(CellSizeType CellSize)
+	{
+		if constexpr (std::is_floating_point_v<CellSizeType>)
+		{
+			return CellSize > 0 && FMath::IsFinite(CellSize);
+		}
+		else
+		{
+			return CellMinDim(CellSize) > 0 && FMath::IsFinite(CellSize[0]) && FMath::IsFinite(CellSize[1]) && FMath::IsFinite(CellSize[2]);
+		}
+	}
 };
 
 
diff --git a/Engine/Source/Runtime/GeometryCore/Public/Intersection/ContainmentQueries3.h b/Engine/Source/Runtime/GeometryCore/Public/Intersection/ContainmentQueries3.h
--- a/Engine/Source/Runtime/GeometryCore/Public/Intersection/ContainmentQueries3.h
+++ b/Engine/Source/Runtime/GeometryCore/Public/Intersection/ContainmentQueries3.h
@@ -24,7 +24,7 @@ namespace UE
 
 	namespace Geometry
 	{
-		template <class GridType, typename RealType> class TTriLinearGridInterpolant;
+		template <class GridType, typename RealType, bool bScalarCellSize> class TTriLinearGridInterpolant;
 		template <typename T> struct TCapsule3;
 		template <typename T> struct THalfspace3;
 		template <typename T> struct TSphere3;
