add conversions from FVector3i <-> FIntVector

diff --git a/Engine/Source/Runtime/GeometryCore/Public/IntVectorTypes.h b/Engine/Source/Runtime/GeometryCore/Public/IntVectorTypes.h
--- a/Engine/Source/Runtime/GeometryCore/Public/IntVectorTypes.h
+++ b/Engine/Source/Runtime/GeometryCore/Public/IntVectorTypes.h
@@ -3,6 +3,7 @@
 #pragma once
 
 #include "Math/Vector.h"
+#include "Math/IntVector.h"
 #include "MathUtil.h"
 #include "VectorTypes.h"
 #include <sstream>
@@ -70,6 +71,10 @@ PRAGMA_ENABLE_DEPRECATION_WARNINGS
 	{
 		return FVector2d((double)X, (double)Y);
 	}
+	explicit operator FIntVector2() const
+	{
+		return FIntVector2(X, Y);
+	}
 
 	explicit FVector2i(const FVector& Vec)
 		: X((int32)Vec.X), Y((int32)Vec.Y)
@@ -86,6 +91,11 @@ PRAGMA_ENABLE_DEPRECATION_WARNINGS
 	{
 	}
 
+	explicit FVector2i(const FIntVector2& Vec)
+		: X(Vec.X), Y(Vec.Y)
+	{
+	}
+
 	constexpr static FVector2i Zero()
 	{
 		return FVector2i(0, 0);
@@ -293,6 +303,11 @@ PRAGMA_ENABLE_DEPRECATION_WARNINGS
 		return FVector3d((double)X, (double)Y, (double)Z);
 	}
 
+	explicit operator FIntVector() const
+	{
+		return FIntVector(X, Y, Z);
+	}
+
 	explicit FVector3i(const FVector3f& Vec)
 		: X((int32)Vec.X), Y((int32)Vec.Y), Z((int32)Vec.Z)
 	{
@@ -303,6 +318,11 @@ PRAGMA_ENABLE_DEPRECATION_WARNINGS
 	{
 	}
 
+	explicit FVector3i(const FIntVector& Vec)
+		: X(Vec.X), Y(Vec.Y), Z(Vec.Z)
+	{
+	}
+
 	static FVector3i Zero()
 	{
 		return FVector3i(0, 0, 0);
