make mesh local param compute vertex normals if they're not already computed on the mesh

diff --git a/Engine/Source/Runtime/GeometryCore/Public/Parameterization/MeshLocalParam.h b/Engine/Source/Runtime/GeometryCore/Public/Parameterization/MeshLocalParam.h
--- a/Engine/Source/Runtime/GeometryCore/Public/Parameterization/MeshLocalParam.h
+++ b/Engine/Source/Runtime/GeometryCore/Public/Parameterization/MeshLocalParam.h
@@ -8,6 +8,7 @@
 #include "FrameTypes.h"
 #include "Util/IndexPriorityQueue.h"
 #include "Util/DynamicVector.h"
+#include "DynamicMesh/MeshNormals.h"
 
 namespace UE
 {
@@ -268,7 +269,20 @@ protected:
 
 	FVector3d GetNormal(const int32 PointID) const
 	{
-		return FVector3d(PointSet->GetVertexNormal(PointID));
+		if (PointSet->HasVertexNormals())
+		{
+			return FVector3d(PointSet->GetVertexNormal(PointID));
+		}
+		// Note: Currently only have a standard vertex normal calculation for FDynamicMesh3
+		else if constexpr (std::is_same_v<PointSetType, FDynamicMesh3>)
+		{
+			return FMeshNormals::ComputeVertexNormal(*PointSet, PointID);
+		}
+		else
+		{
+			ensureMsgf(false, TEXT("Local parameterization requires vertex normals!"));
+			return FVector3d::UnitY();
+		}
 	}
 
 protected:
