MeshModelingTools: Add integrated squared distance from line along a segment, for use as a similarity metric when finding a path corresponding to a straight line along a mesh.

diff --git a/Engine/Source/Runtime/GeometryCore/Public/Distance/DistLine3Segment3.h b/Engine/Source/Runtime/GeometryCore/Public/Distance/DistLine3Segment3.h
--- a/Engine/Source/Runtime/GeometryCore/Public/Distance/DistLine3Segment3.h
+++ b/Engine/Source/Runtime/GeometryCore/Public/Distance/DistLine3Segment3.h
@@ -126,5 +126,31 @@ public:
 typedef TDistLine3Segment3<float> FDistLine3Segment3f;
 typedef TDistLine3Segment3<double> FDistLine3Segment3d;
 
+/**
+ * Integrates the squared distance to a line along a given segment.
+ */
+template <typename Real>
+double SquaredDistanceFromLineIntegratedAlongSegment(const TLine3<Real>& Line, const TSegment3<Real>& Segment)
+{
+	double SegmentLength = Segment.Length();
+
+	TVector<Real> SegmentStart = Segment.StartPoint();
+	// Point at length t along segment is SegmentStart + Segment.Direction * t
+
+	TVector<Real> ProjectedSegmentStart = Line.Origin + (SegmentStart - Line.Origin).Dot(Line.Direction) * Line.Direction;
+	TVector<Real> ProjectedSegmentDirection = Segment.Direction.Dot(Line.Direction) * Line.Direction;
+	// Projected point onto line at t is ProjectedSegmentStart + ProjectedSegmentDirection * t
+
+	// Difference of point on segment and projected point is the vector that we actually want to integrate over t.
+	TVector<Real> P_0 = SegmentStart - ProjectedSegmentStart;
+	TVector<Real> P_v = Segment.Direction - ProjectedSegmentDirection;
+
+	// Squared length of this vector is sum of component-wise squares, and after doing some math, you get that
+	//  the integral from 0 to SegmentLength is: 
+	//  p_0.p_0(SegmentLength) + p_0.p_v(SegmentLength^2) + p_v.p_v (SegmentLength^3)/3
+	// Rearranging a bit:
+	return ((P_v.Dot(P_v) * SegmentLength / 3.0 + P_0.Dot(P_v)) * SegmentLength + P_0.Dot(P_0)) * SegmentLength;
+}
+
 } // end namespace UE::Geometry
 } // end namespace UE
