GeometryProcessing: Added functionality to the UV Layout code to allow for preserving existing UV island's scale and orientation during layout operations.

Additionally cleaned up numerical logic within the UV Packer code and added comments for confusing areas of the code.

diff --git a/Engine/Plugins/Runtime/GeometryProcessing/Source/DynamicMesh/Private/Parameterization/MeshUVPacking.cpp b/Engine/Plugins/Runtime/GeometryProcessing/Source/DynamicMesh/Private/Parameterization/MeshUVPacking.cpp
--- a/Engine/Plugins/Runtime/GeometryProcessing/Source/DynamicMesh/Private/Parameterization/MeshUVPacking.cpp
+++ b/Engine/Plugins/Runtime/GeometryProcessing/Source/DynamicMesh/Private/Parameterization/MeshUVPacking.cpp
@@ -65,6 +65,8 @@ bool FDynamicMeshUVPacker::StandardPack()
 	TRACE_CPUPROFILER_EVENT_SCOPE(DynamicMeshUVPacker_StandardPack);
 	
 	FUVPacker Packer;
+	Packer.bPreserveScale = bPreserveScale;
+	Packer.bPreserveRotation = bPreserveRotation;
 	Packer.bAllowFlips = bAllowFlips;
 	Packer.GutterSize = GutterSize;
 	Packer.TextureResolution = TextureResolution;
@@ -84,6 +86,8 @@ bool FDynamicMeshUVPacker::StackPack()
 	TRACE_CPUPROFILER_EVENT_SCOPE(DynamicMeshUVPacker_StackPack);
 	
 	FUVPacker Packer;
+	Packer.bPreserveScale = bPreserveScale;
+	Packer.bPreserveRotation = bPreserveRotation;
 	Packer.bAllowFlips = bAllowFlips;
 	Packer.GutterSize = GutterSize;
 	Packer.TextureResolution = TextureResolution;
diff --git a/Engine/Plugins/Runtime/GeometryProcessing/Source/DynamicMesh/Public/Parameterization/MeshUVPacking.h b/Engine/Plugins/Runtime/GeometryProcessing/Source/DynamicMesh/Public/Parameterization/MeshUVPacking.h
--- a/Engine/Plugins/Runtime/GeometryProcessing/Source/DynamicMesh/Public/Parameterization/MeshUVPacking.h
+++ b/Engine/Plugins/Runtime/GeometryProcessing/Source/DynamicMesh/Public/Parameterization/MeshUVPacking.h
@@ -34,6 +34,12 @@ public:
 	/** Thickness of gutter/border in pixel dimensions. Not supported by all packing methods  */
 	float GutterSize = 1.0;
 
+	/** If true, original island scale is maintained during any packing process. */
+	bool bPreserveScale = false;
+
+	/** If true, original island rotation is maintained during any packing process. Automatically prevents bAllowFlips from applying, if set. */
+	bool bPreserveRotation = false;
+
 	/** If true, islands can be flipped in addition to rotate/translate/scale */
 	bool bAllowFlips = false;
 
diff --git a/Engine/Plugins/Runtime/MeshModelingToolset/Source/ModelingOperators/Private/ParameterizationOps/UVLayoutOp.cpp b/Engine/Plugins/Runtime/MeshModelingToolset/Source/ModelingOperators/Private/ParameterizationOps/UVLayoutOp.cpp
--- a/Engine/Plugins/Runtime/MeshModelingToolset/Source/ModelingOperators/Private/ParameterizationOps/UVLayoutOp.cpp
+++ b/Engine/Plugins/Runtime/MeshModelingToolset/Source/ModelingOperators/Private/ParameterizationOps/UVLayoutOp.cpp
@@ -164,6 +164,8 @@ void FUVLayoutOp::ExecutePacker(FDynamicMeshUVPacker& Packer)
 {	
 	Packer.GutterSize = this->GutterSize;
 	Packer.bAllowFlips = this->bAllowFlips;
+	Packer.bPreserveScale = this->bPreserveScale;
+	Packer.bPreserveRotation = this->bPreserveRotation;
 
 	if (UVLayoutMode == EUVLayoutOpLayoutModes::RepackToUnitRect)
 	{
@@ -216,6 +218,8 @@ TUniquePtr<FDynamicMeshOperator> UUVLayoutOperatorFactory::MakeNewOperator()
 
 	Op->UVLayerIndex = GetSelectedUVChannel();
 	Op->TextureResolution = Settings->TextureResolution;
+	Op->bPreserveScale = Settings->bPreserveScale;
+	Op->bPreserveRotation = Settings->bPreserveRotation;	
 	Op->bAllowFlips = Settings->bAllowFlips;
 	Op->UVScaleFactor = Settings->Scale;
 	Op->UVTranslation = FVector2f(Settings->Translation);
diff --git a/Engine/Plugins/Runtime/MeshModelingToolset/Source/ModelingOperators/Public/ParameterizationOps/UVLayoutOp.h b/Engine/Plugins/Runtime/MeshModelingToolset/Source/ModelingOperators/Public/ParameterizationOps/UVLayoutOp.h
--- a/Engine/Plugins/Runtime/MeshModelingToolset/Source/ModelingOperators/Public/ParameterizationOps/UVLayoutOp.h
+++ b/Engine/Plugins/Runtime/MeshModelingToolset/Source/ModelingOperators/Public/ParameterizationOps/UVLayoutOp.h
@@ -37,6 +37,8 @@ public:
 
 	int UVLayerIndex = 0;
 	int TextureResolution = 128;
+	bool bPreserveScale = false;
+	bool bPreserveRotation = false;
 	bool bAllowFlips = false;
 	bool bAlwaysSplitBowties = true;
 	float UVScaleFactor = 1.0;
diff --git a/Engine/Plugins/Runtime/MeshModelingToolset/Source/ModelingOperators/Public/Properties/UVLayoutProperties.h b/Engine/Plugins/Runtime/MeshModelingToolset/Source/ModelingOperators/Public/Properties/UVLayoutProperties.h
--- a/Engine/Plugins/Runtime/MeshModelingToolset/Source/ModelingOperators/Public/Properties/UVLayoutProperties.h
+++ b/Engine/Plugins/Runtime/MeshModelingToolset/Source/ModelingOperators/Public/Properties/UVLayoutProperties.h
@@ -49,8 +49,16 @@ public:
 	UPROPERTY(EditAnywhere, Category = "UV Layout")
 	FVector2D Translation = FVector2D(0,0);
 
-	/** Allow the Repack layout type to flip the orientation of UV islands to save space. Note that this may cause problems for downstream operations, and therefore is disabled by default. */
+	/** Force the Repack layout type to preserve existing scaling of UV islands. Note, this might lead to the packing not fitting within a unit square, and therefore is disabled by default. */
+	UPROPERTY(EditAnywhere, Category = "UV Layout", meta = (EditCondition = "LayoutType == EUVLayoutType::Repack"))
+	bool bPreserveScale = false;
+
+	/** Force the Repack layout type to preserve existing rotation of UV islands. Note, this might lead to the packing not being as space efficient as possible, and therefore is disabled by default. */
 	UPROPERTY(EditAnywhere, Category = "UV Layout", meta = (EditCondition = "LayoutType == EUVLayoutType::Repack"))
+	bool bPreserveRotation = false;
+
+	/** Allow the Repack layout type to flip the orientation of UV islands to save space. Note that this may cause problems for downstream operations, and therefore is disabled by default. */
+	UPROPERTY(EditAnywhere, Category = "UV Layout", meta = (EditCondition = "LayoutType == EUVLayoutType::Repack && bPreserveRotation == false"))
 	bool bAllowFlips = false;
 
 	/** Enable UDIM aware layout and keep islands within their originating UDIM tiles when laying out.*/
diff --git a/Engine/Source/Runtime/GeometryCore/Private/Parameterization/UVPacking.cpp b/Engine/Source/Runtime/GeometryCore/Private/Parameterization/UVPacking.cpp
--- a/Engine/Source/Runtime/GeometryCore/Private/Parameterization/UVPacking.cpp
+++ b/Engine/Source/Runtime/GeometryCore/Private/Parameterization/UVPacking.cpp
@@ -45,9 +45,16 @@ struct FUVIsland
 	double		ScaleToWorld; // Scale factor that would make 1 texel ~= 1 unit in world space
 
 	FVector2d	UVScale;
-	FVector2d	PackingScaleU;
-	FVector2d	PackingScaleV;
+	FMatrix2d   UVRotation;
+	FMatrix2d   UVPackingScale;
+
+	FVector2d   OrientationBiasShift;
 	FVector2d	PackingBias;
+
+	FVector2d   GetPackingBias() const
+	{		
+		return (UVRotation * (UVPackingScale * OrientationBiasShift)) + 0.5 + PackingBias;
+	}
 };
 
 
@@ -72,6 +79,9 @@ public:
 	// packing target texture resolution, used to calculate gutter/border size
 	uint32 TextureResolution = 512;
 
+	// if true, disallows any change to island rotations while packing. Implicitly disables bAllowFlips, if set.
+	bool bPreserveRotation = false;
+
 	// if true, UV islands can be mirrored in X and/or Y to improve packing
 	bool bAllowFlips = false;
 
@@ -115,9 +125,10 @@ bool FStandardChartPacker::FindBestPacking(TArray<FUVIsland>& Charts)
 	for (FUVIsland& Chart : Charts)
 	{
 		Chart.PackingBias = FVector2d::Zero();
-		Chart.PackingScaleU = FVector2d::Zero();
-		Chart.PackingScaleV = FVector2d::Zero();
 		Chart.UVScale = FVector2d::Zero();
+		Chart.UVPackingScale = FMatrix2d::Identity();
+		Chart.UVRotation = FMatrix2d::Identity();
+		Chart.OrientationBiasShift = FVector2d::Zero();
 	}
 
 	// Those might require tuning, changing them won't affect the outcome and will maintain backward compatibility
@@ -250,9 +261,6 @@ bool FStandardChartPacker::FindBestPacking(TArray<FUVIsland>& Charts)
 
 
 
-
-
-
 void FStandardChartPacker::ScaleCharts( TArray<FUVIsland>& Charts, double UVScale )
 {
 	for( int32 i = 0; i <Charts.Num(); i++ )
@@ -449,15 +457,16 @@ bool FStandardChartPacker::PackCharts(TArray<FUVIsland>& Charts, double UVScale,
 			//	return Rect.X + Rect.Y * TextureResolution <BestRect.X + BestRect.Y * TextureResolution;
 			//};
 
+			int32 MaxOrientation = bPreserveRotation ? 1 : 8;
 			int32 OrientationStep = (bAllowFlips) ? 1 : 2;
-			for (int32 Orientation = 0; Orientation < 8; Orientation += OrientationStep)
+			for (int32 Orientation = 0; Orientation < MaxOrientation; Orientation += OrientationStep)
 			{
 				// TODO If any dimension is less than 1 pixel shrink dimension to zero
 
 				OrientChart(Chart, Orientation);
 
 				FVector2d ChartSize = Chart.MaxUV - Chart.MinUV;
-				ChartSize = ChartSize.X * Chart.PackingScaleU + ChartSize.Y * Chart.PackingScaleV;
+				ChartSize = Chart.UVRotation * Chart.UVPackingScale * ChartSize;
 
 				// Only need half pixel dilate for rects
 				FUVSpaceAllocator::FRect	Rect;
@@ -489,6 +498,13 @@ bool FStandardChartPacker::PackCharts(TArray<FUVIsland>& Charts, double UVScale,
 				}
 				else
 				{
+
+					// This block is a little confusing - the intent here is that orientations
+					// that are flips always follow orientations that are their flip opposite... 
+					// So, we can avoid explicit rasterization for the flip orientations and
+					// just flip the previous iteration's raster instead.
+					// Technically we can rasterize in all cases and then flip, but it would
+					// probably be more expensive?
 					if (Orientation % 4 == 1)
 					{
 						ChartRaster.FlipX(Rect);
@@ -635,55 +651,63 @@ bool FStandardChartPacker::PackCharts(TArray<FUVIsland>& Charts, double UVScale,
 
 void FStandardChartPacker::OrientChart(FUVIsland& Chart, int32 Orientation)
 {
+	Chart.PackingBias = FVector2d::Zero();
+	Chart.UVPackingScale = FMatrix2d(Chart.UVScale.X, 0.0, 0.0, Chart.UVScale.Y);
+
+	/*
+	
+	Note: The packing bias is being initialized here to "correct" the positioning of the chart's bounding box after rotation.
+
+	How this works: 
+
+	1) Rotate the box, assuming 0,0 is the upper left hand corner
+
+	2) Shift box, so that the new upper left hand corner of the rotated box lands at 0,0 again.
+	
+	*/
+
+
 	switch (Orientation)
 	{
 	case 0:
 		// 0 degrees
-		Chart.PackingScaleU = FVector2d(Chart.UVScale.X, 0);
-		Chart.PackingScaleV = FVector2d(0, Chart.UVScale.Y);
-		Chart.PackingBias = -Chart.MinUV.X * Chart.PackingScaleU - Chart.MinUV.Y * Chart.PackingScaleV + 0.5f;
+		Chart.UVRotation = FMatrix2d(1.0, 0.0, 0.0, 1.0);
+		Chart.OrientationBiasShift = FVector2d(-Chart.MinUV.X, -Chart.MinUV.Y);
 		break;
 	case 1:
 		// 0 degrees, flip x
-		Chart.PackingScaleU = FVector2d(-Chart.UVScale.X, 0);
-		Chart.PackingScaleV = FVector2d(0, Chart.UVScale.Y);
-		Chart.PackingBias = -Chart.MaxUV.X * Chart.PackingScaleU - Chart.MinUV.Y * Chart.PackingScaleV + 0.5f;
+		Chart.UVRotation = FMatrix2d(-1.0, 0.0, 0.0, 1.0);
+		Chart.OrientationBiasShift = FVector2d(-Chart.MaxUV.X, -Chart.MinUV.Y);
 		break;
 	case 2:
 		// 90 degrees
-		Chart.PackingScaleU = FVector2d(0, -Chart.UVScale.X);
-		Chart.PackingScaleV = FVector2d(Chart.UVScale.Y, 0);
-		Chart.PackingBias = -Chart.MaxUV.X * Chart.PackingScaleU - Chart.MinUV.Y * Chart.PackingScaleV + 0.5f;
+		Chart.UVRotation = FMatrix2d(0.0, 1.0, -1.0, 0.0);
+		Chart.OrientationBiasShift = FVector2d(-Chart.MaxUV.X, -Chart.MinUV.Y);
 		break;
 	case 3:
 		// 90 degrees, flip x
-		Chart.PackingScaleU = FVector2d(0, Chart.UVScale.X);
-		Chart.PackingScaleV = FVector2d(Chart.UVScale.Y, 0);
-		Chart.PackingBias = -Chart.MinUV.X * Chart.PackingScaleU - Chart.MinUV.Y * Chart.PackingScaleV + 0.5f;
+		Chart.UVRotation = FMatrix2d(0.0, 1.0, 1.0, 0.0);
+		Chart.OrientationBiasShift = FVector2d(-Chart.MinUV.X, -Chart.MinUV.Y);
 		break;
 	case 4:
 		// 180 degrees
-		Chart.PackingScaleU = FVector2d(-Chart.UVScale.X, 0);
-		Chart.PackingScaleV = FVector2d(0, -Chart.UVScale.Y);
-		Chart.PackingBias = -Chart.MaxUV.X * Chart.PackingScaleU - Chart.MaxUV.Y * Chart.PackingScaleV + 0.5f;
+		Chart.UVRotation = FMatrix2d(-1.0, 0.0, 0.0, -1.0);
+		Chart.OrientationBiasShift = FVector2d(-Chart.MaxUV.X, -Chart.MaxUV.Y);
 		break;
 	case 5:
 		// 180 degrees, flip x
-		Chart.PackingScaleU = FVector2d(Chart.UVScale.X, 0);
-		Chart.PackingScaleV = FVector2d(0, -Chart.UVScale.Y);
-		Chart.PackingBias = -Chart.MinUV.X * Chart.PackingScaleU - Chart.MaxUV.Y * Chart.PackingScaleV + 0.5f;
+		Chart.UVRotation = FMatrix2d(1.0, 0.0, 0.0, -1.0);
+		Chart.OrientationBiasShift = FVector2d(-Chart.MinUV.X, -Chart.MaxUV.Y);
 		break;
 	case 6:
 		// 270 degrees
-		Chart.PackingScaleU = FVector2d(0, Chart.UVScale.X);
-		Chart.PackingScaleV = FVector2d(-Chart.UVScale.Y, 0);
-		Chart.PackingBias = -Chart.MinUV.X * Chart.PackingScaleU - Chart.MaxUV.Y * Chart.PackingScaleV + 0.5f;
+		Chart.UVRotation = FMatrix2d(0.0, -1.0, 1.0, 0.0);
+		Chart.OrientationBiasShift = FVector2d(-Chart.MinUV.X, -Chart.MaxUV.Y);
 		break;
 	case 7:
 		// 270 degrees, flip x
-		Chart.PackingScaleU = FVector2d(0, -Chart.UVScale.X);
-		Chart.PackingScaleV = FVector2d(-Chart.UVScale.Y, 0);
-		Chart.PackingBias = -Chart.MaxUV.X * Chart.PackingScaleU - Chart.MaxUV.Y * Chart.PackingScaleV + 0.5f;
+		Chart.UVRotation = FMatrix2d(0.0, -1.0, -1.0, 0.0);
+		Chart.OrientationBiasShift = FVector2d(-Chart.MaxUV.X, -Chart.MaxUV.Y);
 		break;
 	}
 }
@@ -795,7 +819,7 @@ void FStandardChartPacker::RasterizeChart(const FUVIsland& Chart, uint32 RectW,
 		for (int k = 0; k <3; k++)
 		{
 			FVector2d UV = (FVector2d)Mesh->GetUV(UVTriangle[k]);
-			Points[k] = (FVector2f)(UV.X * Chart.PackingScaleU + UV.Y * Chart.PackingScaleV + Chart.PackingBias);
+			Points[k] = (FVector2f)((Chart.UVRotation * Chart.UVPackingScale * UV) + Chart.GetPackingBias());
 		}
 
 		InternalRasterizeTriangle<16>(OutChartRaster, Points, RectW, RectH);
@@ -818,6 +842,7 @@ bool FUVPacker::StandardPack(IUVMeshView* Mesh, int NumIslands, TFunctionRef<voi
 	Packer.Mesh = Mesh;
 
 	Packer.TextureResolution = this->TextureResolution;
+	Packer.bPreserveRotation = this->bPreserveRotation;
 	Packer.bAllowFlips = this->bAllowFlips;
 
 	int32 NumCharts = NumIslands;
@@ -849,15 +874,42 @@ bool FUVPacker::StandardPack(IUVMeshView* Mesh, int NumIslands, TFunctionRef<voi
 		return false;
 	}
 
-
 	// Commit chart UVs
 	for (int32 i = 0; i <AllCharts.Num(); i++)
 	{
 		FUVIsland& Chart = AllCharts[i];
 
-		Chart.PackingScaleU /= (double)Packer.TextureResolution;
-		Chart.PackingScaleV /= (double)Packer.TextureResolution;
-		Chart.PackingBias /= (double)Packer.TextureResolution;
+		FVector2d PackingBias = Chart.GetPackingBias();
+		PackingBias /= (double)Packer.TextureResolution;
+		FMatrix2d NormalizedScale = Chart.UVPackingScale * (1.0 / Packer.TextureResolution);
+
+		auto PackUVsWithScaling = [&Mesh, &Chart, &NormalizedScale, &PackingBias](TSet<int32>& IslandElements)
+		{
+			for (int32 elemid : IslandElements)
+			{
+				if (elemid >= 0)
+				{
+					FVector2d UV = (FVector2d)Mesh->GetUV(elemid);
+					FVector2d TransformedUV = (Chart.UVRotation * NormalizedScale * UV) + PackingBias;
+					Mesh->SetUV(elemid, (FVector2f)TransformedUV);
+				}
+			}
+		};
+
+		auto PackUVsWithoutScaling = [&Mesh, &Chart, &NormalizedScale, &PackingBias](TSet<int32>& IslandElements)
+		{
+			FMatrix2d InverseScale = NormalizedScale.Inverse();
+
+			for (int32 elemid : IslandElements)
+			{
+				if (elemid >= 0)
+				{
+					FVector2d UV = (FVector2d)Mesh->GetUV(elemid);
+					FVector2d TransformedUV = InverseScale * ((Chart.UVRotation * NormalizedScale * UV) + PackingBias);
+					Mesh->SetUV(elemid, (FVector2f)TransformedUV);
+				}
+			}
+		};
 
 		TSet<int32> IslandElements;
 
@@ -869,23 +921,19 @@ bool FUVPacker::StandardPack(IUVMeshView* Mesh, int NumIslands, TFunctionRef<voi
 			IslandElements.Add(Triangle.C);
 		}
 
-		for (int32 elemid : IslandElements)
+		if (bPreserveScale)
 		{
-			if (elemid >= 0)
-			{
-				FVector2d UV = (FVector2d)Mesh->GetUV(elemid);
-				FVector2d TransformedUV = UV.X * Chart.PackingScaleU + UV.Y * Chart.PackingScaleV + Chart.PackingBias;
-				Mesh->SetUV(elemid, (FVector2f)TransformedUV);
-			}
+			PackUVsWithoutScaling(IslandElements);
+		}
+		else
+		{
+			PackUVsWithScaling(IslandElements);
 		}
-
 	}
 
 	return bPackingFound;
 }
 
-
-
 bool FUVPacker::StackPack(IUVMeshView* Mesh, int NumIslands, TFunctionRef<void(int, TArray<int32>&)> CopyIsland)
 {
 	TRACE_CPUPROFILER_EVENT_SCOPE(UVPacker_StackPack);
diff --git a/Engine/Source/Runtime/GeometryCore/Private/Parameterization/UVSpaceAllocator.cpp b/Engine/Source/Runtime/GeometryCore/Private/Parameterization/UVSpaceAllocator.cpp
--- a/Engine/Source/Runtime/GeometryCore/Private/Parameterization/UVSpaceAllocator.cpp
+++ b/Engine/Source/Runtime/GeometryCore/Private/Parameterization/UVSpaceAllocator.cpp
@@ -432,6 +432,12 @@ bool FUVSpaceAllocator::TestOneRun(const FRun& ThisRun, const FRun& OtherRun, ui
 
 void FUVSpaceAllocator::FlipX( const FRect& Rect )
 {
+	// If we don't actually have anything in our rasterization, then there's no point in flipping anything. Return Early.
+	if (RasterHeight == 0 && RasterWidth == 0)
+	{
+		return;
+	}
+
 	// If we have empty padding around the Rect, keep it there
 	uint32 MaxX = 0;
 	uint32 MaxY = RasterHeight - 1;
@@ -472,6 +478,12 @@ void FUVSpaceAllocator::FlipX( const FRect& Rect )
 
 void FUVSpaceAllocator::FlipY( const FRect& Rect )
 {
+	// If we don't actually have anything in our rasterization, then there's no point in flipping anything. Return Early.
+	if (RasterHeight == 0 && RasterWidth == 0)
+	{
+		return;
+	}
+
 	uint32 MinY = 0;
 	uint32 MaxY = RasterHeight - 1;
 
diff --git a/Engine/Source/Runtime/GeometryCore/Public/Parameterization/UVPacking.h b/Engine/Source/Runtime/GeometryCore/Public/Parameterization/UVPacking.h
--- a/Engine/Source/Runtime/GeometryCore/Public/Parameterization/UVPacking.h
+++ b/Engine/Source/Runtime/GeometryCore/Public/Parameterization/UVPacking.h
@@ -43,6 +43,12 @@ public:
 	/** Thickness of gutter/border in pixel dimensions. Not supported by all packing methods  */
 	float GutterSize = 1.0;
 
+	/** If true, original island scale is maintained during any packing process. */
+	bool bPreserveScale = false;
+
+	/** If true, original island rotation is maintained during any packing process. Automatically prevents bAllowFlips from applying, if set. */
+	bool bPreserveRotation = false;
+
 	/** If true, islands can be flipped in addition to rotate/translate/scale */
 	bool bAllowFlips = false;
 
