[Build Health] TDynamicMeshOverlay, TDynamicMeshVectorOverlay & TDynamicMeshAttributeChange: Add missing explicit instantiation declarations and add appropriate dllapi for them. Addresses "attribute 'dllexport' cannot be applied to member of 'dllexport' class"  error.

diff --git a/Engine/Source/Runtime/GeometryCore/Private/DynamicMesh/DynamicMeshOverlay.cpp b/Engine/Source/Runtime/GeometryCore/Private/DynamicMesh/DynamicMeshOverlay.cpp
--- a/Engine/Source/Runtime/GeometryCore/Private/DynamicMesh/DynamicMeshOverlay.cpp
+++ b/Engine/Source/Runtime/GeometryCore/Private/DynamicMesh/DynamicMeshOverlay.cpp
@@ -1859,37 +1859,29 @@ namespace UE
 namespace Geometry
 {
 
-#if PLATFORM_COMPILER_CLANG
-#define UE_EXTERN_TEMPLATE_API GEOMETRYCORE_API
-#else
-#define UE_EXTERN_TEMPLATE_API
-#endif
-
 // These are explicit instantiations of the templates that are exported from the shared lib.
 // Only these instantiations of the template can be used.
 // This is necessary because we have placed most of the templated functions in this .cpp file, instead of the header.
-template class UE_EXTERN_TEMPLATE_API TDynamicMeshOverlay<float, 1>;
-template class UE_EXTERN_TEMPLATE_API TDynamicMeshOverlay<double, 1>;
-template class UE_EXTERN_TEMPLATE_API TDynamicMeshOverlay<int, 1>;
-template class UE_EXTERN_TEMPLATE_API TDynamicMeshOverlay<float, 2>;
-template class UE_EXTERN_TEMPLATE_API TDynamicMeshOverlay<double, 2>;
-template class UE_EXTERN_TEMPLATE_API TDynamicMeshOverlay<int, 2>;
-template class UE_EXTERN_TEMPLATE_API TDynamicMeshOverlay<float, 3>;
-template class UE_EXTERN_TEMPLATE_API TDynamicMeshOverlay<double, 3>;
-template class UE_EXTERN_TEMPLATE_API TDynamicMeshOverlay<int, 3>;
-template class UE_EXTERN_TEMPLATE_API TDynamicMeshOverlay<float, 4>;
-template class UE_EXTERN_TEMPLATE_API TDynamicMeshOverlay<double, 4>;
-
-template class UE_EXTERN_TEMPLATE_API TDynamicMeshVectorOverlay<float, 2, FVector2f>;
-template class UE_EXTERN_TEMPLATE_API TDynamicMeshVectorOverlay<double, 2, FVector2d>;
-template class UE_EXTERN_TEMPLATE_API TDynamicMeshVectorOverlay<int, 2, FVector2i>;
-template class UE_EXTERN_TEMPLATE_API TDynamicMeshVectorOverlay<float, 3, FVector3f>;
-template class UE_EXTERN_TEMPLATE_API TDynamicMeshVectorOverlay<double, 3, FVector3d>;
-template class UE_EXTERN_TEMPLATE_API TDynamicMeshVectorOverlay<int, 3, FVector3i>;
-template class UE_EXTERN_TEMPLATE_API TDynamicMeshVectorOverlay<float, 4, FVector4f>;
-template class UE_EXTERN_TEMPLATE_API TDynamicMeshVectorOverlay<double, 4, FVector4d>;
-
-#undef UE_EXTERN_TEMPLATE_API
+template class GEOMETRYCORE_API TDynamicMeshOverlay<float, 1>;
+template class GEOMETRYCORE_API TDynamicMeshOverlay<double, 1>;
+template class GEOMETRYCORE_API TDynamicMeshOverlay<int, 1>;
+template class GEOMETRYCORE_API TDynamicMeshOverlay<float, 2>;
+template class GEOMETRYCORE_API TDynamicMeshOverlay<double, 2>;
+template class GEOMETRYCORE_API TDynamicMeshOverlay<int, 2>;
+template class GEOMETRYCORE_API TDynamicMeshOverlay<float, 3>;
+template class GEOMETRYCORE_API TDynamicMeshOverlay<double, 3>;
+template class GEOMETRYCORE_API TDynamicMeshOverlay<int, 3>;
+template class GEOMETRYCORE_API TDynamicMeshOverlay<float, 4>;
+template class GEOMETRYCORE_API TDynamicMeshOverlay<double, 4>;
+
+template class GEOMETRYCORE_API TDynamicMeshVectorOverlay<float, 2, FVector2f>;
+template class GEOMETRYCORE_API TDynamicMeshVectorOverlay<double, 2, FVector2d>;
+template class GEOMETRYCORE_API TDynamicMeshVectorOverlay<int, 2, FVector2i>;
+template class GEOMETRYCORE_API TDynamicMeshVectorOverlay<float, 3, FVector3f>;
+template class GEOMETRYCORE_API TDynamicMeshVectorOverlay<double, 3, FVector3d>;
+template class GEOMETRYCORE_API TDynamicMeshVectorOverlay<int, 3, FVector3i>;
+template class GEOMETRYCORE_API TDynamicMeshVectorOverlay<float, 4, FVector4f>;
+template class GEOMETRYCORE_API TDynamicMeshVectorOverlay<double, 4, FVector4d>;
 
 } // end namespace UE::Geometry
 } // end namespace UE
\ No newline at end of file
diff --git a/Engine/Source/Runtime/GeometryCore/Public/DynamicMesh/DynamicMeshChangeTracker.h b/Engine/Source/Runtime/GeometryCore/Public/DynamicMesh/DynamicMeshChangeTracker.h
--- a/Engine/Source/Runtime/GeometryCore/Public/DynamicMesh/DynamicMeshChangeTracker.h
+++ b/Engine/Source/Runtime/GeometryCore/Public/DynamicMesh/DynamicMeshChangeTracker.h
@@ -93,7 +93,25 @@ protected:
 		const TArray<FChangeTriangle>& InsertTris) const;
 };
 
-
+#if PLATFORM_COMPILER_CLANG
+#define UE_EXTERN_TEMPLATE_API GEOMETRYCORE_API
+#else
+#define UE_EXTERN_TEMPLATE_API
+#endif
+
+extern template class UE_EXTERN_TEMPLATE_API TDynamicMeshAttributeChange<float, 1>;
+extern template class UE_EXTERN_TEMPLATE_API TDynamicMeshAttributeChange<double, 1>;
+extern template class UE_EXTERN_TEMPLATE_API TDynamicMeshAttributeChange<int, 1>;
+extern template class UE_EXTERN_TEMPLATE_API TDynamicMeshAttributeChange<float, 2>;
+extern template class UE_EXTERN_TEMPLATE_API TDynamicMeshAttributeChange<double, 2>;
+extern template class UE_EXTERN_TEMPLATE_API TDynamicMeshAttributeChange<int, 2>;
+extern template class UE_EXTERN_TEMPLATE_API TDynamicMeshAttributeChange<float, 3>;
+extern template class UE_EXTERN_TEMPLATE_API TDynamicMeshAttributeChange<double, 3>;
+extern template class UE_EXTERN_TEMPLATE_API TDynamicMeshAttributeChange<int, 3>;
+extern template class UE_EXTERN_TEMPLATE_API TDynamicMeshAttributeChange<float, 4>;
+extern template class UE_EXTERN_TEMPLATE_API TDynamicMeshAttributeChange<double, 4>;
+
+#undef UE_EXTERN_TEMPLATE_API
 
 /** Standard UV overlay change type - 2-element float */
 typedef TDynamicMeshAttributeChange<float,2> FDynamicMeshUVChange;
diff --git a/Engine/Source/Runtime/GeometryCore/Public/DynamicMesh/DynamicMeshOverlay.h b/Engine/Source/Runtime/GeometryCore/Public/DynamicMesh/DynamicMeshOverlay.h
--- a/Engine/Source/Runtime/GeometryCore/Public/DynamicMesh/DynamicMeshOverlay.h
+++ b/Engine/Source/Runtime/GeometryCore/Public/DynamicMesh/DynamicMeshOverlay.h
@@ -248,7 +248,7 @@ public:
 	}
 
 	/** Discard all elements. */
-	GEOMETRYCORE_API void ClearElements();
+	void ClearElements();
 
 	/** Discard elements for given triangles. */
 	template <typename EnumerableIntType>
@@ -278,9 +278,9 @@ public:
 
 
 	/** Allocate a new element with the given constant value */
-	GEOMETRYCORE_API int AppendElement(RealType ConstantValue);
+	int AppendElement(RealType ConstantValue);
 	/** Allocate a new element with the given value */
-	GEOMETRYCORE_API int AppendElement(const RealType* Value);
+	int AppendElement(const RealType* Value);
 
 	void SetParentVertex(int ElementIndex, int ParentVertexIndex)
 	{
@@ -288,7 +288,7 @@ public:
 	}
 
 	/** Initialize the triangle list to the given size, and set all triangles to InvalidID */
-	GEOMETRYCORE_API void InitializeTriangles(int MaxTriangleID);
+	void InitializeTriangles(int MaxTriangleID);
 
 	/**
 	 * Set the triangle to the given Element index tuple, and increment element reference counts 
@@ -298,7 +298,7 @@ public:
 	 *  when remeshing across existing elements to avoid them being freed while temporarily unreferenced,
 	 *  but then it should eventually be followed by a call to FreeUnusedElements().
 	 */
-	GEOMETRYCORE_API EMeshResult SetTriangle(int TriangleID, const FIndex3i& TriElements, bool bAllowElementFreeing = true);
+	EMeshResult SetTriangle(int TriangleID, const FIndex3i& TriElements, bool bAllowElementFreeing = true);
 
 	/**
 	 * Goes through elements and frees any whose reference counts indicate that they are not being used. This
@@ -307,7 +307,7 @@ public:
 	 * 
 	 * @param ElementsToCheck If provided, only these element ID's will be checked.
 	 */
-	GEOMETRYCORE_API void FreeUnusedElements(const TSet<int>* ElementsToCheck = nullptr);
+	void FreeUnusedElements(const TSet<int>* ElementsToCheck = nullptr);
 
 	/**
 	 * Set the triangle to have InvalidID element IDs, decrementing element reference counts if needed.
@@ -316,7 +316,7 @@ public:
 	 *  become immediately unallocated. This can be set to false as part of a remeshing, but then it 
 	 *  should eventually be followed by a call to FreeUnusedElements().
 	 */
-	GEOMETRYCORE_API void UnsetTriangle(int TriangleID, bool bAllowElementFreeing = true);
+	void UnsetTriangle(int TriangleID, bool bAllowElementFreeing = true);
 
 	/** @return true if this triangle was set */
 	bool IsSetTriangle(int TID) const
@@ -336,7 +336,7 @@ public:
 	 *									  Note if a vertex can be shared between tris A and B, and B and C, it will be shared between all three
 	 * @param InitElementValue Initial element value, copied into all created elements
 	 */
-	GEOMETRYCORE_API void CreateFromPredicate(TFunctionRef<bool(int ParentVertexIdx, int TriIDA, int TriIDB)> TrisCanShareVertexPredicate, RealType InitElementValue);
+	void CreateFromPredicate(TFunctionRef<bool(int ParentVertexIdx, int TriIDA, int TriIDB)> TrisCanShareVertexPredicate, RealType InitElementValue);
 
 	/**
 	 * Refine an existing overlay topology.  For any element on a given triangle, if the predicate returns true, it gets topologically split out so it isn't shared by any other triangle.
@@ -345,7 +345,7 @@ public:
 	 * @param ShouldSplitOutVertex predicate returns true of the element should be split out and not shared w/ any other triangle
 	 * @param GetNewElementValue function to assign a new value to any element that is split out
 	 */
-	GEOMETRYCORE_API void SplitVerticesWithPredicate(TFunctionRef<bool(int ElementIdx, int TriID)> ShouldSplitOutVertex, TFunctionRef<void(int ElementIdx, int TriID, RealType* FillVect)> GetNewElementValue);
+	void SplitVerticesWithPredicate(TFunctionRef<bool(int ElementIdx, int TriID)> ShouldSplitOutVertex, TFunctionRef<void(int ElementIdx, int TriID, RealType* FillVect)> GetNewElementValue);
 
 	/**
 	 * Collapse SourceElementID into TargetElementID, resulting in connecting any containing triangles and reducing the total elements in the overlay.
@@ -354,7 +354,7 @@ public:
 	 * @param TargetElementID the element to merge into
 	 * @return If the operation completed successfully, returns true
 	 */
-	GEOMETRYCORE_API bool MergeElement(int SourceElementID, int TargetElementID);
+	bool MergeElement(int SourceElementID, int TargetElementID);
 
 	/**
 	 * Create a new copy of ElementID, and update connected triangles in the TrianglesToUpdate array to reference the copy of ElementID where they used to reference ElementID
@@ -364,7 +364,7 @@ public:
 	 * @param TrianglesToUpdate the triangles that should now reference the new element
 	 * @return the ID of the newly created element
 	 */
-	GEOMETRYCORE_API int SplitElement(int ElementID, const TArrayView<const int>& TrianglesToUpdate);
+	int SplitElement(int ElementID, const TArrayView<const int>& TrianglesToUpdate);
 
 	/**
 	* Create a new copy of ElementID, and update connected triangles in the TrianglesToUpdate array to reference the copy of ElementID where they used to reference ElementID.  The new element will have the given parent vertex ID.
@@ -375,7 +375,7 @@ public:
 	* @param TrianglesToUpdate the triangles that should now reference the new element.  Note: this is allowed to include triangles that do not have the element at all; sometimes you may want to do so to avoid creating a new array for each call.
 	* @return the ID of the newly created element
 	*/
-	GEOMETRYCORE_API int SplitElementWithNewParent(int ElementID, int SplitParentVertexID, const TArrayView<const int>& TrianglesToUpdate);
+	int SplitElementWithNewParent(int ElementID, int SplitParentVertexID, const TArrayView<const int>& TrianglesToUpdate);
 
 	/**
 	 * Split any bowties at given vertex.
@@ -383,12 +383,12 @@ public:
 	 * @param NewElementIDs If not null, newly created element IDs are placed here. Note that this array is
 	 *   intentionally not cleared before appending to it.
 	 */
-	GEOMETRYCORE_API void SplitBowtiesAtVertex(int32 Vid, TArray<int32>* NewElementIDs = nullptr);
+	void SplitBowtiesAtVertex(int32 Vid, TArray<int32>* NewElementIDs = nullptr);
 
 	/**
 	* Refine an existing overlay topology by splitting any bow ties
 	*/
-	GEOMETRYCORE_API void SplitBowties();
+	void SplitBowties();
 
 
 	//
@@ -416,7 +416,7 @@ public:
 	 * If bUnsafe, we use fast id allocation that does not update free list.
 	 * You should only be using this between BeginUnsafeElementsInsert() / EndUnsafeElementsInsert() calls
 	 */
-	GEOMETRYCORE_API EMeshResult InsertElement(int ElementID, const RealType* Value, bool bUnsafe = false);
+	EMeshResult InsertElement(int ElementID, const RealType* Value, bool bUnsafe = false);
 
 
 	//
@@ -525,39 +525,39 @@ public:
 	*   If present, bIsNonIntersectingOut will be true only if this is a seam edge 
 	*   that does not intersect with another seam or the end of the seam. 
 	*/
-	GEOMETRYCORE_API bool IsSeamEdge(int EdgeID, bool* bIsNonIntersectingOut = nullptr) const;
+	bool IsSeamEdge(int EdgeID, bool* bIsNonIntersectingOut = nullptr) const;
 	/** Returns true if the parent-mesh edge is a "Seam End" in this overlay, meaning the adjacent element triangles share one element, not two */
-	GEOMETRYCORE_API bool IsSeamEndEdge(int EdgeID) const;
+	bool IsSeamEndEdge(int EdgeID) const;
 	/** Returns true if the parent-mesh vertex is connected to any seam edges */
-	GEOMETRYCORE_API bool IsSeamVertex(int VertexID, bool bBoundaryIsSeam = true) const;
+	bool IsSeamVertex(int VertexID, bool bBoundaryIsSeam = true) const;
 
 	/**
 	 * Determines whether the base-mesh vertex has "bowtie" topology in the Overlay.
 	 * Bowtie topology means that one or more elements at the vertex are shared across disconnected UV-components.
 	 * @return true if the base-mesh vertex has "bowtie" topology in the overlay
 	 */
-	GEOMETRYCORE_API bool IsBowtieInOverlay(int32 VertexID) const;
+	bool IsBowtieInOverlay(int32 VertexID) const;
 
 	/** @return true if the two triangles are connected, ie shared edge exists and is not a seam edge */
-	GEOMETRYCORE_API bool AreTrianglesConnected(int TriangleID0, int TriangleID1) const;
+	bool AreTrianglesConnected(int TriangleID0, int TriangleID1) const;
 
 	/** find the elements associated with a given parent-mesh vertex */
-	GEOMETRYCORE_API void GetVertexElements(int VertexID, TArray<int>& OutElements) const;
+	void GetVertexElements(int VertexID, TArray<int>& OutElements) const;
 	/** Count the number of unique elements for a given parent-mesh vertex */
-	GEOMETRYCORE_API int CountVertexElements(int VertexID, bool bBruteForce = false) const;
+	int CountVertexElements(int VertexID, bool bBruteForce = false) const;
 
 	/** find the triangles connected to an element */
-	GEOMETRYCORE_API void GetElementTriangles(int ElementID, TArray<int>& OutTriangles) const;
+	void GetElementTriangles(int ElementID, TArray<int>& OutTriangles) const;
 
 	/** 
 	 * Find the element ID at a vertex of a triangle. 
 	 * 
 	 * @return Returns the element ID or FDynamicMesh3::InvalidID if the vertex is not a parent of any element contained in the triangle. 
 	 */
-	GEOMETRYCORE_API int GetElementIDAtVertex(int TriangleID, int VertexID) const;
+	int GetElementIDAtVertex(int TriangleID, int VertexID) const;
 
 	/** @return true if overlay has any interior seam edges. This requires an O(N) search unless it early-outs. */
-	GEOMETRYCORE_API bool HasInteriorSeamEdges() const;
+	bool HasInteriorSeamEdges() const;
 
 	/**
 	 * Compute interpolated parameter value inside triangle using barycentric coordinates
@@ -582,12 +582,12 @@ public:
 	/**
 	 * Checks that the overlay mesh is well-formed, ie all internal data structures are consistent
 	 */
-	GEOMETRYCORE_API bool CheckValidity(bool bAllowNonManifoldVertices = true, EValidityCheckFailMode FailMode = EValidityCheckFailMode::Check) const;
+	bool CheckValidity(bool bAllowNonManifoldVertices = true, EValidityCheckFailMode FailMode = EValidityCheckFailMode::Check) const;
 
 	/**
 	 * Returns true if this overlay is the same as Other.
 	 */
-	GEOMETRYCORE_API bool IsSameAs(const TDynamicMeshOverlay<RealType, ElementSize>& Other, bool bIgnoreDataLayout) const;
+	bool IsSameAs(const TDynamicMeshOverlay<RealType, ElementSize>& Other, bool bIgnoreDataLayout) const;
 
 	/**
 	 * Serialization operator for FDynamicMeshOverlay.
@@ -609,7 +609,7 @@ public:
 	 * @param CompactMaps If this is not a null pointer, the mesh serialization compacted the vertex and/or triangle data using the provided mapping. 
 	 * @param bUseCompression Use compression for serializing bulk data.
 	 */
-	GEOMETRYCORE_API void Serialize(FArchive& Ar, const FCompactMaps* CompactMaps, bool bUseCompression);
+	void Serialize(FArchive& Ar, const FCompactMaps* CompactMaps, bool bUseCompression);
 
 public:
 	/** Set a triangle's element indices to InvalidID */
@@ -760,12 +760,41 @@ public:
 	 * @param bFindUniqueElements if true, ProcessFunc is only called once for each ElementID, otherwise it is called once for each Triangle
 	 * @return true if at least one valid Element was found, ie if ProcessFunc was called at least one time
 	 */
-	GEOMETRYCORE_API bool EnumerateVertexElements(
+	bool EnumerateVertexElements(
 		int VertexID,
 		TFunctionRef<bool(int TriangleID, int ElementID, const VectorType& Value)> ProcessFunc,
 		bool bFindUniqueElements = true) const;
 };
 
 
+#if PLATFORM_COMPILER_CLANG
+#define UE_EXTERN_TEMPLATE_API GEOMETRYCORE_API
+#else
+#define UE_EXTERN_TEMPLATE_API
+#endif
+
+extern template class UE_EXTERN_TEMPLATE_API TDynamicMeshOverlay<float, 1>;
+extern template class UE_EXTERN_TEMPLATE_API TDynamicMeshOverlay<double, 1>;
+extern template class UE_EXTERN_TEMPLATE_API TDynamicMeshOverlay<int, 1>;
+extern template class UE_EXTERN_TEMPLATE_API TDynamicMeshOverlay<float, 2>;
+extern template class UE_EXTERN_TEMPLATE_API TDynamicMeshOverlay<double, 2>;
+extern template class UE_EXTERN_TEMPLATE_API TDynamicMeshOverlay<int, 2>;
+extern template class UE_EXTERN_TEMPLATE_API TDynamicMeshOverlay<float, 3>;
+extern template class UE_EXTERN_TEMPLATE_API TDynamicMeshOverlay<double, 3>;
+extern template class UE_EXTERN_TEMPLATE_API TDynamicMeshOverlay<int, 3>;
+extern template class UE_EXTERN_TEMPLATE_API TDynamicMeshOverlay<float, 4>;
+extern template class UE_EXTERN_TEMPLATE_API TDynamicMeshOverlay<double, 4>;
+
+extern template class UE_EXTERN_TEMPLATE_API TDynamicMeshVectorOverlay<float, 2, FVector2f>;
+extern template class UE_EXTERN_TEMPLATE_API TDynamicMeshVectorOverlay<double, 2, FVector2d>;
+extern template class UE_EXTERN_TEMPLATE_API TDynamicMeshVectorOverlay<int, 2, FVector2i>;
+extern template class UE_EXTERN_TEMPLATE_API TDynamicMeshVectorOverlay<float, 3, FVector3f>;
+extern template class UE_EXTERN_TEMPLATE_API TDynamicMeshVectorOverlay<double, 3, FVector3d>;
+extern template class UE_EXTERN_TEMPLATE_API TDynamicMeshVectorOverlay<int, 3, FVector3i>;
+extern template class UE_EXTERN_TEMPLATE_API TDynamicMeshVectorOverlay<float, 4, FVector4f>;
+extern template class UE_EXTERN_TEMPLATE_API TDynamicMeshVectorOverlay<double, 4, FVector4d>;
+
+#undef UE_EXTERN_TEMPLATE_API
+
 } // end namespace UE::Geometry
 } // end namespace UE
\ No newline at end of file
