Dynamic Mesh -  fix an algorithmic comment

diff --git a/Engine/Source/Runtime/GeometryCore/Private/DynamicMesh/DynamicMesh3_Edits.cpp b/Engine/Source/Runtime/GeometryCore/Private/DynamicMesh/DynamicMesh3_Edits.cpp
--- a/Engine/Source/Runtime/GeometryCore/Private/DynamicMesh/DynamicMesh3_Edits.cpp
+++ b/Engine/Source/Runtime/GeometryCore/Private/DynamicMesh/DynamicMesh3_Edits.cpp
@@ -1325,7 +1325,8 @@ EMeshResult FDynamicMesh3::CanCollapseEdgeInternal(int vKeep, int vRemove, doubl
 	}
 
 	// We cannot collapse if edge lists of a and b share vertices other
-	//  than c and d  (because then we will make a triangle [x b b].
+	//  than c and d  (because then we will make two edges that connect the same vertex pair, 
+	//  violating assumptions made in this mesh data structure). 
 	//  Unfortunately I cannot see a way to do this more efficiently than brute-force search
 	//  [TODO] if we had tri iterator for a, couldn't we check each tri for b  (skipping t0 and t1) ?
 	int edges_a_count = VertexEdgeLists.GetCount(a);
