make dynamic mesh IsSameAs method also report the kind of difference it found (if any) and the mesh element that was different (if applicable)
+ similarly change the equivalent geometry script method

diff --git a/Engine/Plugins/Runtime/GeometryScripting/Source/GeometryScriptingCore/Private/MeshComparisonFunctions.cpp b/Engine/Plugins/Runtime/GeometryScripting/Source/GeometryScriptingCore/Private/MeshComparisonFunctions.cpp
--- a/Engine/Plugins/Runtime/GeometryScripting/Source/GeometryScriptingCore/Private/MeshComparisonFunctions.cpp
+++ b/Engine/Plugins/Runtime/GeometryScripting/Source/GeometryScriptingCore/Private/MeshComparisonFunctions.cpp
@@ -21,6 +21,7 @@ UDynamicMesh* UGeometryScriptLibrary_MeshComparisonFunctions::IsSameMeshAs(
 	UDynamicMesh* OtherMesh,
 	FGeometryScriptIsSameMeshOptions Options,
 	bool &bIsSameMesh,
+	FGeometryScriptMeshDifferenceInfo& DifferenceInfo,
 	UGeometryScriptDebug* Debug)
 {
 	bIsSameMesh = false;
@@ -36,11 +37,13 @@ UDynamicMesh* UGeometryScriptLibrary_MeshComparisonFunctions::IsSameMeshAs(
 		return TargetMesh;
 	}
 
-	TargetMesh->ProcessMesh([&](const FDynamicMesh3& ReadMesh) 
+
+	TargetMesh->ProcessMesh([&](const FDynamicMesh3& ReadMesh)
 	{
 		OtherMesh->ProcessMesh([&](const FDynamicMesh3& OtherMesh)
 		{
 			FDynamicMesh3::FSameAsOptions CompareOptions;
+			FDynamicMesh3::FMeshDifferenceInfo Difference;
 			CompareOptions.bCheckConnectivity = Options.bCheckConnectivity;
 			CompareOptions.bCheckEdgeIDs = Options.bCheckEdgeIDs;
 			CompareOptions.bCheckNormals = Options.bCheckNormals;
@@ -49,7 +52,58 @@ UDynamicMesh* UGeometryScriptLibrary_MeshComparisonFunctions::IsSameMeshAs(
 			CompareOptions.bCheckGroups = Options.bCheckGroups;
 			CompareOptions.bCheckAttributes = Options.bCheckAttributes;
 			CompareOptions.Epsilon = Options.Epsilon;
-			bIsSameMesh = ReadMesh.IsSameAs(OtherMesh, CompareOptions);
+			bIsSameMesh = ReadMesh.IsSameAs(OtherMesh, CompareOptions, Difference);
+			auto ConvertIDType = [](FDynamicMesh3::FMeshDifferenceInfo::EIDType IDType)
+			{
+				switch (IDType)
+				{
+				case FDynamicMesh3::FMeshDifferenceInfo::EIDType::Vertex:
+					return EGeometryScriptIndexType::Vertex;
+				case FDynamicMesh3::FMeshDifferenceInfo::EIDType::Triangle:
+					return EGeometryScriptIndexType::Triangle;
+				case FDynamicMesh3::FMeshDifferenceInfo::EIDType::Edge:
+					return EGeometryScriptIndexType::Edge;
+				}
+				return EGeometryScriptIndexType::Any;
+			};
+			auto ConvertReason = [](FDynamicMesh3::FMeshDifferenceInfo::EReason Reason) -> EGeometryScriptMeshDifferenceReason
+			{
+				switch (Reason)
+				{
+				case FDynamicMesh3::FMeshDifferenceInfo::EReason::VertexCount:
+					return EGeometryScriptMeshDifferenceReason::VertexCount;
+				case FDynamicMesh3::FMeshDifferenceInfo::EReason::TriangleCount:
+					return EGeometryScriptMeshDifferenceReason::TriangleCount;
+				case FDynamicMesh3::FMeshDifferenceInfo::EReason::EdgeCount:
+					return EGeometryScriptMeshDifferenceReason::EdgeCount;
+				case FDynamicMesh3::FMeshDifferenceInfo::EReason::Vertex:
+					return EGeometryScriptMeshDifferenceReason::Vertex;
+				case FDynamicMesh3::FMeshDifferenceInfo::EReason::Triangle:
+					return EGeometryScriptMeshDifferenceReason::Triangle;
+				case FDynamicMesh3::FMeshDifferenceInfo::EReason::Edge:
+					return EGeometryScriptMeshDifferenceReason::Edge;
+				case FDynamicMesh3::FMeshDifferenceInfo::EReason::Connectivity:
+					return EGeometryScriptMeshDifferenceReason::Connectivity;
+				case FDynamicMesh3::FMeshDifferenceInfo::EReason::Normal:
+					return EGeometryScriptMeshDifferenceReason::Normal;
+				case FDynamicMesh3::FMeshDifferenceInfo::EReason::Color:
+					return EGeometryScriptMeshDifferenceReason::Color;
+				case FDynamicMesh3::FMeshDifferenceInfo::EReason::UV:
+					return EGeometryScriptMeshDifferenceReason::UV;
+				case FDynamicMesh3::FMeshDifferenceInfo::EReason::Group:
+					return EGeometryScriptMeshDifferenceReason::Group;
+				case FDynamicMesh3::FMeshDifferenceInfo::EReason::Attribute:
+					return EGeometryScriptMeshDifferenceReason::Attribute;
+				}
+				return EGeometryScriptMeshDifferenceReason::Unknown;
+			};
+			if (!bIsSameMesh)
+			{
+				DifferenceInfo.ElementIDType = ConvertIDType(Difference.IDType);
+				DifferenceInfo.Reason = ConvertReason(Difference.Reason);
+				DifferenceInfo.TargetMeshElementID = Difference.ID;
+				DifferenceInfo.OtherMeshElementID = Difference.OtherID;
+			}
 		});
 	});
 
@@ -57,6 +111,18 @@ UDynamicMesh* UGeometryScriptLibrary_MeshComparisonFunctions::IsSameMeshAs(
 }
 
 
+UDynamicMesh* UGeometryScriptLibrary_MeshComparisonFunctions::IsSameMeshAs(
+	UDynamicMesh* TargetMesh,
+	UDynamicMesh* OtherMesh,
+	FGeometryScriptIsSameMeshOptions Options,
+	bool& bIsSameMesh,
+	UGeometryScriptDebug* Debug)
+{
+	FGeometryScriptMeshDifferenceInfo DifferenceInfo_Unused;
+	return IsSameMeshAs(TargetMesh, OtherMesh, Options, bIsSameMesh, DifferenceInfo_Unused, Debug);
+
+}
+
 
 UDynamicMesh* UGeometryScriptLibrary_MeshComparisonFunctions::IsIntersectingMesh(
 	UDynamicMesh* TargetMesh,
diff --git a/Engine/Plugins/Runtime/GeometryScripting/Source/GeometryScriptingCore/Public/GeometryScript/MeshComparisonFunctions.h b/Engine/Plugins/Runtime/GeometryScripting/Source/GeometryScriptingCore/Public/GeometryScript/MeshComparisonFunctions.h
--- a/Engine/Plugins/Runtime/GeometryScripting/Source/GeometryScriptingCore/Public/GeometryScript/MeshComparisonFunctions.h
+++ b/Engine/Plugins/Runtime/GeometryScripting/Source/GeometryScriptingCore/Public/GeometryScript/MeshComparisonFunctions.h
@@ -51,7 +51,42 @@ public:
 	bool bSymmetric = true;
 };
 
+UENUM(BlueprintType)
+enum class EGeometryScriptMeshDifferenceReason : uint8
+{
+	Unknown = 0,
+	VertexCount,
+	TriangleCount,
+	EdgeCount,
+	Vertex,
+	Triangle,
+	Edge,
+	Connectivity,
+	Normal,
+	Color,
+	UV,
+	Group,
+	Attribute
+};
 
+USTRUCT(BlueprintType)
+struct GEOMETRYSCRIPTINGCORE_API FGeometryScriptMeshDifferenceInfo
+{
+	GENERATED_BODY()
+public:
+	UPROPERTY(EditAnywhere, BlueprintReadWrite, Category = Difference)
+	EGeometryScriptMeshDifferenceReason Reason = EGeometryScriptMeshDifferenceReason::Unknown;
+
+	UPROPERTY(EditAnywhere, BlueprintReadWrite, Category = Difference)
+	int32 TargetMeshElementID = INDEX_NONE;
+
+	UPROPERTY(EditAnywhere, BlueprintReadWrite, Category = Difference)
+	int32 OtherMeshElementID = INDEX_NONE;
+
+	// Indicates the type of element that TargetMeshElementID and OtherMeshElementID reference
+	UPROPERTY(EditAnywhere, BlueprintReadWrite, Category = Difference)
+	EGeometryScriptIndexType ElementIDType = EGeometryScriptIndexType::Any;
+};
 
 UCLASS(meta = (ScriptName = "GeometryScript_MeshComparison"))
 class GEOMETRYSCRIPTINGCORE_API UGeometryScriptLibrary_MeshComparisonFunctions : public UBlueprintFunctionLibrary
@@ -60,7 +95,8 @@ class GEOMETRYSCRIPTINGCORE_API UGeometryScriptLibrary_MeshComparisonFunctions :
 public:
 
 	/**
-	 * Returns true if the two input meshes are equivalent under the comparisons defined by the input options.
+	 * Returns true if the two input meshes are equivalent under the comparisons defined by the input options. If false, DifferenceInfo provides info on the first difference found.
+	 * @param DifferenceInfo If the meshes are different, provides info on the first difference found.
 	 */
 	UFUNCTION(BlueprintCallable, Category = "GeometryScript|Comparison", meta = (ScriptMethod))
 	static UPARAM(DisplayName = "Target Mesh") UDynamicMesh*
@@ -69,6 +105,16 @@ public:
 		UDynamicMesh* OtherMesh,
 		FGeometryScriptIsSameMeshOptions Options,
 		bool &bIsSameMesh,
+		FGeometryScriptMeshDifferenceInfo& DifferenceInfo,
+		UGeometryScriptDebug* Debug = nullptr);
+
+	// Non-blueprint overload of IsSameMeshAs, without DifferenceInfo, for C++ API backwards compatibilty
+	static UPARAM(DisplayName = "Target Mesh") UDynamicMesh*
+	IsSameMeshAs(
+		UDynamicMesh* TargetMesh,
+		UDynamicMesh* OtherMesh,
+		FGeometryScriptIsSameMeshOptions Options,
+		bool& bIsSameMesh,
 		UGeometryScriptDebug* Debug = nullptr);
 
 	/**
diff --git a/Engine/Source/Runtime/GeometryCore/Private/DynamicMesh/DynamicMesh3.cpp b/Engine/Source/Runtime/GeometryCore/Private/DynamicMesh/DynamicMesh3.cpp
--- a/Engine/Source/Runtime/GeometryCore/Private/DynamicMesh/DynamicMesh3.cpp
+++ b/Engine/Source/Runtime/GeometryCore/Private/DynamicMesh/DynamicMesh3.cpp
@@ -695,6 +695,20 @@ FString FDynamicMesh3::MeshInfoString() const
 
 bool FDynamicMesh3::IsSameAs(const FDynamicMesh3& m2, const FSameAsOptions& Options) const
 {
+	return IsSameAs_Helper(m2, Options, nullptr);
+}
+bool FDynamicMesh3::IsSameAs(const FDynamicMesh3& m2, const FSameAsOptions& Options, FMeshDifferenceInfo& OutMeshDifferenceInfo) const
+{
+	return IsSameAs_Helper(m2, Options, &OutMeshDifferenceInfo);
+}
+bool FDynamicMesh3::IsSameAs_Helper(const FDynamicMesh3& m2, const FSameAsOptions& Options, FMeshDifferenceInfo* OutMeshDifferenceInfo) const
+{
+	if (OutMeshDifferenceInfo)
+	{
+		// reset difference info to defaults
+		*OutMeshDifferenceInfo = FMeshDifferenceInfo{};
+	}
+
 	auto SameVertex = [this, &m2, &Options](const int Vid, const int VidM2)
 	{
 		return VectorUtil::EpsilonEqual(GetVertexRef(Vid), m2.GetVertexRef(VidM2), static_cast<double>(Options.Epsilon));
@@ -702,11 +716,19 @@ bool FDynamicMesh3::IsSameAs(const FDynamicMesh3& m2, const FSameAsOptions& Opti
 
 	if (VertexCount() != m2.VertexCount())
 	{
+		if (OutMeshDifferenceInfo)
+		{
+			OutMeshDifferenceInfo->Reason = FMeshDifferenceInfo::EReason::VertexCount;
+		}
 		return false;
 	}
 
 	if (TriangleCount() != m2.TriangleCount())
 	{
+		if (OutMeshDifferenceInfo)
+		{
+			OutMeshDifferenceInfo->Reason = FMeshDifferenceInfo::EReason::TriangleCount;
+		}
 		return false;
 	}
 
@@ -719,6 +741,11 @@ bool FDynamicMesh3::IsSameAs(const FDynamicMesh3& m2, const FSameAsOptions& Opti
 		{
 			if (m2.IsVertex(Vid) == false || !SameVertex(Vid, Vid))
 			{
+				if (OutMeshDifferenceInfo)
+				{
+					OutMeshDifferenceInfo->Reason = FMeshDifferenceInfo::EReason::Vertex;
+					OutMeshDifferenceInfo->SetVID(Vid);
+				}
 				return false;
 			}
 		}
@@ -727,6 +754,11 @@ bool FDynamicMesh3::IsSameAs(const FDynamicMesh3& m2, const FSameAsOptions& Opti
 		{
 			if (m2.IsTriangle(Tid) == false || (GetTriangle(Tid) != m2.GetTriangle(Tid)))
 			{
+				if (OutMeshDifferenceInfo)
+				{
+					OutMeshDifferenceInfo->Reason = FMeshDifferenceInfo::EReason::Triangle;
+					OutMeshDifferenceInfo->SetTID(Tid);
+				}
 				return false;
 			}
 		}
@@ -752,6 +784,11 @@ bool FDynamicMesh3::IsSameAs(const FDynamicMesh3& m2, const FSameAsOptions& Opti
 				if (!SameVertex(*ItVid, *ItVidM2))
 				{
 					// Vertices are not the same.
+					if (OutMeshDifferenceInfo)
+					{
+						OutMeshDifferenceInfo->Reason = FMeshDifferenceInfo::EReason::Vertex;
+						OutMeshDifferenceInfo->SetVID(*ItVid, *ItVidM2);
+					}
 					return false;
 				}
 			
@@ -762,6 +799,10 @@ bool FDynamicMesh3::IsSameAs(const FDynamicMesh3& m2, const FSameAsOptions& Opti
 			if (ItVid != ItEndVid || ItVidM2 != ItEndVidM2)
 			{
 				// Number of vertices is not the same.
+				if (OutMeshDifferenceInfo)
+				{
+					OutMeshDifferenceInfo->Reason = FMeshDifferenceInfo::EReason::VertexCount;
+				}
 				return false;
 			}
 		}
@@ -788,6 +829,11 @@ bool FDynamicMesh3::IsSameAs(const FDynamicMesh3& m2, const FSameAsOptions& Opti
 					if (TriM2[i] != (VidMapping ? (*VidMapping)[Tri[i]] : Tri[i]))
 					{
 						// Triangle vertices are not the same.
+						if (OutMeshDifferenceInfo)
+						{
+							OutMeshDifferenceInfo->Reason = FMeshDifferenceInfo::EReason::Triangle;
+							OutMeshDifferenceInfo->SetTID(*ItTid, *ItTidM2);
+						}
 						return false;
 					}
 				}
@@ -799,6 +845,10 @@ bool FDynamicMesh3::IsSameAs(const FDynamicMesh3& m2, const FSameAsOptions& Opti
 			if (ItTid != ItEndTid || ItTidM2 != ItEndTidM2)
 			{
 				// Number of triangles is not the same.
+				if (OutMeshDifferenceInfo)
+				{
+					OutMeshDifferenceInfo->Reason = FMeshDifferenceInfo::EReason::TriangleCount;
+				}
 				return false;
 			}
 		}
@@ -808,6 +858,10 @@ bool FDynamicMesh3::IsSameAs(const FDynamicMesh3& m2, const FSameAsOptions& Opti
 	{
 		if (EdgeCount() != m2.EdgeCount())
 		{
+			if (OutMeshDifferenceInfo)
+			{
+				OutMeshDifferenceInfo->Reason = FMeshDifferenceInfo::EReason::EdgeCount;
+			}
 			return false;
 		}
 		for (int eid : EdgeIndicesItr())
@@ -818,6 +872,11 @@ bool FDynamicMesh3::IsSameAs(const FDynamicMesh3& m2, const FSameAsOptions& Opti
 			const int other_eid = m2.FindEdge(Vert0, Vert1);
 			if (other_eid == InvalidID)
 			{
+				if (OutMeshDifferenceInfo)
+				{
+					OutMeshDifferenceInfo->Reason = FMeshDifferenceInfo::EReason::Connectivity;
+					OutMeshDifferenceInfo->SetEID(eid);
+				}
 				return false;
 			}
 			const FEdge oe = m2.GetEdge(other_eid);
@@ -830,6 +889,11 @@ bool FDynamicMesh3::IsSameAs(const FDynamicMesh3& m2, const FSameAsOptions& Opti
 			if (FMath::Min(eTri[0], eTri[1]) != FMath::Min(oe.Tri[0], oe.Tri[1]) ||
 			    FMath::Max(eTri[0], eTri[1]) != FMath::Max(oe.Tri[0], oe.Tri[1]))
 			{
+				if (OutMeshDifferenceInfo)
+				{
+					OutMeshDifferenceInfo->Reason = FMeshDifferenceInfo::EReason::Connectivity;
+					OutMeshDifferenceInfo->SetEID(eid, other_eid);
+				}
 				return false;
 			}
 		}
@@ -838,12 +902,21 @@ bool FDynamicMesh3::IsSameAs(const FDynamicMesh3& m2, const FSameAsOptions& Opti
 	{
 		if (EdgeCount() != m2.EdgeCount())
 		{
+			if (OutMeshDifferenceInfo)
+			{
+				OutMeshDifferenceInfo->Reason = FMeshDifferenceInfo::EReason::EdgeCount;
+			}
 			return false;
 		}
 		for (const int Eid : EdgeIndicesItr())
 		{
 			if (m2.IsEdge(Eid) == false)
 			{
+				if (OutMeshDifferenceInfo)
+				{
+					OutMeshDifferenceInfo->Reason = FMeshDifferenceInfo::EReason::Edge;
+					OutMeshDifferenceInfo->SetEID(Eid);
+				}
 				return false;
 			}
 			FEdge Edge = GetEdge(Eid);
@@ -854,6 +927,11 @@ bool FDynamicMesh3::IsSameAs(const FDynamicMesh3& m2, const FSameAsOptions& Opti
 			}
 			if (Edge != m2.GetEdge(Eid))
 			{
+				if (OutMeshDifferenceInfo)
+				{
+					OutMeshDifferenceInfo->Reason = FMeshDifferenceInfo::EReason::Edge;
+					OutMeshDifferenceInfo->SetEID(Eid);
+				}
 				return false;
 			}
 		}
@@ -862,6 +940,10 @@ bool FDynamicMesh3::IsSameAs(const FDynamicMesh3& m2, const FSameAsOptions& Opti
 	{
 		if (HasVertexNormals() != m2.HasVertexNormals())
 		{
+			if (OutMeshDifferenceInfo)
+			{
+				OutMeshDifferenceInfo->Reason = FMeshDifferenceInfo::EReason::Normal;
+			}
 			return false;
 		}
 		if (HasVertexNormals())
@@ -870,6 +952,11 @@ bool FDynamicMesh3::IsSameAs(const FDynamicMesh3& m2, const FSameAsOptions& Opti
 			{
 				if (VectorUtil::EpsilonEqual(GetVertexNormal(vid), m2.GetVertexNormal(VidMapping ? (*VidMapping)[vid] : vid), Options.Epsilon) == false)
 				{
+					if (OutMeshDifferenceInfo)
+					{
+						OutMeshDifferenceInfo->Reason = FMeshDifferenceInfo::EReason::Normal;
+						OutMeshDifferenceInfo->SetVID(vid);
+					}
 					return false;
 				}
 			}
@@ -879,6 +966,10 @@ bool FDynamicMesh3::IsSameAs(const FDynamicMesh3& m2, const FSameAsOptions& Opti
 	{
 		if (HasVertexColors() != m2.HasVertexColors())
 		{
+			if (OutMeshDifferenceInfo)
+			{
+				OutMeshDifferenceInfo->Reason = FMeshDifferenceInfo::EReason::Color;
+			}
 			return false;
 		}
 		if (HasVertexColors())
@@ -887,6 +978,11 @@ bool FDynamicMesh3::IsSameAs(const FDynamicMesh3& m2, const FSameAsOptions& Opti
 			{
 				if (VectorUtil::EpsilonEqual(GetVertexColor(vid), m2.GetVertexColor(VidMapping ? (*VidMapping)[vid] : vid), Options.Epsilon) == false)
 				{
+					if (OutMeshDifferenceInfo)
+					{
+						OutMeshDifferenceInfo->Reason = FMeshDifferenceInfo::EReason::Color;
+						OutMeshDifferenceInfo->SetVID(vid);
+					}
 					return false;
 				}
 			}
@@ -896,6 +992,10 @@ bool FDynamicMesh3::IsSameAs(const FDynamicMesh3& m2, const FSameAsOptions& Opti
 	{
 		if (HasVertexUVs() != m2.HasVertexUVs())
 		{
+			if (OutMeshDifferenceInfo)
+			{
+				OutMeshDifferenceInfo->Reason = FMeshDifferenceInfo::EReason::UV;
+			}
 			return false;
 		}
 		if (HasVertexUVs())
@@ -904,6 +1004,11 @@ bool FDynamicMesh3::IsSameAs(const FDynamicMesh3& m2, const FSameAsOptions& Opti
 			{
 				if (VectorUtil::EpsilonEqual(GetVertexUV(vid), m2.GetVertexUV(VidMapping ? (*VidMapping)[vid] : vid), Options.Epsilon) == false)
 				{
+					if (OutMeshDifferenceInfo)
+					{
+						OutMeshDifferenceInfo->Reason = FMeshDifferenceInfo::EReason::UV;
+						OutMeshDifferenceInfo->SetVID(vid);
+					}
 					return false;
 				}
 			}
@@ -913,6 +1018,10 @@ bool FDynamicMesh3::IsSameAs(const FDynamicMesh3& m2, const FSameAsOptions& Opti
 	{
 		if (HasTriangleGroups() != m2.HasTriangleGroups())
 		{
+			if (OutMeshDifferenceInfo)
+			{
+				OutMeshDifferenceInfo->Reason = FMeshDifferenceInfo::EReason::Group;
+			}
 			return false;
 		}
 		if (HasTriangleGroups())
@@ -922,6 +1031,11 @@ bool FDynamicMesh3::IsSameAs(const FDynamicMesh3& m2, const FSameAsOptions& Opti
 				const int TidM2 = TidMapping ? (*TidMapping)[Tid] : Tid;
 				if (GetTriangleGroup(Tid) != m2.GetTriangleGroup(TidM2))
 				{
+					if (OutMeshDifferenceInfo)
+					{
+						OutMeshDifferenceInfo->Reason = FMeshDifferenceInfo::EReason::Group;
+						OutMeshDifferenceInfo->SetTID(Tid, TidM2);
+					}
 					return false;
 				}
 			}
@@ -931,12 +1045,20 @@ bool FDynamicMesh3::IsSameAs(const FDynamicMesh3& m2, const FSameAsOptions& Opti
 	{
 		if (HasAttributes() != m2.HasAttributes())
 		{
+			if (OutMeshDifferenceInfo)
+			{
+				OutMeshDifferenceInfo->Reason = FMeshDifferenceInfo::EReason::Attribute;
+			}
 			return false;
 		}
 		if (HasAttributes())
 		{
 			if (!AttributeSet->IsSameAs(*m2.AttributeSet, Options.bIgnoreDataLayout))
 			{
+				if (OutMeshDifferenceInfo)
+				{
+					OutMeshDifferenceInfo->Reason = FMeshDifferenceInfo::EReason::Attribute;
+				}
 				return false;
 			}
 		}
diff --git a/Engine/Source/Runtime/GeometryCore/Public/DynamicMesh/DynamicMesh3.h b/Engine/Source/Runtime/GeometryCore/Public/DynamicMesh/DynamicMesh3.h
--- a/Engine/Source/Runtime/GeometryCore/Public/DynamicMesh/DynamicMesh3.h
+++ b/Engine/Source/Runtime/GeometryCore/Public/DynamicMesh/DynamicMesh3.h
@@ -1394,11 +1394,78 @@ public:
 		float Epsilon = TMathUtil<float>::Epsilon;
 	};
 
+	struct FMeshDifferenceInfo
+	{
+		// Reasons for difference between meshes
+		enum class EReason
+		{
+			Unknown,
+			VertexCount,
+			TriangleCount,
+			EdgeCount,
+			Vertex,
+			Triangle,
+			Edge,
+			Connectivity,
+			Normal,
+			Color,
+			UV,
+			Group,
+			Attribute
+		};
+		EReason Reason = EReason::Unknown;
+
+		// Types of element ID that could be set below
+		enum class EIDType
+		{
+			None,
+			Vertex,
+			Triangle,
+			Edge
+		};
+
+		// ID in this mesh of element where the difference was found, or InvalidID if not applicable
+		int32 ID = InvalidID;
+		// ID in the compared-against mesh where the difference was found, or InvalidID if not applicable
+		int32 OtherID = InvalidID;
+		// Type of element that the ID references
+		EIDType IDType = EIDType::None;
+
+		// Helpers to set different ID types
+		void SetVID(int32 VID, int32 OtherVID = InvalidID)
+		{
+			ID = VID;
+			OtherID = OtherVID;
+			IDType = EIDType::Vertex;
+		}
+		void SetTID(int32 TID, int32 OtherTID = InvalidID)
+		{
+			ID = TID;
+			OtherID = OtherTID;
+			IDType = EIDType::Triangle;
+		}
+		void SetEID(int32 EID, int32 OtherEID = InvalidID)
+		{
+			ID = EID;
+			OtherID = OtherEID;
+			IDType = EIDType::Edge;
+		}
+	};
+
 	/**
 	 * Check if another mesh is the same as this mesh. By default only checks
 	 * vertices and triangles, turn on other parameters w/ flags
 	 */
 	GEOMETRYCORE_API virtual bool IsSameAs(const FDynamicMesh3& OtherMesh, const FSameAsOptions& Options) const;
+	/**
+	 * Check if another mesh is the same as this mesh. By default only checks
+	 * vertices and triangles, turn on other parameters w/ flags
+	 * @param OutMeshDifferenceInfo If the meshes are not the same, this struct may provide additional info as to the source of difference. Note it will only indicate the first difference found, not a complete report of differences.
+	 */
+	GEOMETRYCORE_API virtual bool IsSameAs(const FDynamicMesh3& OtherMesh, const FSameAsOptions& Options, FMeshDifferenceInfo& OutMeshDifferenceInfo) const;
+private:
+	GEOMETRYCORE_API bool IsSameAs_Helper(const FDynamicMesh3& OtherMesh, const FSameAsOptions& Options, FMeshDifferenceInfo* OutMeshDifferenceInfo = nullptr) const;
+public:
 
 	/**
 	 * Options for what the validity check will permit
