PR #12231: Fix compile error C4800 in SpatialDatastructure PointHashGrid2 and PointHashGrid3

diff --git a/Engine/Source/Runtime/GeometryCore/Public/Spatial/PointHashGrid2.h b/Engine/Source/Runtime/GeometryCore/Public/Spatial/PointHashGrid2.h
--- a/Engine/Source/Runtime/GeometryCore/Public/Spatial/PointHashGrid2.h
+++ b/Engine/Source/Runtime/GeometryCore/Public/Spatial/PointHashGrid2.h
@@ -133,7 +133,7 @@ public:
 		bool bWasAtOldPos;
 		{
 			FScopeLock Lock(&CriticalSection);
-			bWasAtOldPos = Hash.RemoveSingle(old_idx, Value);
+			bWasAtOldPos = Hash.RemoveSingle(old_idx, Value) > 0;
 		}
 		check(bWasAtOldPos);
 		{
@@ -158,7 +158,7 @@ public:
 		{
 			return;
 		}
-		bool bWasAtOldPos = Hash.RemoveSingle(old_idx, Value);
+		bool bWasAtOldPos = Hash.RemoveSingle(old_idx, Value) > 0;
 		check(bWasAtOldPos);
 		Hash.Add(new_idx, Value);
 		return;
diff --git a/Engine/Source/Runtime/GeometryCore/Public/Spatial/PointHashGrid3.h b/Engine/Source/Runtime/GeometryCore/Public/Spatial/PointHashGrid3.h
--- a/Engine/Source/Runtime/GeometryCore/Public/Spatial/PointHashGrid3.h
+++ b/Engine/Source/Runtime/GeometryCore/Public/Spatial/PointHashGrid3.h
@@ -162,7 +162,7 @@ public:
 		bool bWasAtOldPos;
 		{
 			FScopeLock Lock(&CriticalSection);
-			bWasAtOldPos = Hash.RemoveSingle(old_idx, Value) == 1;
+			bWasAtOldPos = Hash.RemoveSingle(old_idx, Value) > 0;
 		}
 		check(bWasAtOldPos);
 		{
@@ -187,7 +187,7 @@ public:
 		{
 			return;
 		}
-		bool bWasAtOldPos = Hash.RemoveSingle(old_idx, Value);
+		bool bWasAtOldPos = Hash.RemoveSingle(old_idx, Value) > 0;
 		check(bWasAtOldPos);
 		Hash.Add(new_idx, Value);
 		return;
