add a details string to the mesh difference info struct, so we can report more specific information about differences when it's useful to do so (e.g., the triangle counts for both meshes, if the triangle count was what differed)

diff --git a/Engine/Plugins/Runtime/GeometryScripting/Source/GeometryScriptingCore/Private/MeshComparisonFunctions.cpp b/Engine/Plugins/Runtime/GeometryScripting/Source/GeometryScriptingCore/Private/MeshComparisonFunctions.cpp
--- a/Engine/Plugins/Runtime/GeometryScripting/Source/GeometryScriptingCore/Private/MeshComparisonFunctions.cpp
+++ b/Engine/Plugins/Runtime/GeometryScripting/Source/GeometryScriptingCore/Private/MeshComparisonFunctions.cpp
@@ -101,6 +101,7 @@ UDynamicMesh* UGeometryScriptLibrary_MeshComparisonFunctions::IsSameMeshAs(
 			{
 				DifferenceInfo.ElementIDType = ConvertIDType(Difference.IDType);
 				DifferenceInfo.Reason = ConvertReason(Difference.Reason);
+				DifferenceInfo.Detail = MoveTemp(Difference.Detail);
 				DifferenceInfo.TargetMeshElementID = Difference.ID;
 				DifferenceInfo.OtherMeshElementID = Difference.OtherID;
 			}
diff --git a/Engine/Plugins/Runtime/GeometryScripting/Source/GeometryScriptingCore/Public/GeometryScript/MeshComparisonFunctions.h b/Engine/Plugins/Runtime/GeometryScripting/Source/GeometryScriptingCore/Public/GeometryScript/MeshComparisonFunctions.h
--- a/Engine/Plugins/Runtime/GeometryScripting/Source/GeometryScriptingCore/Public/GeometryScript/MeshComparisonFunctions.h
+++ b/Engine/Plugins/Runtime/GeometryScripting/Source/GeometryScriptingCore/Public/GeometryScript/MeshComparisonFunctions.h
@@ -77,6 +77,10 @@ public:
 	UPROPERTY(EditAnywhere, BlueprintReadWrite, Category = Difference)
 	EGeometryScriptMeshDifferenceReason Reason = EGeometryScriptMeshDifferenceReason::Unknown;
 
+	// String that may contain additional detail on the difference
+	UPROPERTY(EditAnywhere, BlueprintReadWrite, Category = Difference)
+	FString Detail;
+
 	UPROPERTY(EditAnywhere, BlueprintReadWrite, Category = Difference)
 	int32 TargetMeshElementID = INDEX_NONE;
 
diff --git a/Engine/Source/Runtime/GeometryCore/Private/DynamicMesh/DynamicMesh3.cpp b/Engine/Source/Runtime/GeometryCore/Private/DynamicMesh/DynamicMesh3.cpp
--- a/Engine/Source/Runtime/GeometryCore/Private/DynamicMesh/DynamicMesh3.cpp
+++ b/Engine/Source/Runtime/GeometryCore/Private/DynamicMesh/DynamicMesh3.cpp
@@ -719,6 +719,7 @@ bool FDynamicMesh3::IsSameAs_Helper(const FDynamicMesh3& m2, const FSameAsOption
 		if (OutMeshDifferenceInfo)
 		{
 			OutMeshDifferenceInfo->Reason = FMeshDifferenceInfo::EReason::VertexCount;
+			OutMeshDifferenceInfo->Detail = FString::Printf(TEXT("Vertex Count: %d != %d"), VertexCount(), m2.VertexCount());
 		}
 		return false;
 	}
@@ -728,6 +729,7 @@ bool FDynamicMesh3::IsSameAs_Helper(const FDynamicMesh3& m2, const FSameAsOption
 		if (OutMeshDifferenceInfo)
 		{
 			OutMeshDifferenceInfo->Reason = FMeshDifferenceInfo::EReason::TriangleCount;
+			OutMeshDifferenceInfo->Detail = FString::Printf(TEXT("Triangle Count: %d != %d"), TriangleCount(), m2.TriangleCount());
 		}
 		return false;
 	}
@@ -861,6 +863,7 @@ bool FDynamicMesh3::IsSameAs_Helper(const FDynamicMesh3& m2, const FSameAsOption
 			if (OutMeshDifferenceInfo)
 			{
 				OutMeshDifferenceInfo->Reason = FMeshDifferenceInfo::EReason::EdgeCount;
+				OutMeshDifferenceInfo->Detail = FString::Printf(TEXT("Edge Count: %d != %d"), EdgeCount(), m2.EdgeCount());
 			}
 			return false;
 		}
@@ -905,6 +908,7 @@ bool FDynamicMesh3::IsSameAs_Helper(const FDynamicMesh3& m2, const FSameAsOption
 			if (OutMeshDifferenceInfo)
 			{
 				OutMeshDifferenceInfo->Reason = FMeshDifferenceInfo::EReason::EdgeCount;
+				OutMeshDifferenceInfo->Detail = FString::Printf(TEXT("Edge Count: %d != %d"), EdgeCount(), m2.EdgeCount());
 			}
 			return false;
 		}
@@ -943,6 +947,7 @@ bool FDynamicMesh3::IsSameAs_Helper(const FDynamicMesh3& m2, const FSameAsOption
 			if (OutMeshDifferenceInfo)
 			{
 				OutMeshDifferenceInfo->Reason = FMeshDifferenceInfo::EReason::Normal;
+				OutMeshDifferenceInfo->Detail = FString::Printf(TEXT("Has Vertex Normals: %d != %d"), HasVertexNormals(), m2.HasVertexNormals());
 			}
 			return false;
 		}
@@ -969,6 +974,7 @@ bool FDynamicMesh3::IsSameAs_Helper(const FDynamicMesh3& m2, const FSameAsOption
 			if (OutMeshDifferenceInfo)
 			{
 				OutMeshDifferenceInfo->Reason = FMeshDifferenceInfo::EReason::Color;
+				OutMeshDifferenceInfo->Detail = FString::Printf(TEXT("Has Vertex Colors: %d != %d"), HasVertexColors(), m2.HasVertexColors());
 			}
 			return false;
 		}
@@ -995,6 +1001,7 @@ bool FDynamicMesh3::IsSameAs_Helper(const FDynamicMesh3& m2, const FSameAsOption
 			if (OutMeshDifferenceInfo)
 			{
 				OutMeshDifferenceInfo->Reason = FMeshDifferenceInfo::EReason::UV;
+				OutMeshDifferenceInfo->Detail = FString::Printf(TEXT("Has Vertex UVs: %d != %d"), HasVertexUVs(), m2.HasVertexUVs());
 			}
 			return false;
 		}
@@ -1021,6 +1028,7 @@ bool FDynamicMesh3::IsSameAs_Helper(const FDynamicMesh3& m2, const FSameAsOption
 			if (OutMeshDifferenceInfo)
 			{
 				OutMeshDifferenceInfo->Reason = FMeshDifferenceInfo::EReason::Group;
+				OutMeshDifferenceInfo->Detail = FString::Printf(TEXT("Has Triangle Groups: %d != %d"), HasTriangleGroups(), m2.HasTriangleGroups());
 			}
 			return false;
 		}
@@ -1034,6 +1042,7 @@ bool FDynamicMesh3::IsSameAs_Helper(const FDynamicMesh3& m2, const FSameAsOption
 					if (OutMeshDifferenceInfo)
 					{
 						OutMeshDifferenceInfo->Reason = FMeshDifferenceInfo::EReason::Group;
+						OutMeshDifferenceInfo->Detail = FString::Printf(TEXT("Triangle Group: %d != %d"), GetTriangleGroup(Tid), GetTriangleGroup(TidM2));
 						OutMeshDifferenceInfo->SetTID(Tid, TidM2);
 					}
 					return false;
@@ -1048,6 +1057,7 @@ bool FDynamicMesh3::IsSameAs_Helper(const FDynamicMesh3& m2, const FSameAsOption
 			if (OutMeshDifferenceInfo)
 			{
 				OutMeshDifferenceInfo->Reason = FMeshDifferenceInfo::EReason::Attribute;
+				OutMeshDifferenceInfo->Detail = FString::Printf(TEXT("Has Attribute Set: %d != %d"), HasAttributes(), m2.HasAttributes());
 			}
 			return false;
 		}
diff --git a/Engine/Source/Runtime/GeometryCore/Public/DynamicMesh/DynamicMesh3.h b/Engine/Source/Runtime/GeometryCore/Public/DynamicMesh/DynamicMesh3.h
--- a/Engine/Source/Runtime/GeometryCore/Public/DynamicMesh/DynamicMesh3.h
+++ b/Engine/Source/Runtime/GeometryCore/Public/DynamicMesh/DynamicMesh3.h
@@ -1415,6 +1415,9 @@ public:
 		};
 		EReason Reason = EReason::Unknown;
 
+		// May contain further detail on the difference
+		FString Detail;
+
 		// Types of element ID that could be set below
 		enum class EIDType
 		{
