DynamicMesh: Allow more permissive options for edge collapse.

diff --git a/Engine/Source/Runtime/GeometryCore/Private/DynamicMesh/DynamicMesh3_Edits.cpp b/Engine/Source/Runtime/GeometryCore/Private/DynamicMesh/DynamicMesh3_Edits.cpp
--- a/Engine/Source/Runtime/GeometryCore/Private/DynamicMesh/DynamicMesh3_Edits.cpp
+++ b/Engine/Source/Runtime/GeometryCore/Private/DynamicMesh/DynamicMesh3_Edits.cpp
@@ -1282,6 +1282,12 @@ bool FDynamicMesh3::SplitVertexWouldLeaveIsolated(int VertexID, const TArrayView
 
 
 EMeshResult FDynamicMesh3::CanCollapseEdgeInternal(int vKeep, int vRemove, double collapse_t, FEdgeCollapseInfo* OutCollapseInfo) const
+{
+	return CanCollapseEdgeInternal(vKeep, vRemove, collapse_t, FEdgeCollapseOptions(),OutCollapseInfo);
+}
+
+EMeshResult UE::Geometry::FDynamicMesh3::CanCollapseEdgeInternal(int vKeep, int vRemove, double collapse_t, 
+	const FEdgeCollapseOptions& Options, FEdgeCollapseInfo* OutCollapseInfo) const
 {
 	if (IsVertex(vKeep) == false || IsVertex(vRemove) == false)
 	{
@@ -1324,11 +1330,12 @@ EMeshResult FDynamicMesh3::CanCollapseEdgeInternal(int vKeep, int vRemove, doubl
 		bIsBoundaryEdge = true;
 	}
 
-	// We cannot collapse if edge lists of a and b share vertices other
-	//  than c and d  (because then we will make two edges that connect the same vertex pair, 
-	//  violating assumptions made in this mesh data structure). 
-	//  Unfortunately I cannot see a way to do this more efficiently than brute-force search
-	//  [TODO] if we had tri iterator for a, couldn't we check each tri for b  (skipping t0 and t1) ?
+	// We cannot collapse if there is some other vertex x that is connected to both a and b,
+	//  and either xa or xb is an interior edge. In other words, if there are more than two
+	//  triangles that have edge xa or xb, then after collapsing a to b, we'll end up with more
+	//  than two triangles trying to share edge xb, which is disallowed.
+	// Additionally, depending on options, we might not even allow such a collapse even if the
+	//  both edges are boundary edges.
 	int edges_a_count = VertexEdgeLists.GetCount(a);
 	int eac = InvalidID, ead = InvalidID, ebc = InvalidID, ebd = InvalidID;
 	for (int eid_a : VertexEdgeLists.Values(a))
@@ -1352,18 +1359,19 @@ EMeshResult FDynamicMesh3::CanCollapseEdgeInternal(int vKeep, int vRemove, doubl
 		{
 			if (GetOtherEdgeVertex(eid_b, b) == vax)
 			{
-				return EMeshResult::Failed_InvalidNeighbourhood;
+				if (!Options.bAllowHoleCollapse || !IsBoundaryEdge(eid_b) || !IsBoundaryEdge(eid_a))
+				{
+					return EMeshResult::Failed_InvalidNeighbourhood;
+				}
+				break;
 			}
 		}
 	}
 
-	// I am not sure this tetrahedron case will detect bowtie vertices.
-	// But the single-triangle case does
-
-	// We cannot collapse if we have a tetrahedron. In this case a has 3 nbr edges,
+	// We may not allow collapse if we have a tetrahedron. In this case a has 3 nbr edges,
 	//  and edge cd exists. But that is not conclusive, we also have to check that
 	//  cd is an internal edge, and that each of its tris contain a or b
-	if (edges_a_count == 3 && bIsBoundaryEdge == false)
+	if (!Options.bAllowTetrahedronCollapse && edges_a_count == 3 && bIsBoundaryEdge == false)
 	{
 		int edc = FindEdge(d, c);
 		if (edc != InvalidID)
@@ -1392,14 +1400,13 @@ EMeshResult FDynamicMesh3::CanCollapseEdgeInternal(int vKeep, int vRemove, doubl
 		}
 	}
 
-	// TODO: it's unclear how this case would ever be encountered; check if it is needed
-	//
-	// cannot collapse an edge where both vertices are boundary vertices
-	// because that would create a bowtie
+	// We might not allow collapsing an edge where both vertices are boundary vertices
+	//  because that would sometimes create a bowtie
 	//
 	// NOTE: potentially scanning all edges here...couldn't we
 	//  pick up eac/bc/ad/bd as we go? somehow?
-	if (bIsBoundaryEdge == false && IsBoundaryVertex(a) && IsBoundaryVertex(b))
+	if (!Options.bAllowCollapsingInternalEdgeWithExternalVertices && !bIsBoundaryEdge
+		&& IsBoundaryVertex(a) && IsBoundaryVertex(b))
 	{
 		return EMeshResult::Failed_InvalidNeighbourhood;
 	}
@@ -1423,15 +1430,24 @@ EMeshResult FDynamicMesh3::CanCollapseEdgeInternal(int vKeep, int vRemove, doubl
 
 EMeshResult FDynamicMesh3::CanCollapseEdge(int vKeep, int vRemove, double collapse_t) const
 {
-	return CanCollapseEdgeInternal(vKeep, vRemove, collapse_t, nullptr);
+	return CanCollapseEdgeInternal(vKeep, vRemove, 0, FEdgeCollapseOptions(), nullptr);
 }
 
+EMeshResult FDynamicMesh3::CanCollapseEdge(int vKeep, int vRemove, const FEdgeCollapseOptions& Options) const
+{
+	return CanCollapseEdgeInternal(vKeep, vRemove, 0, Options, nullptr);
+}
 
 EMeshResult FDynamicMesh3::CollapseEdge(int vKeep, int vRemove, double collapse_t, FEdgeCollapseInfo& CollapseInfo)
+{
+	return CollapseEdge(vKeep, vRemove, collapse_t, FEdgeCollapseOptions(), CollapseInfo);
+}
+
+EMeshResult FDynamicMesh3::CollapseEdge(int vKeep, int vRemove, double collapse_t, const FEdgeCollapseOptions& Options, FEdgeCollapseInfo& CollapseInfo)
 {
 	CollapseInfo = FEdgeCollapseInfo();
 
-	const EMeshResult CanCollapseResult = CanCollapseEdgeInternal(vKeep, vRemove, collapse_t, &CollapseInfo);
+	const EMeshResult CanCollapseResult = CanCollapseEdgeInternal(vKeep, vRemove, collapse_t, Options, &CollapseInfo);
 	if (CanCollapseResult != EMeshResult::Ok)
 	{
 		return CanCollapseResult;
diff --git a/Engine/Source/Runtime/GeometryCore/Public/DynamicMesh/DynamicMesh3.h b/Engine/Source/Runtime/GeometryCore/Public/DynamicMesh/DynamicMesh3.h
--- a/Engine/Source/Runtime/GeometryCore/Public/DynamicMesh/DynamicMesh3.h
+++ b/Engine/Source/Runtime/GeometryCore/Public/DynamicMesh/DynamicMesh3.h
@@ -1309,8 +1309,34 @@ public:
 	 */
 	GEOMETRYCORE_API virtual bool SplitVertexWouldLeaveIsolated(int VertexID, const TArrayView<const int>& TrianglesToUpdate);
 
+	struct FEdgeCollapseOptions
+	{
+		/**
+		 * When false, collapse is disallowed if the edge is the boundary of a single triangle hole,
+		 *  such that collapsing it would clse the hole. I.e. collapse is dissallowed if there
+		 *  is some vertex that connects vKeep and vRemove that is not part of the triangle(s) being
+		 *  collapsed (note that if such a vertex is connected by a non-boundary edge, collapse will
+		 *  always be disallowed regardless of bAllowHoleCollapse, as it would create non-manifold geometry).
+		 */ 
+		bool bAllowHoleCollapse = false;
+		/**
+		 * When false, collapse is disallowed if the edge is an interior edge, yet both vertices are
+		 *  connected to boundary edges. In some circumstances this could create a bowtie. In other
+		 *  cases, it could disconnected parts of a mesh that were connected by a bowtie.
+		 */
+		bool bAllowCollapsingInternalEdgeWithExternalVertices = false;
+		/**
+		 * When false, collapse is disallowed if we are collapsing the side of a tetrahedron. Note
+		 *  that a base edge of an open-base tetrahedron could be collapsed even if bAllowTetrahedronCollapse
+		 *  is false if bAllowBorderVerticesWhenInternal is true.
+		 */
+		bool bAllowTetrahedronCollapse = false;
+	};
+	GEOMETRYCORE_API virtual EMeshResult CanCollapseEdge(int vKeep, int vRemove, const FEdgeCollapseOptions& Options) const;
+
 	/**
-	 * Tests whether collapsing the specified edge using the CollapseEdge function would succeed
+	 * Tests whether collapsing the specified edge using the CollapseEdge function would succeed.
+	 *  Equivalent to calling the options overload with default options.
 	 * @param KeepVertID index of the vertex that should be kept
 	 * @param RemoveVertID index of the vertex that should be removed
 	 * @param EdgeParameterT vKeep is moved to Lerp(KeepPos, RemovePos, EdgeParameterT). Note: Does not currently affect whether the edge is collapsable.
@@ -1323,11 +1349,37 @@ public:
 	 * @param KeepVertID index of the vertex that should be kept
 	 * @param RemoveVertID index of the vertex that should be removed
 	 * @param EdgeParameterT vKeep is moved to Lerp(KeepPos, RemovePos, EdgeParameterT)
+	 * @param Options Sets options for the collapse
+	 * @param CollapseInfo returned information about new and modified mesh elements
+	 * @return Ok on success, or enum value indicates why operation cannot be applied. Mesh remains unmodified on error.
+	 */
+	GEOMETRYCORE_API virtual EMeshResult CollapseEdge(int KeepVertID, int RemoveVertID, double EdgeParameterT, 
+		const FEdgeCollapseOptions& Options, FEdgeCollapseInfo& CollapseInfo);
+
+	/**
+	 * Collapse the edge between the two vertices, if topologically possible. Equivalent to
+	 *  using the other overload with 0 for EdgeParameterT.
+	 */
+	virtual EMeshResult CollapseEdge(int KeepVertID, int RemoveVertID, const FEdgeCollapseOptions& Options, FEdgeCollapseInfo& CollapseInfo)
+	{
+		return CollapseEdge(KeepVertID, RemoveVertID, 0, Options, CollapseInfo);
+	}
+
+	/**
+	 * Collapse the edge between the two vertices, if topologically possible. Equivalent to calling
+	 *  the options overload with default options.
+	 * @param KeepVertID index of the vertex that should be kept
+	 * @param RemoveVertID index of the vertex that should be removed
+	 * @param EdgeParameterT vKeep is moved to Lerp(KeepPos, RemovePos, EdgeParameterT)
 	 * @param CollapseInfo returned information about new and modified mesh elements
 	 * @return Ok on success, or enum value indicates why operation cannot be applied. Mesh remains unmodified on error.
 	 */
 	GEOMETRYCORE_API virtual EMeshResult CollapseEdge(int KeepVertID, int RemoveVertID, double EdgeParameterT,
 	                                 FEdgeCollapseInfo& CollapseInfo);
+	/**
+	 * Collapse the edge between the two vertices, if topologically possible. Equivalent to calling
+	 *  the options overload with default options and using 0 for EdgeParameterT.
+	 */
 	virtual EMeshResult CollapseEdge(int KeepVertID, int RemoveVertID, FEdgeCollapseInfo& CollapseInfo)
 	{
 		return CollapseEdge(KeepVertID, RemoveVertID, 0, CollapseInfo);
@@ -1669,7 +1721,9 @@ protected:
 
 	/* We keep this version of CanCollapseEdge internal because the CollapseInfo struct may only be partially filled out by the function */
 	virtual EMeshResult CanCollapseEdgeInternal(int vKeep, int vRemove, double collapse_t, FEdgeCollapseInfo* OutCollapseInfo) const;
-
+private:
+	virtual EMeshResult CanCollapseEdgeInternal(int vKeep, int vRemove, double collapse_t,
+		const FEdgeCollapseOptions& Options, FEdgeCollapseInfo* OutCollapseInfo) const;
 };
 
 
