DynamicMesh: Fix static analysis warning during DynamicVertexSkinWeightsAttributes serialization

diff --git a/Engine/Source/Runtime/GeometryCore/Public/DynamicMesh/DynamicVertexSkinWeightsAttribute.h b/Engine/Source/Runtime/GeometryCore/Public/DynamicMesh/DynamicVertexSkinWeightsAttribute.h
--- a/Engine/Source/Runtime/GeometryCore/Public/DynamicMesh/DynamicVertexSkinWeightsAttribute.h
+++ b/Engine/Source/Runtime/GeometryCore/Public/DynamicMesh/DynamicVertexSkinWeightsAttribute.h
@@ -551,7 +551,7 @@ public:
 						BoneWeights.SetBoneWeight(reinterpret_cast<AnimationCore::FBoneWeight&>(*BufferPtr++), BoneWeightsSettings);
 					}
 				}
-				checkSlow(BufferPtr == &*Buffer.end());
+				checkSlow(BufferPtr == Buffer.GetData() + Buffer.Num());
 			}
 			else
 			{
@@ -622,7 +622,7 @@ public:
 						}
 					}
 				}
-				checkSlow(BufferPtr == &*Buffer.end());
+				checkSlow(BufferPtr == Buffer.GetData() + Buffer.Num());
 
 				// Compress buffer to archive.
 				Ar.SerializeCompressedNew(Buffer.GetData(), Buffer.Num() * sizeof(int32), NAME_Oodle, NAME_Oodle, COMPRESS_NoFlags, false, nullptr);
