Removed dllexported template instantiations when building with merged modules. Reason is because some of the console platform compilers don't handle this properly.

diff --git a/Engine/Source/Runtime/Core/Private/Containers/ContainerAllocationPolicies.cpp b/Engine/Source/Runtime/Core/Private/Containers/ContainerAllocationPolicies.cpp
--- a/Engine/Source/Runtime/Core/Private/Containers/ContainerAllocationPolicies.cpp
+++ b/Engine/Source/Runtime/Core/Private/Containers/ContainerAllocationPolicies.cpp
@@ -15,7 +15,9 @@ FORCENOINLINE void UE::Core::Private::OnInvalidSizedHeapAllocatorNum(int32 Index
 	for (;;);
 }
 
+#if !IS_MERGEDMODULES
 template void TSizedHeapAllocator<32, FMemory>::ForAnyElementType::ResizeAllocation(SizeType PreviousNumElements, SizeType NumElements, SIZE_T NumBytesPerElement);
 template void TSizedHeapAllocator<32, FMemory>::ForAnyElementType::ResizeAllocation(SizeType PreviousNumElements, SizeType NumElements, SIZE_T NumBytesPerElement, uint32 AlignmentOfElement);
 template void TSizedHeapAllocator<64, FMemory>::ForAnyElementType::ResizeAllocation(SizeType PreviousNumElements, SizeType NumElements, SIZE_T NumBytesPerElement);
 template void TSizedHeapAllocator<64, FMemory>::ForAnyElementType::ResizeAllocation(SizeType PreviousNumElements, SizeType NumElements, SIZE_T NumBytesPerElement, uint32 AlignmentOfElement);
+#endif
\ No newline at end of file
diff --git a/Engine/Source/Runtime/Core/Public/Containers/ContainerAllocationPolicies.h b/Engine/Source/Runtime/Core/Public/Containers/ContainerAllocationPolicies.h
--- a/Engine/Source/Runtime/Core/Public/Containers/ContainerAllocationPolicies.h
+++ b/Engine/Source/Runtime/Core/Public/Containers/ContainerAllocationPolicies.h
@@ -826,10 +826,12 @@ public:
 };
 
 // Define the ResizeAllocation functions with the regular allocator as exported to avoid bloat
+#if !IS_MERGEDMODULES
 extern template CORE_API FORCENOINLINE void TSizedHeapAllocator<32, FMemory>::ForAnyElementType::ResizeAllocation(SizeType PreviousNumElements, SizeType NumElements, SIZE_T NumBytesPerElement);
 extern template CORE_API FORCENOINLINE void TSizedHeapAllocator<32, FMemory>::ForAnyElementType::ResizeAllocation(SizeType PreviousNumElements, SizeType NumElements, SIZE_T NumBytesPerElement, uint32 AlignmentOfElement);
 extern template CORE_API FORCENOINLINE void TSizedHeapAllocator<64, FMemory>::ForAnyElementType::ResizeAllocation(SizeType PreviousNumElements, SizeType NumElements, SIZE_T NumBytesPerElement);
 extern template CORE_API FORCENOINLINE void TSizedHeapAllocator<64, FMemory>::ForAnyElementType::ResizeAllocation(SizeType PreviousNumElements, SizeType NumElements, SIZE_T NumBytesPerElement, uint32 AlignmentOfElement);
+#endif
 
 template <uint8 IndexSize>
 struct TAllocatorTraits<TSizedHeapAllocator<IndexSize>> : TAllocatorTraitsBase<TSizedHeapAllocator<IndexSize>>
diff --git a/Engine/Source/Runtime/Experimental/Chaos/Private/Chaos/BoundingVolume.cpp b/Engine/Source/Runtime/Experimental/Chaos/Private/Chaos/BoundingVolume.cpp
--- a/Engine/Source/Runtime/Experimental/Chaos/Private/Chaos/BoundingVolume.cpp
+++ b/Engine/Source/Runtime/Experimental/Chaos/Private/Chaos/BoundingVolume.cpp
@@ -26,6 +26,8 @@ namespace Chaos
 	    MaxDirtyElements,
 	    TEXT("The max number of dirty elements. This forces a flush which is very expensive"));
 
+#if !IS_MERGEDMODULES
 	template class CHAOS_API Chaos::TBoundingVolume<int32, FReal, 3>;
 	template class CHAOS_API Chaos::TBoundingVolume<FAccelerationStructureHandle, FReal, 3>;
+#endif
 }
diff --git a/Engine/Source/Runtime/Experimental/Chaos/Private/Chaos/SpatialAcceleration.cpp b/Engine/Source/Runtime/Experimental/Chaos/Private/Chaos/SpatialAcceleration.cpp
--- a/Engine/Source/Runtime/Experimental/Chaos/Private/Chaos/SpatialAcceleration.cpp
+++ b/Engine/Source/Runtime/Experimental/Chaos/Private/Chaos/SpatialAcceleration.cpp
@@ -29,12 +29,14 @@ namespace Chaos
 		}
 	}
 
+	#if !IS_MERGEDMODULES
 	template class CHAOS_API Chaos::ISpatialAcceleration<int32, FReal, 3>;
 	template class CHAOS_API Chaos::ISpatialAcceleration<FAccelerationStructureHandle, FReal, 3>;
-	template class CHAOS_API Chaos::ISpatialAcceleration<FAccelerationStructureHandle, FReal, 3>* TSpatialAccelerationSerializationFactory<ESpatialAcceleration::BoundingVolume, FAccelerationStructureHandle, FReal, 3>::Create();
-	template class CHAOS_API Chaos::ISpatialAcceleration<FAccelerationStructureHandle, FReal, 3>* TSpatialAccelerationSerializationFactory<ESpatialAcceleration::AABBTree, FAccelerationStructureHandle, FReal, 3>::Create();
-	template class CHAOS_API Chaos::ISpatialAcceleration<FAccelerationStructureHandle, FReal, 3>* TSpatialAccelerationSerializationFactory<ESpatialAcceleration::AABBTreeBV, FAccelerationStructureHandle, FReal, 3>::Create();
-
 	template class CHAOS_API Chaos::ISpatialVisitor<int32, FReal>;
 	template class CHAOS_API Chaos::ISpatialVisitor<FAccelerationStructureHandle, FReal>;
+	#endif
+
+	template class Chaos::ISpatialAcceleration<FAccelerationStructureHandle, FReal, 3>* TSpatialAccelerationSerializationFactory<ESpatialAcceleration::BoundingVolume, FAccelerationStructureHandle, FReal, 3>::Create();
+	template class Chaos::ISpatialAcceleration<FAccelerationStructureHandle, FReal, 3>* TSpatialAccelerationSerializationFactory<ESpatialAcceleration::AABBTree, FAccelerationStructureHandle, FReal, 3>::Create();
+	template class Chaos::ISpatialAcceleration<FAccelerationStructureHandle, FReal, 3>* TSpatialAccelerationSerializationFactory<ESpatialAcceleration::AABBTreeBV, FAccelerationStructureHandle, FReal, 3>::Create();
 }
\ No newline at end of file
diff --git a/Engine/Source/Runtime/Experimental/Chaos/Public/Chaos/BoundingVolume.h b/Engine/Source/Runtime/Experimental/Chaos/Public/Chaos/BoundingVolume.h
--- a/Engine/Source/Runtime/Experimental/Chaos/Public/Chaos/BoundingVolume.h
+++ b/Engine/Source/Runtime/Experimental/Chaos/Public/Chaos/BoundingVolume.h
@@ -1376,6 +1376,7 @@ FArchive& operator<<(FChaosArchive& Ar, TBoundingVolume<TPayloadType, T, d>& Bou
 	return Ar;
 }
 
+#if !IS_MERGEDMODULES
 #if PLATFORM_COMPILER_CLANG
 extern template class CHAOS_API Chaos::TBoundingVolume<int32, Chaos::FReal, 3>;
 extern template class CHAOS_API Chaos::TBoundingVolume<Chaos::FAccelerationStructureHandle, Chaos::FReal, 3>;
@@ -1383,5 +1384,6 @@ extern template class CHAOS_API Chaos::TBoundingVolume<Chaos::FAccelerationStruc
 extern template class TBoundingVolume<int32, FReal, 3>;
 extern template class TBoundingVolume<class FAccelerationStructureHandle, FReal, 3>;
 #endif
+#endif
 
 }
diff --git a/Engine/Source/Runtime/Experimental/Chaos/Public/Chaos/ISpatialAcceleration.h b/Engine/Source/Runtime/Experimental/Chaos/Public/Chaos/ISpatialAcceleration.h
--- a/Engine/Source/Runtime/Experimental/Chaos/Public/Chaos/ISpatialAcceleration.h
+++ b/Engine/Source/Runtime/Experimental/Chaos/Public/Chaos/ISpatialAcceleration.h
@@ -259,7 +259,7 @@ class ISpatialAcceleration;
 template<ESpatialAcceleration SpatialType, typename TPayloadType, typename T, int d>
 struct TSpatialAccelerationSerializationFactory
 {
-	static ISpatialAcceleration<TPayloadType, T, d>* Create();
+	static CHAOS_API ISpatialAcceleration<TPayloadType, T, d>* Create();
 };
 
 template <typename TPayloadType, typename T, int d>
@@ -1020,6 +1020,7 @@ FORCEINLINE bool PrePreFilterHelper(const int32 Payload, const TVisitor& Visitor
 	return false;
 }
 
+#if !IS_MERGEDMODULES
 #if PLATFORM_COMPILER_CLANG
 extern template class CHAOS_API ISpatialAcceleration<int32, FReal, 3>;
 extern template class CHAOS_API ISpatialVisitor<int32, FReal>;
@@ -1027,5 +1028,6 @@ extern template class CHAOS_API ISpatialVisitor<int32, FReal>;
 extern template class ISpatialAcceleration<int32, FReal, 3>;
 extern template class ISpatialVisitor<int32, FReal>;
 #endif
+#endif
 
 }
diff --git a/Engine/Source/Runtime/Experimental/Chaos/Public/Chaos/ParticleHandle.h b/Engine/Source/Runtime/Experimental/Chaos/Public/Chaos/ParticleHandle.h
--- a/Engine/Source/Runtime/Experimental/Chaos/Public/Chaos/ParticleHandle.h
+++ b/Engine/Source/Runtime/Experimental/Chaos/Public/Chaos/ParticleHandle.h
@@ -2694,7 +2694,7 @@ public:
 		return Geometry && Geometry->IsValidGeometry();	//todo: if we want support for sample particles without geometry we need to adjust this
 	}
 
-	static TGeometryParticle<T, d>* SerializationFactory(FChaosArchive& Ar, TGeometryParticle<T, d>* Serializable);
+	static inline TGeometryParticle<T, d>* SerializationFactory(FChaosArchive& Ar, TGeometryParticle<T, d>* Serializable);
 
 	const TVector<T, d>& X() const { return MXR.Read().X(); }
 	const TVector<T, d>& GetX() const { return MXR.Read().X(); }
@@ -3923,6 +3923,7 @@ inline void SetObjectStateHelper(IPhysicsProxyBase& Proxy, FPBDRigidParticle& Ri
 
 CHAOS_API void SetObjectStateHelper(IPhysicsProxyBase& Proxy, FPBDRigidParticleHandle& Rigid, EObjectStateType InState, bool bAllowEvents = false, bool bInvalidate = true);
 
+#if !IS_MERGEDMODULES
 #if PLATFORM_COMPILER_CLANG
 extern template class CHAOS_API ISpatialAcceleration<FAccelerationStructureHandle, FReal, 3>;
 extern template class CHAOS_API ISpatialVisitor<FAccelerationStructureHandle, FReal>;
@@ -3930,7 +3931,7 @@ extern template class CHAOS_API ISpatialVisitor<FAccelerationStructureHandle, FR
 extern template class ISpatialAcceleration<FAccelerationStructureHandle, FReal, 3>;
 extern template class ISpatialVisitor<FAccelerationStructureHandle, FReal>;
 #endif
-
+#endif
 
 } // namespace Chaos
 
diff --git a/Engine/Source/Runtime/GeometryCore/Private/DynamicMesh/DynamicMeshChangeTracker.cpp b/Engine/Source/Runtime/GeometryCore/Private/DynamicMesh/DynamicMeshChangeTracker.cpp
--- a/Engine/Source/Runtime/GeometryCore/Private/DynamicMesh/DynamicMeshChangeTracker.cpp
+++ b/Engine/Source/Runtime/GeometryCore/Private/DynamicMesh/DynamicMeshChangeTracker.cpp
@@ -144,17 +144,17 @@ namespace Geometry
 {
 
 // because ::ApplyReplaceChange() is in the .cpp we need to declare possible attribute types here
-template class GEOMETRYCORE_API TDynamicMeshAttributeChange<float, 1>;
-template class GEOMETRYCORE_API TDynamicMeshAttributeChange<double, 1>;
-template class GEOMETRYCORE_API TDynamicMeshAttributeChange<int, 1>;
-template class GEOMETRYCORE_API TDynamicMeshAttributeChange<float, 2>;
-template class GEOMETRYCORE_API TDynamicMeshAttributeChange<double, 2>;
-template class GEOMETRYCORE_API TDynamicMeshAttributeChange<int, 2>;
-template class GEOMETRYCORE_API TDynamicMeshAttributeChange<float, 3>;
-template class GEOMETRYCORE_API TDynamicMeshAttributeChange<double, 3>;
-template class GEOMETRYCORE_API TDynamicMeshAttributeChange<int, 3>;
-template class GEOMETRYCORE_API TDynamicMeshAttributeChange<float, 4>;
-template class GEOMETRYCORE_API TDynamicMeshAttributeChange<double, 4>;
+template class TDynamicMeshAttributeChange<float, 1>;
+template class TDynamicMeshAttributeChange<double, 1>;
+template class TDynamicMeshAttributeChange<int, 1>;
+template class TDynamicMeshAttributeChange<float, 2>;
+template class TDynamicMeshAttributeChange<double, 2>;
+template class TDynamicMeshAttributeChange<int, 2>;
+template class TDynamicMeshAttributeChange<float, 3>;
+template class TDynamicMeshAttributeChange<double, 3>;
+template class TDynamicMeshAttributeChange<int, 3>;
+template class TDynamicMeshAttributeChange<float, 4>;
+template class TDynamicMeshAttributeChange<double, 4>;
 
 } // end namespace UE::Geometry
 } // end namespace UE
diff --git a/Engine/Source/Runtime/GeometryCore/Public/DynamicMesh/DynamicMeshChangeTracker.h b/Engine/Source/Runtime/GeometryCore/Public/DynamicMesh/DynamicMeshChangeTracker.h
--- a/Engine/Source/Runtime/GeometryCore/Public/DynamicMesh/DynamicMeshChangeTracker.h
+++ b/Engine/Source/Runtime/GeometryCore/Public/DynamicMesh/DynamicMeshChangeTracker.h
@@ -56,13 +56,13 @@ template<typename RealType, int ElementSize>
 class TDynamicMeshAttributeChange
 {
 public:
-	void SaveInitialElement(const TDynamicMeshOverlay<RealType,ElementSize>* Overlay, int ElementID);
-	void SaveInitialTriangle(const TDynamicMeshOverlay<RealType,ElementSize>* Overlay, int TriangleID);
+	GEOMETRYCORE_API void SaveInitialElement(const TDynamicMeshOverlay<RealType,ElementSize>* Overlay, int ElementID);
+	GEOMETRYCORE_API void SaveInitialTriangle(const TDynamicMeshOverlay<RealType,ElementSize>* Overlay, int TriangleID);
 
-	void StoreFinalElement(const TDynamicMeshOverlay<RealType, ElementSize>* Overlay, int ElementID);
-	void StoreFinalTriangle(const TDynamicMeshOverlay<RealType, ElementSize>* Overlay, int TriangleID);
+	GEOMETRYCORE_API void StoreFinalElement(const TDynamicMeshOverlay<RealType, ElementSize>* Overlay, int ElementID);
+	GEOMETRYCORE_API void StoreFinalTriangle(const TDynamicMeshOverlay<RealType, ElementSize>* Overlay, int TriangleID);
 
-	bool Apply(TDynamicMeshOverlay<RealType, ElementSize>* Overlay, bool bRevert) const;
+	GEOMETRYCORE_API bool Apply(TDynamicMeshOverlay<RealType, ElementSize>* Overlay, bool bRevert) const;
 
 protected:
 	struct FChangeElement
@@ -86,33 +86,13 @@ protected:
 	TArray<RealType> NewElementData;
 	TArray<FChangeTriangle> NewTriangles;
 
-	void ApplyReplaceChange(TDynamicMeshOverlay<RealType,ElementSize>* Overlay,
+	GEOMETRYCORE_API void ApplyReplaceChange(TDynamicMeshOverlay<RealType,ElementSize>* Overlay,
 		const TArray<FChangeTriangle>& RemoveTris,
 		const TArray<FChangeElement>& InsertElements,
 		const TArray<RealType>& InsertElementData,
 		const TArray<FChangeTriangle>& InsertTris) const;
 };
 
-#if PLATFORM_COMPILER_CLANG
-#define UE_EXTERN_TEMPLATE_API GEOMETRYCORE_API
-#else
-#define UE_EXTERN_TEMPLATE_API
-#endif
-
-extern template class UE_EXTERN_TEMPLATE_API TDynamicMeshAttributeChange<float, 1>;
-extern template class UE_EXTERN_TEMPLATE_API TDynamicMeshAttributeChange<double, 1>;
-extern template class UE_EXTERN_TEMPLATE_API TDynamicMeshAttributeChange<int, 1>;
-extern template class UE_EXTERN_TEMPLATE_API TDynamicMeshAttributeChange<float, 2>;
-extern template class UE_EXTERN_TEMPLATE_API TDynamicMeshAttributeChange<double, 2>;
-extern template class UE_EXTERN_TEMPLATE_API TDynamicMeshAttributeChange<int, 2>;
-extern template class UE_EXTERN_TEMPLATE_API TDynamicMeshAttributeChange<float, 3>;
-extern template class UE_EXTERN_TEMPLATE_API TDynamicMeshAttributeChange<double, 3>;
-extern template class UE_EXTERN_TEMPLATE_API TDynamicMeshAttributeChange<int, 3>;
-extern template class UE_EXTERN_TEMPLATE_API TDynamicMeshAttributeChange<float, 4>;
-extern template class UE_EXTERN_TEMPLATE_API TDynamicMeshAttributeChange<double, 4>;
-
-#undef UE_EXTERN_TEMPLATE_API
-
 /** Standard UV overlay change type - 2-element float */
 typedef TDynamicMeshAttributeChange<float,2> FDynamicMeshUVChange;
 /** Standard Normal overlay change type - 3-element float */
diff --git a/Engine/Source/Runtime/GeometryCore/Public/DynamicMesh/DynamicMeshOverlay.h b/Engine/Source/Runtime/GeometryCore/Public/DynamicMesh/DynamicMeshOverlay.h
--- a/Engine/Source/Runtime/GeometryCore/Public/DynamicMesh/DynamicMeshOverlay.h
+++ b/Engine/Source/Runtime/GeometryCore/Public/DynamicMesh/DynamicMeshOverlay.h
@@ -773,6 +773,7 @@ public:
 #define UE_EXTERN_TEMPLATE_API
 #endif
 
+#if !IS_MERGEDMODULES
 extern template class UE_EXTERN_TEMPLATE_API TDynamicMeshOverlay<float, 1>;
 extern template class UE_EXTERN_TEMPLATE_API TDynamicMeshOverlay<double, 1>;
 extern template class UE_EXTERN_TEMPLATE_API TDynamicMeshOverlay<int, 1>;
@@ -793,6 +794,7 @@ extern template class UE_EXTERN_TEMPLATE_API TDynamicMeshVectorOverlay<double, 3
 extern template class UE_EXTERN_TEMPLATE_API TDynamicMeshVectorOverlay<int, 3, FVector3i>;
 extern template class UE_EXTERN_TEMPLATE_API TDynamicMeshVectorOverlay<float, 4, FVector4f>;
 extern template class UE_EXTERN_TEMPLATE_API TDynamicMeshVectorOverlay<double, 4, FVector4d>;
+#endif
 
 #undef UE_EXTERN_TEMPLATE_API
 
