Geometry: Added 2d triangle circle intersection utility. Fixed 2d triangle IsInsideOrOn function to properly handle degenerate cases and to have tolerance. Added 2d segment IsOnSegment function and GetSide function that handles tolerance properly. Added unit tests.

diff --git a/Engine/Source/Runtime/GeometryCore/Private/Intersection/IntersectionQueries2.cpp b/Engine/Source/Runtime/GeometryCore/Private/Intersection/IntersectionQueries2.cpp
--- a/Engine/Source/Runtime/GeometryCore/Private/Intersection/IntersectionQueries2.cpp
+++ b/Engine/Source/Runtime/GeometryCore/Private/Intersection/IntersectionQueries2.cpp
@@ -33,6 +33,68 @@ bool UE::Geometry::TestIntersection(const TSegment2<Real>& Segment, const TAxisA
 	// of box if both endpoints are outside the box.
 }
 
+template<typename Real>
+bool UE::Geometry::DoesTriangleIntersectCircle2D(
+	const TVector2<Real>& A, const TVector2<Real>& B, const TVector2<Real>& C,
+	const TVector2<Real>& Center, double RadiusSquared)
+{
+	// Eliminate the fully contained case by checking for containment of circle center. Could use
+	//  TTriangle2::IsInside, but we need to reuse the edge vectors anyway for the other tests, so
+	//  we do it ourselves.
+	TVector2<Real> SideVectors[3]{ B - A, C - B, A - C };
+	TVector2<Real> CenterRelativeCorner[3]{ Center - A, Center - B, Center - C };
+	double Signs[3]
+	{
+		DotPerp(SideVectors[0], CenterRelativeCorner[0]),
+		DotPerp(SideVectors[1], CenterRelativeCorner[1]),
+		DotPerp(SideVectors[2], CenterRelativeCorner[2]),
+	};
+	// Note that we don't use >= here because that is susceptible to issues with degenerate triangles.
+	//  If one or more of the signs are legitimately zero, then we are on a side of the triangle,
+	//  and we should succeed in the checks further below.
+	if (Signs[0] * Signs[1] > 0 && Signs[1] * Signs[2] > 0 && Signs[2] * Signs[0] > 0)
+	{
+		return true;
+	}
+
+	// If not inside, try seeing if the circle covers one of the corner points
+	for (int i = 0; i < 3; ++i)
+	{
+		if (CenterRelativeCorner[i].SizeSquared() <= RadiusSquared)
+		{
+			return true;
+		}
+	}
+
+	// If still no intersection, try projecting onto each of the sides and see if the distance 
+	//  to the projected point is close enough.
+	for (int i = 0; i < 3; ++i)
+	{
+		// First check to see if we're inside the bounds of the segment (i.e. inside the two
+		//  perpendicular lines on each end). Note that we don't allow equality here because 
+		//  that would mean that we're in-line with the end corner, at which point we should have
+		//  succeeded in the corner containment tests above, if we were going to succeed. At the 
+		//  same time, not allowing equality to zero allows us to not worry about degenerate edges
+		//  causing the dot product to go to zero.
+		if (SideVectors[i].Dot(CenterRelativeCorner[i]) > 0
+			&& SideVectors[i].Dot(CenterRelativeCorner[(i + 1) % 3]) < 0
+
+			// Now check the projected distance. Signs is dot of unnormalized side perpendicular
+			//  vector with the relative center location, so dividing by side length would give us
+			//  the distance from that side. Rearranging a bit to avoid divide and sqrt, we get:
+			&& FMath::Square(Signs[i]) <= RadiusSquared * SideVectors[i].SizeSquared())
+		{
+			// Note about Signs[i] being 0 "incorrectly": the dot product not being zero tells us
+			//  that the side was not zero. Thus, Signs[i] should be a valid value, except possibly
+			//  up to two TNumericLimits<Real>::Min() lower than it should be if it underflowed in
+			//  DotPerp.
+			return true;
+		}
+	}
+
+	return false;
+}
+
 namespace UE
 {
 namespace Geometry
@@ -41,5 +103,12 @@ namespace Geometry
 template bool GEOMETRYCORE_API TestIntersection(const TSegment2<float>& Segment, const TAxisAlignedBox2<float>& Box);
 template bool GEOMETRYCORE_API TestIntersection(const TSegment2<double>& Segment, const TAxisAlignedBox2<double>& Box);
 		
+template bool GEOMETRYCORE_API DoesTriangleIntersectCircle2D(
+	const TVector2<float>& A, const TVector2<float>& B, const TVector2<float>& C,
+	const TVector2<float>& Center, double RadiusSquared);
+template bool GEOMETRYCORE_API DoesTriangleIntersectCircle2D(
+	const TVector2<double>& A, const TVector2<double>& B, const TVector2<double>& C,
+	const TVector2<double>& Center, double RadiusSquared);
+
 } // namespace UE::Geometry
 } // namespace UE
diff --git a/Engine/Source/Runtime/GeometryCore/Public/Intersection/IntersectionQueries2.h b/Engine/Source/Runtime/GeometryCore/Public/Intersection/IntersectionQueries2.h
--- a/Engine/Source/Runtime/GeometryCore/Public/Intersection/IntersectionQueries2.h
+++ b/Engine/Source/Runtime/GeometryCore/Public/Intersection/IntersectionQueries2.h
@@ -16,6 +16,11 @@ template <typename T> struct TSegment2;
 // Return true if Segment and Box intersect and false otherwise
 template<typename Real>
 GEOMETRYCORE_API bool TestIntersection(const TSegment2<Real>& Segment, const TAxisAlignedBox2<Real>& Box);
+
+template<typename Real>
+GEOMETRYCORE_API bool DoesTriangleIntersectCircle2D(
+	const TVector2<Real>& A, const TVector2<Real>& B, const TVector2<Real>& C, 
+	const TVector2<Real>& Center, double RadiusSquared);
 		
 } // namespace UE::Geometry
 } // namespace UE
diff --git a/Engine/Source/Runtime/GeometryCore/Public/SegmentTypes.h b/Engine/Source/Runtime/GeometryCore/Public/SegmentTypes.h
--- a/Engine/Source/Runtime/GeometryCore/Public/SegmentTypes.h
+++ b/Engine/Source/Runtime/GeometryCore/Public/SegmentTypes.h
@@ -197,14 +197,30 @@ public:
 		return TMathUtil<T>::Clamp(Alpha, (T)0, (T)1);
 	}
 
+	/**
+	 * Determine which side of the segment the query point lies on.
+	 * @param QueryPoint test point
+	 * @param Tolerance tolerance band in which we return 0
+	 * @return -1 if point is to right of line, +1 if left, and 0 if on line or within tolerance band
+	 */
+	int GetSide(const TVector2<T>& QueryPoint, T Tolerance = 0)
+	{
+		TVector2<T> StartPt = Center - Extent * Direction;
+		TVector2<T> AQ = QueryPoint - StartPt;
 
+		// Note that we don't have to adjust tolerance because direction is already normalized
+		double Dist = DotPerp(AQ, Direction);
+		return (Dist > Tolerance ? -1 : (Dist < -Tolerance ? 1 : 0));
+	}
 
 	/**
-	 * Determine which side of the segment the query point lies on
+	 * Determine which side of the segment the query point lies on. This version is deprecated because it
+	 *  and the static function give opposite results. It also does not correctly handle tolerance.
 	 * @param QueryPoint test point
 	 * @param Tolerance tolerance band in which we return 0
-	 * @return +1 if point is to right of line, -1 if left, and 0 if on line or within tolerance band
+	 * @return -1 if point is to right of line, +1 if left, and 0 if on line or within tolerance band
 	 */
+	UE_DEPRECATED(5.5, "Use GetSide instead.")
 	int WhichSide(const TVector2<T>& QueryPoint, T Tolerance = 0)
 	{
 		// [TODO] subtract Center from test?
@@ -321,24 +337,86 @@ public:
 		return dx * dx + dy * dy;
 	}
 
+	/**
+	 * Determine which side of the segment the query point lies on. 
+	 * @param StartPt first point of Segment
+	 * @param EndPt second point of Segment
+	 * @param QueryPoint test point
+	 * @param Tolerance tolerance band in which we return 0. Note that using tolerance == 0 is faster because
+	 *   we do not need to get distance of point from segment line.
+	 * @return -1 if point is to right of line, +1 if left, and 0 if on line or within tolerance band, or segment was degenerate
+	 */
+	static int GetSide(const TVector2<T>& StartPt, const TVector2<T>& EndPt, const TVector2<T>& QueryPt, T Tolerance = (T)0)
+	{
+		TVector2<T> AB = EndPt - StartPt;
+		T DotPerpResult = DotPerp(QueryPt - StartPt, AB);
+		if (Tolerance == 0)
+		{
+			return (DotPerpResult > 0 ? -1 : (DotPerpResult < 0 ? 1 : 0));
+		}
+
+		// If tolerance is nonzero, we have to use segment length to adjust it, since we use unnormalized
+		//  direction in DotPerp.
+		T LengthSquared = AB.SizeSquared();
+		if (LengthSquared == 0)
+		{
+			// It's possible to get here without DotPerp underflowing, so we early out to avoid incorrectly 
+			//  zeroing out the tolerance below while comparing against a nonzero DotPerpResult. 
+			// We could scale AB to deal with this case to try to better match behavior with the tolerance
+			//  0 case (which blindly uses DotPerp result), but it's fair to expect tolerance to affect the 
+			//  treatment of near-degenerate segments, so we don't bother.
+			return 0;
+		}
+		T ToleranceToUse = Tolerance * FMath::Sqrt(LengthSquared);
+
+		return (DotPerpResult > ToleranceToUse ? -1 : (DotPerpResult < -ToleranceToUse ? 1 : 0));
+	}
 
 	/**
-	 * Determine which side of the segment the query point lies on
+	 * Determine which side of the segment the query point lies on. This version is deprecated because it
+	 *  and the member function give opposite results. It also does not correctly handle tolerance.
 	 * @param StartPt first point of Segment
 	 * @param EndPt second point of Segment
 	 * @param QueryPoint test point
 	 * @param Tolerance tolerance band in which we return 0
 	 * @return +1 if point is to right of line, -1 if left, and 0 if on line or within tolerance band
 	 */
+	UE_DEPRECATED(5.5, "Use GetSide instead.")
 	static int WhichSide(const TVector2<T>& StartPt, const TVector2<T>& EndPt, const TVector2<T>& QueryPt, T Tolerance = (T)0)
 	{
 		T det = -Orient(StartPt, EndPt, QueryPt);
 		return (det > Tolerance ? +1 : (det < -Tolerance ? -1 : 0));
 	}
 
+	/**
+	 * Return true if QueryPt is on the segment between A and B.
+	 */
+	static bool IsOnSegment(const TVector2<T>& A, const TVector2<T>& B, const TVector2<T>& QueryPt, T Tolerance = (T)0)
+	{
+		const TVector2<T> AB = B - A;
 
+		T SquaredLength = AB.SquaredLength();
+		if (SquaredLength == 0)
+		{
+			// It's possible to underflow here (i.e. A and B are not coincident), but the distance between
+			//  A and B is negligible in that case, and we'll still correctly identify B as on the segment
+			//  (since we'll underflow the same way). The fact that we're slighly more tolerant near A at
+			//  these scales seems ok.
+			return TVector2<T>::DistSquared(A, QueryPt) <= Tolerance * Tolerance;
+		}
+		// We can either normalize AB so that the values we get from dot products are distances,
+		//  or just multiply the other side by length.
+		T DotTolerance = FMath::Sqrt(SquaredLength) * Tolerance;
+		const TVector2<T> AQ = QueryPt - A;
+		if (FMath::Abs(DotPerp(AB, AQ)) > DotTolerance)
+		{
+			// Must not be on the line
+			return false;
+		}
 
-
+		// Make sure it's between the endpoints
+		return (AB.Dot(AQ) >= -DotTolerance && AB.Dot(QueryPt - B) <= DotTolerance);
+	}
 
 protected:
 
diff --git a/Engine/Source/Runtime/GeometryCore/Public/TriangleTypes.h b/Engine/Source/Runtime/GeometryCore/Public/TriangleTypes.h
--- a/Engine/Source/Runtime/GeometryCore/Public/TriangleTypes.h
+++ b/Engine/Source/Runtime/GeometryCore/Public/TriangleTypes.h
@@ -6,6 +6,7 @@
 #include "VectorUtil.h"
 #include "IndexTypes.h"
 #include "BoxTypes.h"
+#include "SegmentTypes.h"
 
 namespace UE
 {
@@ -115,19 +116,60 @@ struct TTriangle2
 	}
 
 
-	/** @return true if QueryPoint is inside triangle or on edge */
-	static bool IsInsideOrOn(const TVector2<RealType>& A, const TVector2<RealType>& B, const TVector2<RealType>& C, const TVector2<RealType>& QueryPoint)
+	/** 
+	 * @return true if QueryPoint is inside triangle or on edge. Note that this is slower 
+	 *  than IsInside because of the need to handle degeneracy and tolerance. 
+	 */
+	static bool IsInsideOrOn(const TVector2<RealType>& A, const TVector2<RealType>& B, 
+		const TVector2<RealType>& C, const TVector2<RealType>& QueryPoint, RealType Tolerance = (RealType)0)
 	{
-		RealType Sign1 = Orient(A, B, QueryPoint);
-		RealType Sign2 = Orient(B, C, QueryPoint);
-		RealType Sign3 = Orient(C, A, QueryPoint);
-		return (Sign1*Sign2 >= 0) && (Sign2*Sign3 >= 0) && (Sign3*Sign1 >= 0);
+		RealType Sign1 = TSegment2<RealType>::GetSide(A, B, QueryPoint, Tolerance);
+		RealType Sign2 = TSegment2<RealType>::GetSide(B, C, QueryPoint, Tolerance);
+		RealType Sign3 = TSegment2<RealType>::GetSide(C, A, QueryPoint, Tolerance);
+
+		// If any of the signs are opposite, then definitely outside
+		if (Sign1 * Sign2 < 0 || Sign2 * Sign3 < 0 || Sign3 * Sign1 < 0)
+		{
+			return false;
+		}
+
+		// If some signs were zero, then things are more complicated because we are either colinear
+		//  with that edge, or the edge is degenerate enough to get a 0 value on the DotPerp.
+		int NumZero = static_cast<int>(Sign1 == 0) + static_cast<int>(Sign2 == 0) + static_cast<int>(Sign3 == 0);
+		switch (NumZero)
+		{
+		case 0:
+			// All were nonzero and none disagreed, so inside.
+			return true;
+		case 1:
+			// If only one sign was zero, seems more than likely that we're on that edge.
+			//  Hard to imagine some underflow case where that isn't actually the case, but we'll
+			//  do the segment check just in case.
+			return Sign1 == 0 ? TSegment2<RealType>::IsOnSegment(A, B, QueryPoint, Tolerance)
+				: Sign2 == 0 ? TSegment2<RealType>::IsOnSegment(B, C, QueryPoint, Tolerance)
+				: TSegment2<RealType>::IsOnSegment(C, A, QueryPoint, Tolerance);
+		case 2:
+			// Two signs were zero, so we expect to be on the vertex between those edges
+			return Sign1 != 0 ? TVector2<RealType>::DistSquared(QueryPoint, C) <= Tolerance * Tolerance
+				: Sign2 != 0 ? TVector2<RealType>::DistSquared(QueryPoint, A) <= Tolerance * Tolerance
+				: TVector2<RealType>::DistSquared(QueryPoint, B) <= Tolerance * Tolerance;
+		case 3:
+			// All three signs were zero. We're dealing with a denerate triangle of some
+			//  sort. It should be sufficient to check any two segments
+			return TSegment2<RealType>::IsOnSegment(A, B, QueryPoint, Tolerance)
+				|| TSegment2<RealType>::IsOnSegment(B, C, QueryPoint, Tolerance);
+		default:
+			return ensure(false);
+		}
 	}
 
-	/** @return true if QueryPoint is inside triangle or on edge */
-	bool IsInsideOrOn(const TVector2<RealType>& QueryPoint, RealType Epsilon = 0) const
+	/**
+	 * @return true if QueryPoint is inside triangle or on edge. Note that this is slower
+	 *  than IsInside because of the need to handle degeneracy and tolerance.
+	 */
+	bool IsInsideOrOn(const TVector2<RealType>& QueryPoint, RealType Tolerance = 0) const
 	{
-		return IsInsideOrOn(V[0], V[1], V[2], QueryPoint);
+		return IsInsideOrOn(V[0], V[1], V[2], QueryPoint, Tolerance);
 	}
 
 
@@ -136,9 +178,9 @@ struct TTriangle2
 	 *  (the code early-outs at the first 'outside' edge, which only works if the triangle is oriented as expected)
 	 * @return 1 if outside, -1 if inside, 0 if on boundary
 	 */
-	int IsInsideOrOn_Oriented(const TVector2<RealType>& QueryPoint) const
+	int IsInsideOrOn_Oriented(const TVector2<RealType>& QueryPoint, RealType Tolerance = (RealType)0) const
 	{
-		return IsInsideOrOn_Oriented(V[0], V[1], V[2], QueryPoint);
+		return IsInsideOrOn_Oriented(V[0], V[1], V[2], QueryPoint, Tolerance);
 	}
 
 	/**
@@ -146,29 +188,63 @@ struct TTriangle2
 	 *  (the code early-outs at the first 'outside' edge, which only works if the triangle is oriented as expected)
 	 * @return 1 if outside, -1 if inside, 0 if on boundary
 	 */
-	static int IsInsideOrOn_Oriented(const TVector2<RealType>& A, const TVector2<RealType>& B, const TVector2<RealType>& C, const TVector2<RealType>& QueryPoint)
+	static int IsInsideOrOn_Oriented(const TVector2<RealType>& A, const TVector2<RealType>& B, const TVector2<RealType>& C, 
+		const TVector2<RealType>& QueryPoint, RealType Tolerance = (RealType)0)
 	{
 		checkSlow(Orient(A, B, C) <= 0); // TODO: remove this checkSlow; it's just to make sure the orientation is as expected
 
-		RealType Sign1 = Orient(A, B, QueryPoint);
+		RealType Sign1 = TSegment2<RealType>::GetSide(A, B, QueryPoint, Tolerance);
 		if (Sign1 > 0)
 		{
 			return 1;
 		}
 		
-		RealType Sign2 = Orient(B, C, QueryPoint);
+		RealType Sign2 = TSegment2<RealType>::GetSide(B, C, QueryPoint, Tolerance);
 		if (Sign2 > 0)
 		{
 			return 1;
 		}
 		
-		RealType Sign3 = Orient(A, C, QueryPoint);
+		RealType Sign3 = TSegment2<RealType>::GetSide(A, C, QueryPoint, Tolerance);
 		if (Sign3 < 0) // note this edge is queried backwards so the sign test is also backwards
 		{
 			return 1;
 		}
 
-		return (Sign1 != 0 && Sign2 != 0 && Sign3 != 0) ? -1 : 0;
+		// If some signs were zero, then things are more complicated because we are either colinear
+		//  with that edge, or the edge is degenerate enough to get a 0 value on the DotPerp.
+		int NumZero = static_cast<int>(Sign1 == 0) + static_cast<int>(Sign2 == 0) + static_cast<int>(Sign3 == 0);
+		bool bIsOnEdge = false;
+		
+		switch (NumZero)
+		{
+		case 0:
+			// All signs were nonzero and in the correct direction, so must be inside triangle.
+			return -1;
+		case 1:
+			// If only one sign was zero, seems more than likely that we're on the opposite edge.
+			//  Hard to imagine some underflow case where that isn't actually the case, but we'll
+			//  do the segment check just in case.
+			bIsOnEdge = Sign1 == 0 ? TSegment2<RealType>::IsOnSegment(A, B, QueryPoint, Tolerance)
+				: Sign2 == 0 ? TSegment2<RealType>::IsOnSegment(B, C, QueryPoint, Tolerance)
+				: TSegment2<RealType>::IsOnSegment(C, A, QueryPoint, Tolerance);
+			break;
+		case 2:
+			// Two signs were zero, so it better be on the vertex between those edges
+			bIsOnEdge = Sign1 != 0 ? TVector2<RealType>::DistSquared(QueryPoint, C) <= Tolerance * Tolerance
+				: Sign2 != 0 ? TVector2<RealType>::DistSquared(QueryPoint, A) <= Tolerance * Tolerance
+				: TVector2<RealType>::DistSquared(QueryPoint, B) <= Tolerance * Tolerance;
+			break;
+		case 3:
+			// All three signs were zero. We're dealing with a degenerate triangle of some
+			//  sort. It should be sufficient to check any two segments.
+			bIsOnEdge = TSegment2<RealType>::IsOnSegment(A, B, QueryPoint, Tolerance)
+				|| TSegment2<RealType>::IsOnSegment(B, C, QueryPoint, Tolerance);
+			break;
+		default:
+			ensure(false);
+		}
+		return bIsOnEdge ? 0 : 1;
 	}
 
 
