clarify comment in mesh selection growth/contraction methods -- they need to iterate the entire selection, not the entire mesh

diff --git a/Engine/Source/Runtime/GeometryCore/Public/Selections/MeshEdgeSelection.h b/Engine/Source/Runtime/GeometryCore/Public/Selections/MeshEdgeSelection.h
--- a/Engine/Source/Runtime/GeometryCore/Public/Selections/MeshEdgeSelection.h
+++ b/Engine/Source/Runtime/GeometryCore/Public/Selections/MeshEdgeSelection.h
@@ -235,8 +235,8 @@ public:
 
 	/**
 	 *  Add all one-ring neighbors of current selection to set.
-	 *  On a large mesh this is quite expensive as we don't know the boundary,
-	 *  so we have to iterate over all selected triangles.
+	 *  On a large selection this is quite expensive as we don't know the boundary,
+	 *  so we have to iterate over all selected edges.
 	 *
 	 *  Return false from FilterF to prevent vertices from being included.
 	 */
diff --git a/Engine/Source/Runtime/GeometryCore/Public/Selections/MeshFaceSelection.h b/Engine/Source/Runtime/GeometryCore/Public/Selections/MeshFaceSelection.h
--- a/Engine/Source/Runtime/GeometryCore/Public/Selections/MeshFaceSelection.h
+++ b/Engine/Source/Runtime/GeometryCore/Public/Selections/MeshFaceSelection.h
@@ -348,8 +348,8 @@ public:
 
 	/**
 	 *  Add all triangles in vertex one-rings of current selection to set.
-	 *  On a large mesh this is quite expensive as we don't know the boundary,
-	 *  so we have to iterate over all triangles.
+	 *  On a large selection this is quite expensive as we don't know the boundary,
+	 *  so we have to iterate over all selected triangles.
 	 *  
 	 *  Return false from FilterF to prevent triangles from being included.
 	 */
@@ -366,8 +366,8 @@ public:
 
 	/**
 	 *  remove all triangles in vertex one-rings of current selection to set.
-	 *  On a large mesh this is quite expensive as we don't know the boundary,
-	 *  so we have to iterate over all triangles.
+	 *  On a large selection this is quite expensive as we don't know the boundary,
+	 *  so we have to iterate over all selected triangles.
 	 *  
 	 *  Return false from FilterF to prevent triangles from being deselected.
 	 */
diff --git a/Engine/Source/Runtime/GeometryCore/Public/Selections/MeshVertexSelection.h b/Engine/Source/Runtime/GeometryCore/Public/Selections/MeshVertexSelection.h
--- a/Engine/Source/Runtime/GeometryCore/Public/Selections/MeshVertexSelection.h
+++ b/Engine/Source/Runtime/GeometryCore/Public/Selections/MeshVertexSelection.h
@@ -224,8 +224,8 @@ public:
 
 	/**
 	 *  Add all one-ring neighbours of current selection to set.
-	 *  On a large mesh this is quite expensive as we don't know the boundary,
-	 *  so we have to iterate over all triangles.
+	 *  On a large selection this is quite expensive as we don't know the boundary,
+	 *  so we have to iterate over all selected vertices.
 	 *  
 	 *  Return false from FilterF to prevent vertices from being included.
 	 */
