DynamicMesh: Merge and collapse edge improvements:
- Fix bugs in the new permissive collapse options and add unit tests for them
- Let MergeEdges take an interpolation parameter for positions/attributes.
- Fix MergeEdges doing an average of vertex attributes when welding to an edge- it will now use the interpolation parameter like collapse with 0 being the default, so a default weld will keep attribute values on the welded-to vertices.
- Add MergeVertices operation that resolves either as collapse, weld, or bowtie creation depending on the selected vertices. Allow existing weld/collapse unit tests to also run with vertex merge.

diff --git a/Engine/Plugins/Runtime/MeshModelingToolset/Source/MeshModelingTools/Private/EditMeshPolygonsTool.cpp b/Engine/Plugins/Runtime/MeshModelingToolset/Source/MeshModelingTools/Private/EditMeshPolygonsTool.cpp
--- a/Engine/Plugins/Runtime/MeshModelingToolset/Source/MeshModelingTools/Private/EditMeshPolygonsTool.cpp
+++ b/Engine/Plugins/Runtime/MeshModelingToolset/Source/MeshModelingTools/Private/EditMeshPolygonsTool.cpp
@@ -2899,9 +2899,9 @@ void UEditMeshPolygonsTool::ApplyCollapseEdge()
 
 	FText PartialCollapseFailure = LOCTEXT("OnCollapseFailures", "Some edges could not be collapsed");
 
-	FDynamicMesh3::FEdgeCollapseOptions CollapseOptions;
+	FDynamicMesh3::FCollapseEdgeOptions CollapseOptions;
 	CollapseOptions.bAllowHoleCollapse = true;
-	CollapseOptions.bAllowCollapsingInternalEdgeWithExternalVertices = true;
+	CollapseOptions.bAllowCollapsingInternalEdgeWithBoundaryVertices = true;
 	CollapseOptions.bAllowTetrahedronCollapse = true;
 
 	FDynamicMesh3* Mesh = CurrentMesh.Get();
diff --git a/Engine/Source/Runtime/GeometryCore/Private/DynamicMesh/DynamicMesh3_Edits.cpp b/Engine/Source/Runtime/GeometryCore/Private/DynamicMesh/DynamicMesh3_Edits.cpp
--- a/Engine/Source/Runtime/GeometryCore/Private/DynamicMesh/DynamicMesh3_Edits.cpp
+++ b/Engine/Source/Runtime/GeometryCore/Private/DynamicMesh/DynamicMesh3_Edits.cpp
@@ -1283,11 +1283,11 @@ bool FDynamicMesh3::SplitVertexWouldLeaveIsolated(int VertexID, const TArrayView
 
 EMeshResult FDynamicMesh3::CanCollapseEdgeInternal(int vKeep, int vRemove, double collapse_t, FEdgeCollapseInfo* OutCollapseInfo) const
 {
-	return CanCollapseEdgeInternal(vKeep, vRemove, collapse_t, FEdgeCollapseOptions(),OutCollapseInfo);
+	return CanCollapseEdgeInternal(vKeep, vRemove, collapse_t, FCollapseEdgeOptions(),OutCollapseInfo);
 }
 
 EMeshResult UE::Geometry::FDynamicMesh3::CanCollapseEdgeInternal(int vKeep, int vRemove, double collapse_t, 
-	const FEdgeCollapseOptions& Options, FEdgeCollapseInfo* OutCollapseInfo) const
+	const FCollapseEdgeOptions& Options, FEdgeCollapseInfo* OutCollapseInfo) const
 {
 	if (IsVertex(vKeep) == false || IsVertex(vRemove) == false)
 	{
@@ -1405,12 +1405,31 @@ EMeshResult UE::Geometry::FDynamicMesh3::CanCollapseEdgeInternal(int vKeep, int
 	//
 	// NOTE: potentially scanning all edges here...couldn't we
 	//  pick up eac/bc/ad/bd as we go? somehow?
-	if (!Options.bAllowCollapsingInternalEdgeWithExternalVertices && !bIsBoundaryEdge
+	if (!Options.bAllowCollapsingInternalEdgeWithBoundaryVertices && !bIsBoundaryEdge
 		&& IsBoundaryVertex(a) && IsBoundaryVertex(b))
 	{
 		return EMeshResult::Failed_InvalidNeighbourhood;
 	}
 
+	// If we're allowing internal edge collapse with boundary vertices, we open the possibility
+	//  to another place where we would end up collapsing away an entire component- a quad. As
+	//  with single and double sided triangles, we currently disallow collapse in this case.
+	// Note that if we ever add an option to allow this, we have to deal with the possibility
+	//  that our kept vert (vKeep) may not actually survive a collapse.
+	if (Options.bAllowCollapsingInternalEdgeWithBoundaryVertices && !bIsBoundaryEdge
+		&& IsBoundaryEdge(eac) && IsBoundaryEdge(ead))
+	{
+		ebc = FindEdgeFromTri(b, c, t0);
+		if (IsBoundaryEdge(ebc))
+		{
+			ebd = FindEdgeFromTri(b, d, t1);
+			if (IsBoundaryEdge(ebd))
+			{
+				return EMeshResult::Failed_CollapseQuad;
+			}
+		}
+	}
+
 	if (OutCollapseInfo)
 	{
 		OutCollapseInfo->OpposingVerts = FIndex2i(c, d);
@@ -1430,20 +1449,20 @@ EMeshResult UE::Geometry::FDynamicMesh3::CanCollapseEdgeInternal(int vKeep, int
 
 EMeshResult FDynamicMesh3::CanCollapseEdge(int vKeep, int vRemove, double collapse_t) const
 {
-	return CanCollapseEdgeInternal(vKeep, vRemove, 0, FEdgeCollapseOptions(), nullptr);
+	return CanCollapseEdgeInternal(vKeep, vRemove, 0, FCollapseEdgeOptions(), nullptr);
 }
 
-EMeshResult FDynamicMesh3::CanCollapseEdge(int vKeep, int vRemove, const FEdgeCollapseOptions& Options) const
+EMeshResult FDynamicMesh3::CanCollapseEdge(int vKeep, int vRemove, const FCollapseEdgeOptions& Options) const
 {
 	return CanCollapseEdgeInternal(vKeep, vRemove, 0, Options, nullptr);
 }
 
 EMeshResult FDynamicMesh3::CollapseEdge(int vKeep, int vRemove, double collapse_t, FEdgeCollapseInfo& CollapseInfo)
 {
-	return CollapseEdge(vKeep, vRemove, collapse_t, FEdgeCollapseOptions(), CollapseInfo);
+	return CollapseEdge(vKeep, vRemove, collapse_t, FCollapseEdgeOptions(), CollapseInfo);
 }
 
-EMeshResult FDynamicMesh3::CollapseEdge(int vKeep, int vRemove, double collapse_t, const FEdgeCollapseOptions& Options, FEdgeCollapseInfo& CollapseInfo)
+EMeshResult FDynamicMesh3::CollapseEdge(int vKeep, int vRemove, double collapse_t, const FCollapseEdgeOptions& Options, FEdgeCollapseInfo& CollapseInfo)
 {
 	CollapseInfo = FEdgeCollapseInfo();
 
@@ -1524,12 +1543,34 @@ EMeshResult FDynamicMesh3::CollapseEdge(int vKeep, int vRemove, double collapse_
 		}
 		else
 		{
-			if (ReplaceEdgeVertex(eid, a, b) == -1)
+			// This is some edge oa, not in a triangle with ab, that we can in most cases change to be ob. However it's
+			//  possible that ob already exists as a boundary edge, and we need to instead weld the triangle incident to
+			//  oa to that edge. This situation is only permitted by CanCollapseEdgeInternal if bAllowHoleCollapse
+			//  is true and eid is a boundary edge.
+			int32 ExistingEdge = InvalidID;
+			if (Options.bAllowHoleCollapse && IsBoundaryEdge(eid) && (ExistingEdge = FindEdge(o, b)) != InvalidID)
 			{
-				checkfSlow(false, TEXT("FDynamicMesh3::CollapseEdge: failed at remove case else"));
-				return EMeshResult::Failed_UnrecoverableError;
+				int32 WeldedTriangle = GetEdgeT(eid).A;
+				if (ReplaceTriangleEdge(WeldedTriangle, eid, ExistingEdge) == -1
+					|| ReplaceEdgeTriangle(ExistingEdge, InvalidID, WeldedTriangle) == -1)
+				{
+					checkfSlow(false, TEXT("FDynamicMesh3::CollapseEdge: failed at remove case else"));
+					return EMeshResult::Failed_UnrecoverableError;
+				}
+				// Edge (o,a) should no longer exist
+				VertexEdgeLists.Remove(o, eid);
+				EdgeRefCounts.Decrement(eid);
+				checkSlow(EdgeRefCounts.IsValid(eid) == false);
+			}
+			else
+			{
+				if (ReplaceEdgeVertex(eid, a, b) == -1)
+				{
+					checkfSlow(false, TEXT("FDynamicMesh3::CollapseEdge: failed at remove case else"));
+					return EMeshResult::Failed_UnrecoverableError;
+				}
+				VertexEdgeLists.Insert(b, eid);
 			}
-			VertexEdgeLists.Insert(b, eid);
 		}
 
 		// [TODO] perhaps we can already have unique tri list because of the manifold-nbrhood check we need to do...
@@ -1615,6 +1656,46 @@ EMeshResult FDynamicMesh3::CollapseEdge(int vKeep, int vRemove, double collapse_
 			}
 		}
 
+		// If both bd and ad were boundary edges, or both bc and ac, then the edge bd/bc will have
+		//  no incident triangles and will need to be deleted. In that case, if vert c/d was not 
+		//  kept alive by some bowtie, that vertex will also need deleting. 
+		// This cannot happen if bAllowCollapsingInternalEdgeWithBoundaryVertices is false, and it
+		//  cannot happen for a boundary ab edge because that would require a single triangle, which
+		//  we currently disallow collapsing.
+		if (Options.bAllowCollapsingInternalEdgeWithBoundaryVertices)
+		{
+			if (GetEdgeT(ebc).A == InvalidID)
+			{
+				VertexEdgeLists.Remove(b, ebc);
+				EdgeRefCounts.Decrement(ebc);
+				if (VertexRefCounts.GetRefCount(c) == 1)
+				{
+					VertexEdgeLists.Clear(c);
+					VertexRefCounts.Decrement(c);
+				}
+				else
+				{
+					// The vert must still be part of a bowtie. Still need to remove the deleted edge.
+					VertexEdgeLists.Remove(c, ebc);
+				}
+			}
+			if (GetEdgeT(ebd).A == InvalidID)
+			{
+				VertexEdgeLists.Remove(b, ebd);
+				VertexEdgeLists.Remove(d, ebd);
+				EdgeRefCounts.Decrement(ebd);
+				if (VertexRefCounts.GetRefCount(d) == 1)
+				{
+					VertexEdgeLists.Clear(d);
+					VertexRefCounts.Decrement(d);
+				}
+				else
+				{
+					// The vert must still be part of a bowtie. Still need to remove the deleted edge.
+					VertexEdgeLists.Remove(d, ebd);
+				}
+			}
+		}
 	}
 	else
 	{
@@ -1694,9 +1775,15 @@ EMeshResult FDynamicMesh3::CollapseEdge(int vKeep, int vRemove, double collapse_
 
 
 
-EMeshResult FDynamicMesh3::MergeEdges(int eKeep, int eDiscard, FMergeEdgesInfo& MergeInfo, bool bCheckValidOrientation)
+EMeshResult FDynamicMesh3::MergeEdges(int KeepEdgeID, int DiscardEdgeID, FMergeEdgesInfo& MergeInfo, bool bCheckValidOrientation)
+{
+	return MergeEdges(KeepEdgeID, DiscardEdgeID, 0, MergeInfo, bCheckValidOrientation);
+}
+
+EMeshResult UE::Geometry::FDynamicMesh3::MergeEdges(int eKeep, int eDiscard, double InterpolationT, FMergeEdgesInfo& MergeInfo, bool bCheckValidOrientation)
 {
 	MergeInfo = FMergeEdgesInfo();
+	MergeInfo.InterpolationT = InterpolationT;
 
 	if (IsEdge(eKeep) == false || IsEdge(eDiscard) == false)
 	{
@@ -1815,9 +1902,32 @@ EMeshResult FDynamicMesh3::MergeEdges(int eKeep, int eDiscard, FMergeEdgesInfo&
 		}
 	}
 
+	auto ApplyInterpolation = [this, InterpolationT](int32 KeepVid, int32 RemoveVid)
+	{
+		SetVertex(KeepVid, Lerp(GetVertex(KeepVid), GetVertex(RemoveVid), InterpolationT));
+		if (HasVertexUVs())
+		{
+			SetVertexUV(KeepVid, Lerp(GetVertexUV(KeepVid), GetVertexUV(RemoveVid), (float)InterpolationT));
+		}
+		if (HasVertexNormals())
+		{
+			SetVertexNormal(KeepVid, Normalized(Lerp(GetVertexNormal(KeepVid), GetVertexNormal(RemoveVid), (float)InterpolationT)));
+		}
+		if (HasVertexColors())
+		{
+			SetVertexColor(KeepVid, Lerp(GetVertexColor(KeepVid), GetVertexColor(RemoveVid), (float)InterpolationT));
+		}
+	};
+
 	// [TODO] this acts on each interior tri twice. could avoid using vtx-tri iterator?
 	if (a != c)
 	{
+		if (InterpolationT != 0)
+		{
+			// Do the interpolation before we remove c
+			ApplyInterpolation(a, c);
+		}
+
 		// replace c w/ a in edges and tris connected to c, and move edges to a
 		for (int eid : VertexEdgeLists.Values(c))
 		{
@@ -1859,6 +1969,12 @@ EMeshResult FDynamicMesh3::MergeEdges(int eKeep, int eDiscard, FMergeEdgesInfo&
 
 	if (d != b)
 	{
+		if (InterpolationT != 0)
+		{
+			// Do the interpolation before we remove d
+			ApplyInterpolation(b, d);
+		}
+
 		// replace d w/ b in edges and tris connected to d, and move edges to b
 		for (int eid : VertexEdgeLists.Values(d))
 		{
@@ -1986,11 +2102,125 @@ EMeshResult FDynamicMesh3::MergeEdges(int eKeep, int eDiscard, FMergeEdgesInfo&
 }
 
 
+EMeshResult UE::Geometry::FDynamicMesh3::MergeVertices(int KeepVid, int DiscardVid, double InterpolationT, 
+	const FMergeVerticesOptions& Options, FMergeVerticesInfo& MergeInfo)
+{
+	MergeInfo = FMergeVerticesInfo();
+	MergeInfo.InterpolationT = InterpolationT;
+	MergeInfo.KeptVertex = KeepVid;
+	MergeInfo.RemovedVertex = DiscardVid;
+
+	if (!IsVertex(KeepVid) || !IsVertex(DiscardVid))
+	{
+		return EMeshResult::Failed_NotAVertex;
+	}
+	if (KeepVid == DiscardVid)
+	{
+		return EMeshResult::Failed_VertexAlreadyExists;
+	}
+
+	// See if we can resolve this as an edge collapse. 
+	// Note that this should be checked before trying to resolve as an edge weld because edge welding
+	//  will currently fail if there is an edge connecting the removed and kept vert (i.e. welding two
+	//  sides of a single triangle hole will fail, but collapse of the intervening edge will succeed
+	//  with our permissiveness options).
+	if (FindEdge(KeepVid, DiscardVid) != InvalidID)
+	{
+		FCollapseEdgeOptions CollapseOptions;
+		CollapseOptions.bAllowCollapsingInternalEdgeWithBoundaryVertices = true;
+		CollapseOptions.bAllowHoleCollapse = true;
+		CollapseOptions.bAllowTetrahedronCollapse = true;
+		MergeInfo.EdgeCollapseInfo.Emplace();
+		return CollapseEdge(KeepVid, DiscardVid, InterpolationT, CollapseOptions, MergeInfo.EdgeCollapseInfo.GetValue());
+	}
+
+	// See if we can resolve this as an edge weld
+	for (int32 KeepAdjacentEid : VertexEdgeLists.Values(KeepVid))
+	{
+		int32 KeepAdjacentVid = GetOtherEdgeVertex(KeepAdjacentEid, KeepVid);
+		
+		// See if the adjacent vert is also adjacent to DiscardVid
+		for (int32 DiscardAdjacentEid : VertexEdgeLists.Values(DiscardVid))
+		{
+			if (GetOtherEdgeVertex(DiscardAdjacentEid, DiscardVid) == KeepAdjacentVid)
+			{
+				// We've found a V shape (a vertex adjacent to both of the vertices we're working with).
+				//  Neither of the edges in the V shape can be a non boundary edge, else the merge would
+				//  create a non-manifold edge- we do this check ourselves so that we return Failed_InvalidNeighbourhood
+				//  instead of Failed_NotABoundaryEdge.
+				if (!IsBoundaryEdge(KeepAdjacentEid) || !IsBoundaryEdge(DiscardAdjacentEid))
+				{
+					return EMeshResult::Failed_InvalidNeighbourhood;
+				}
+				// The MergeEdges operation will do the other neighbor checks for us
+
+				MergeInfo.MergeEdgesInfo.Emplace();
+				return MergeEdges(KeepAdjacentEid, DiscardAdjacentEid, InterpolationT, MergeInfo.MergeEdgesInfo.GetValue(), false);
+			}
+		}// end for each adjacent vid to KeepVid
+	}//end V shape search
+
+	// If we got to here, the vertices are at least three edges apart, and we're creating a bowtie.
 
+	if (!Options.bAllowNonBoundaryBowtieCreation
+		&& (!IsBoundaryVertex(KeepVid) || !IsBoundaryVertex(DiscardVid)))
+	{
+		return EMeshResult::Failed_WouldCreateBowtie;
+	}
 
+	// Apply interpolation first, before removing DiscardVid
+	if (InterpolationT != 0)
+	{
+		SetVertex(KeepVid, Lerp(GetVertex(KeepVid), GetVertex(DiscardVid), InterpolationT));
+		if (HasVertexUVs())
+		{
+			SetVertexUV(KeepVid, Lerp(GetVertexUV(KeepVid), GetVertexUV(DiscardVid), (float)InterpolationT));
+		}
+		if (HasVertexNormals())
+		{
+			SetVertexNormal(KeepVid, Normalized(Lerp(GetVertexNormal(KeepVid), GetVertexNormal(DiscardVid), (float)InterpolationT)));
+		}
+		if (HasVertexColors())
+		{
+			SetVertexColor(KeepVid, Lerp(GetVertexColor(KeepVid), GetVertexColor(DiscardVid), (float)InterpolationT));
+		}
+	}
 
+	// Replace DiscardVid w/ KeepVid in edges and tris connected to DiscardVid, and move edges to KeepVid
+	for (int Eid : VertexEdgeLists.Values(DiscardVid))
+	{
+		ReplaceEdgeVertex(Eid, DiscardVid, KeepVid);
+		short ReplaceCount = 0;
+		const FEdge Edge = Edges[Eid];
+		if (ReplaceTriangleVertex(Edge.Tri[0], DiscardVid, KeepVid) >= 0)
+		{
+			ReplaceCount++;
+		}
+		if (Edge.Tri[1] != InvalidID)
+		{
+			if (ReplaceTriangleVertex(Edge.Tri[1], DiscardVid, KeepVid) >= 0)
+			{
+				ReplaceCount++;
+			}
+		}
+		VertexEdgeLists.Insert(KeepVid, Eid);
+		if (ReplaceCount > 0)
+		{
+			VertexRefCounts.Increment(KeepVid, ReplaceCount);
+			VertexRefCounts.Decrement(DiscardVid, ReplaceCount);
+		}
+	}
+	VertexEdgeLists.Clear(DiscardVid);
+	VertexRefCounts.Decrement(DiscardVid);
 
+	if (HasAttributes())
+	{
+		Attributes()->OnMergeVertices(MergeInfo);
+	}
 
+	UpdateChangeStamps(true, true);
+	return EMeshResult::Ok;
+}
 
 
 EMeshResult FDynamicMesh3::PokeTriangle(int TriangleID, const FVector3d& BaryCoordinates, FPokeTriangleInfo& PokeResult)
diff --git a/Engine/Source/Runtime/GeometryCore/Private/DynamicMesh/DynamicMeshAttributeSet.cpp b/Engine/Source/Runtime/GeometryCore/Private/DynamicMesh/DynamicMeshAttributeSet.cpp
--- a/Engine/Source/Runtime/GeometryCore/Private/DynamicMesh/DynamicMeshAttributeSet.cpp
+++ b/Engine/Source/Runtime/GeometryCore/Private/DynamicMesh/DynamicMeshAttributeSet.cpp
@@ -1054,6 +1054,38 @@ void FDynamicMeshAttributeSet::OnMergeEdges(const FDynamicMesh3::FMergeEdgesInfo
 	}
 }
 
+void FDynamicMeshAttributeSet::OnMergeVertices(const DynamicMeshInfo::FMergeVerticesInfo& MergeInfo)
+{
+	FDynamicMeshAttributeSetBase::OnMergeVertices(MergeInfo);
+
+	for (FDynamicMeshUVOverlay& UVLayer : UVLayers)
+	{
+		UVLayer.OnMergeVertices(MergeInfo);
+	}
+	for (FDynamicMeshNormalOverlay& NormalLayer : NormalLayers)
+	{
+		NormalLayer.OnMergeVertices(MergeInfo);
+	}
+	if (ColorLayer)
+	{
+		ColorLayer->OnMergeVertices(MergeInfo);
+	}
+	if (MaterialIDAttrib)
+	{
+		MaterialIDAttrib->OnMergeVertices(MergeInfo);
+	}
+
+	for (FDynamicMeshPolygroupAttribute& PolygroupLayer : PolygroupLayers)
+	{
+		PolygroupLayer.OnMergeVertices(MergeInfo);
+	}
+
+	for (FDynamicMeshWeightAttribute& WeightLayer : WeightLayers)
+	{
+		WeightLayer.OnMergeVertices(MergeInfo);
+	}
+}
+
 void FDynamicMeshAttributeSet::OnSplitVertex(const DynamicMeshInfo::FVertexSplitInfo& SplitInfo, const TArrayView<const int>& TrianglesToUpdate)
 {
 	FDynamicMeshAttributeSetBase::OnSplitVertex(SplitInfo, TrianglesToUpdate);
diff --git a/Engine/Source/Runtime/GeometryCore/Private/DynamicMesh/DynamicMeshOverlay.cpp b/Engine/Source/Runtime/GeometryCore/Private/DynamicMesh/DynamicMeshOverlay.cpp
--- a/Engine/Source/Runtime/GeometryCore/Private/DynamicMesh/DynamicMeshOverlay.cpp
+++ b/Engine/Source/Runtime/GeometryCore/Private/DynamicMesh/DynamicMeshOverlay.cpp
@@ -1392,6 +1392,28 @@ void TDynamicMeshOverlay<RealType, ElementSize>::OnMergeEdges(const FDynamicMesh
 	}
 }
 
+template<typename RealType, int ElementSize>
+void TDynamicMeshOverlay<RealType, ElementSize>::OnMergeVertices(const FDynamicMesh3::FMergeVerticesInfo& MergeInfo)
+{
+	// Attributes don't change, but one of the vertices got removed, so any elements on that vertex
+	//  should be reparented to the kept vertex. We have to iterate around the kept vertex since
+	//  ParentMesh is already updated.
+	for (int Tid : ParentMesh->VtxTrianglesItr(MergeInfo.KeptVertex))
+	{
+		if (!IsSetTriangle(Tid))
+		{
+			continue;
+		}
+		FIndex3i Triangle = GetTriangle(Tid);
+		for (int j = 0; j < 3; ++j)
+		{
+			if (Triangle[j] != FDynamicMesh3::InvalidID && ParentVertices[Triangle[j]] == MergeInfo.RemovedVertex)
+			{
+				ParentVertices[Triangle[j]] = MergeInfo.KeptVertex;
+			}
+		}
+	}
+}
 
 
 template<typename RealType, int ElementSize>
diff --git a/Engine/Source/Runtime/GeometryCore/Public/DynamicMesh/DynamicAttribute.h b/Engine/Source/Runtime/GeometryCore/Public/DynamicMesh/DynamicAttribute.h
--- a/Engine/Source/Runtime/GeometryCore/Public/DynamicMesh/DynamicAttribute.h
+++ b/Engine/Source/Runtime/GeometryCore/Public/DynamicMesh/DynamicAttribute.h
@@ -181,6 +181,14 @@ public:
 	{
 	}
 
+	/** 
+	 * Update to reflect a vertex merge that does not resolve as a collapse or edge merge (i.e. a
+	 *  vertex merge that resolves as bowtie creation).
+	 */
+	virtual void OnMergeVertices(const DynamicMeshInfo::FMergeVerticesInfo& MergeInfo)
+	{
+	}
+
 	/** Update to reflect an edge merge in the parent mesh */
 	virtual void OnSplitVertex(const DynamicMeshInfo::FVertexSplitInfo& SplitInfo, const TArrayView<const int>& TrianglesToUpdate)
 	{
@@ -366,6 +374,26 @@ public:
 			A->OnMergeEdges(MergeInfo);
 		}
 	}
+	virtual void OnMergeVertices(const DynamicMeshInfo::FMergeVerticesInfo& MergeInfo)
+	{
+		if (!ensureMsgf(!MergeInfo.EdgeCollapseInfo.IsSet(), TEXT("Vertex merge that resolves as edge collapse "
+			"is expected to have called OnCollapseEdge, not OnMergeVertices.")))
+		{
+			OnCollapseEdge(MergeInfo.EdgeCollapseInfo.GetValue());
+			return;
+		}
+		if (!ensureMsgf(!MergeInfo.MergeEdgesInfo.IsSet(), TEXT("Vertex merge that resolves as edge merge "
+			"is expected to have called OnMergeEdges, not OnMergeVertices.")))
+		{
+			OnMergeEdges(MergeInfo.MergeEdgesInfo.GetValue());
+			return;
+		}
+
+		for (TDynamicAttributeBase<ParentType>* A : RegisteredAttributes)
+		{
+			A->OnMergeVertices(MergeInfo);
+		}
+	}
 	virtual void OnSplitVertex(const DynamicMeshInfo::FVertexSplitInfo& SplitInfo, const TArrayView<const int>& TrianglesToUpdate)
 	{
 		for (TDynamicAttributeBase<ParentType>* A : RegisteredAttributes)
diff --git a/Engine/Source/Runtime/GeometryCore/Public/DynamicMesh/DynamicMesh3.h b/Engine/Source/Runtime/GeometryCore/Public/DynamicMesh/DynamicMesh3.h
--- a/Engine/Source/Runtime/GeometryCore/Public/DynamicMesh/DynamicMesh3.h
+++ b/Engine/Source/Runtime/GeometryCore/Public/DynamicMesh/DynamicMesh3.h
@@ -116,6 +116,7 @@ public:
 	using FEdgeSplitInfo    = DynamicMeshInfo::FEdgeSplitInfo;
 	using FEdgeCollapseInfo = DynamicMeshInfo::FEdgeCollapseInfo;
 	using FMergeEdgesInfo   = DynamicMeshInfo::FMergeEdgesInfo;
+	using FMergeVerticesInfo = DynamicMeshInfo::FMergeVerticesInfo;
 	using FPokeTriangleInfo = DynamicMeshInfo::FPokeTriangleInfo;
 	using FVertexSplitInfo  = DynamicMeshInfo::FVertexSplitInfo;
 
@@ -1309,7 +1310,7 @@ public:
 	 */
 	GEOMETRYCORE_API virtual bool SplitVertexWouldLeaveIsolated(int VertexID, const TArrayView<const int>& TrianglesToUpdate);
 
-	struct FEdgeCollapseOptions
+	struct FCollapseEdgeOptions
 	{
 		/**
 		 * When false, collapse is disallowed if the edge is the boundary of a single triangle hole,
@@ -1324,15 +1325,15 @@ public:
 		 *  connected to boundary edges. In some circumstances this could create a bowtie. In other
 		 *  cases, it could disconnected parts of a mesh that were connected by a bowtie.
 		 */
-		bool bAllowCollapsingInternalEdgeWithExternalVertices = false;
+		bool bAllowCollapsingInternalEdgeWithBoundaryVertices = false;
 		/**
 		 * When false, collapse is disallowed if we are collapsing the side of a tetrahedron. Note
 		 *  that a base edge of an open-base tetrahedron could be collapsed even if bAllowTetrahedronCollapse
-		 *  is false if bAllowBorderVerticesWhenInternal is true.
+		 *  is false if bAllowHoleCollapse is true.
 		 */
 		bool bAllowTetrahedronCollapse = false;
 	};
-	GEOMETRYCORE_API virtual EMeshResult CanCollapseEdge(int vKeep, int vRemove, const FEdgeCollapseOptions& Options) const;
+	GEOMETRYCORE_API virtual EMeshResult CanCollapseEdge(int vKeep, int vRemove, const FCollapseEdgeOptions& Options) const;
 
 	/**
 	 * Tests whether collapsing the specified edge using the CollapseEdge function would succeed.
@@ -1354,13 +1355,13 @@ public:
 	 * @return Ok on success, or enum value indicates why operation cannot be applied. Mesh remains unmodified on error.
 	 */
 	GEOMETRYCORE_API virtual EMeshResult CollapseEdge(int KeepVertID, int RemoveVertID, double EdgeParameterT, 
-		const FEdgeCollapseOptions& Options, FEdgeCollapseInfo& CollapseInfo);
+		const FCollapseEdgeOptions& Options, FEdgeCollapseInfo& CollapseInfo);
 
 	/**
 	 * Collapse the edge between the two vertices, if topologically possible. Equivalent to
 	 *  using the other overload with 0 for EdgeParameterT.
 	 */
-	virtual EMeshResult CollapseEdge(int KeepVertID, int RemoveVertID, const FEdgeCollapseOptions& Options, FEdgeCollapseInfo& CollapseInfo)
+	virtual EMeshResult CollapseEdge(int KeepVertID, int RemoveVertID, const FCollapseEdgeOptions& Options, FEdgeCollapseInfo& CollapseInfo)
 	{
 		return CollapseEdge(KeepVertID, RemoveVertID, 0, Options, CollapseInfo);
 	}
@@ -1396,12 +1397,62 @@ public:
 	 *
 	 * @param KeepEdgeID index of the edge that should be kept
 	 * @param DiscardEdgeID index of the edge that should be removed
+	 * @param InterpolationT each kept vertex is moved to Lerp(KeptPos, RemovePos, InterpolationT)
 	 * @param MergeInfo returned information about new and modified mesh elements
-	 * @param CheckValidOrientation perform edge consistency orientation checks before merging.
+	 * @param CheckValidOrientation perform edge consistency orientation checks before merging. Specifically,
+	 *  check that each discarded vertex is closer to the vertex it is being collapsed to than the other
+	 *  kept vertex (where the pairing is determined by the adjoining triangle winding). 
 	 * @return Ok on success, or enum value indicates why operation cannot be applied. Mesh remains unmodified on error.
 	 */
+	GEOMETRYCORE_API virtual EMeshResult MergeEdges(int KeepEdgeID, int DiscardEdgeID, double InterpolationT, FMergeEdgesInfo& MergeInfo, bool bCheckValidOrientation = true);
+	
+	/**
+	 * Weld one edge to the other. Equivalent to calling the other overload with 0 for InterpolationT
+	 *  (i.e. the vertices stay at unmodified kept vertex positions).
+	 */
 	GEOMETRYCORE_API virtual EMeshResult MergeEdges(int KeepEdgeID, int DiscardEdgeID, FMergeEdgesInfo& MergeInfo, bool bCheckValidOrientation=true);
 
+	struct FMergeVerticesOptions
+	{
+		// If false, we disallow vertex merges that attempt to merge one non-boundary vert to a
+		//  a non-adjacent vert, even if this is possible through a bowtie. Note that merging
+		//  boundary verts to create a bowtie on the boundary is still allowed, as this is a 
+		//  common intermediate step when welding edges. 
+		bool bAllowNonBoundaryBowtieCreation = false;
+	};
+
+	/**
+	 * Weld DiscardVid to KeepVid, if topologically possible and options allow. If the vertices are connected 
+	 *  by an existing edge, this resolves as a collapse of that edge, and therefore calls OnCollapseEdge in
+	 *  overlays. If not, but the two vertices share a vertex neighbor, this resolves as a weld of the intervening 
+	 *  edges (failing if the edges are not boundary edges, since that would create non-manifold edge), and therefore
+	 *  calls OnMergeEdges in the overlays. Otherwise, the merge resolves as bowtie creation, and calls OnMergeVertices
+	 *  in the overlays.
+	 * @param KeepVid vertex ID of the kept vertex
+	 * @param DiscardVid vertex ID of the vertex whose triangles are reattached to the kept vertex
+	 * @param InterpolationT the kept vertex is moved to Lerp(KeptPos, RemovePos, InterpolationT)
+	 * @param Options set the options for the merge
+	 * @param MergeInfo returned information about new and modified mesh elements
+	 * @return Ok on success, or enum value indicates why operation cannot be applied. Mesh remains unmodified on error.
+	 */
+	GEOMETRYCORE_API virtual EMeshResult MergeVertices(int KeepVid, int DiscardVid, double InterpolationT, 
+		const FMergeVerticesOptions& Options, FMergeVerticesInfo& MergeInfo);
+
+	/**
+	 * Weld DiscardVid to KeepVid. Equivalent to calling the options overload with default options.
+	 */
+	GEOMETRYCORE_API virtual EMeshResult MergeVertices(int KeepVid, int DiscardVid, double InterpolationT, FMergeVerticesInfo& MergeInfo)
+	{
+		return MergeVertices(KeepVid, DiscardVid, InterpolationT, FMergeVerticesOptions(), MergeInfo);
+	}
+
+	/**
+	 * Weld DiscardVid to KeepVid. Equivalent to calling the options overload with default options and 0 for InterpolationT.
+	 */
+	GEOMETRYCORE_API virtual EMeshResult MergeVertices(int KeepVid, int DiscardVid, FMergeVerticesInfo& MergeInfo)
+	{
+		return MergeVertices(KeepVid, DiscardVid, 0, FMergeVerticesOptions(), MergeInfo);
+	}
 
 	/**
 	 * Insert a new vertex inside a triangle, ie do a 1 to 3 triangle split
@@ -1723,7 +1774,7 @@ protected:
 	virtual EMeshResult CanCollapseEdgeInternal(int vKeep, int vRemove, double collapse_t, FEdgeCollapseInfo* OutCollapseInfo) const;
 private:
 	virtual EMeshResult CanCollapseEdgeInternal(int vKeep, int vRemove, double collapse_t,
-		const FEdgeCollapseOptions& Options, FEdgeCollapseInfo* OutCollapseInfo) const;
+		const FCollapseEdgeOptions& Options, FEdgeCollapseInfo* OutCollapseInfo) const;
 };
 
 
diff --git a/Engine/Source/Runtime/GeometryCore/Public/DynamicMesh/DynamicMeshAttributeSet.h b/Engine/Source/Runtime/GeometryCore/Public/DynamicMesh/DynamicMeshAttributeSet.h
--- a/Engine/Source/Runtime/GeometryCore/Public/DynamicMesh/DynamicMeshAttributeSet.h
+++ b/Engine/Source/Runtime/GeometryCore/Public/DynamicMesh/DynamicMeshAttributeSet.h
@@ -615,6 +615,7 @@ protected:
 	GEOMETRYCORE_API virtual void OnCollapseEdge(const DynamicMeshInfo::FEdgeCollapseInfo & collapseInfo);
 	GEOMETRYCORE_API virtual void OnPokeTriangle(const DynamicMeshInfo::FPokeTriangleInfo & pokeInfo);
 	GEOMETRYCORE_API virtual void OnMergeEdges(const DynamicMeshInfo::FMergeEdgesInfo & mergeInfo);
+	GEOMETRYCORE_API virtual void OnMergeVertices(const DynamicMeshInfo::FMergeVerticesInfo& mergeInfo);
 	GEOMETRYCORE_API virtual void OnSplitVertex(const DynamicMeshInfo::FVertexSplitInfo& SplitInfo, const TArrayView<const int>& TrianglesToUpdate);
 
 	/**
diff --git a/Engine/Source/Runtime/GeometryCore/Public/DynamicMesh/DynamicMeshOverlay.h b/Engine/Source/Runtime/GeometryCore/Public/DynamicMesh/DynamicMeshOverlay.h
--- a/Engine/Source/Runtime/GeometryCore/Public/DynamicMesh/DynamicMeshOverlay.h
+++ b/Engine/Source/Runtime/GeometryCore/Public/DynamicMesh/DynamicMeshOverlay.h
@@ -628,6 +628,8 @@ public:
 	void OnPokeTriangle(const DynamicMeshInfo::FPokeTriangleInfo& PokeInfo);
 	/** Update the overlay to reflect an edge merge in the parent mesh */
 	void OnMergeEdges(const DynamicMeshInfo::FMergeEdgesInfo& MergeInfo);
+	/** Update the overlay to reflect a vertex merge in the parent mesh */
+	void OnMergeVertices(const DynamicMeshInfo::FMergeVerticesInfo& MergeInfo);
 	/** Update the overlay to reflect a vertex split in the parent mesh */
 	void OnSplitVertex(const DynamicMeshInfo::FVertexSplitInfo& SplitInfo, const TArrayView<const int>& TrianglesToUpdate);
 
diff --git a/Engine/Source/Runtime/GeometryCore/Public/DynamicMesh/DynamicMeshTriangleAttribute.h b/Engine/Source/Runtime/GeometryCore/Public/DynamicMesh/DynamicMeshTriangleAttribute.h
--- a/Engine/Source/Runtime/GeometryCore/Public/DynamicMesh/DynamicMeshTriangleAttribute.h
+++ b/Engine/Source/Runtime/GeometryCore/Public/DynamicMesh/DynamicMeshTriangleAttribute.h
@@ -346,6 +346,12 @@ public:
 		// nothing to do here because triangles did not change
 	}
 
+	void OnMergeVertices(const DynamicMeshInfo::FMergeVerticesInfo& MergeInfo) override
+	{
+		// This resolves as either an edge collapse, edge weld, or merge of disconnected vertices. 
+		//  The triangles either get removed or unchanged- nothing more to do here.
+	}
+
 	/** Update the overlay to reflect a vertex split in the parent */
 	void OnSplitVertex(const DynamicMeshInfo::FVertexSplitInfo& SplitInfo, const TArrayView<const int>& TrianglesToUpdate) override
 	{
diff --git a/Engine/Source/Runtime/GeometryCore/Public/DynamicMesh/DynamicVertexAttribute.h b/Engine/Source/Runtime/GeometryCore/Public/DynamicMesh/DynamicVertexAttribute.h
--- a/Engine/Source/Runtime/GeometryCore/Public/DynamicMesh/DynamicVertexAttribute.h
+++ b/Engine/Source/Runtime/GeometryCore/Public/DynamicMesh/DynamicVertexAttribute.h
@@ -329,14 +329,19 @@ public:
 		// just blend the attributes?
 		if (MergeInfo.RemovedVerts.A != FDynamicMesh3::InvalidID)
 		{
-			SetAttributeFromLerp(MergeInfo.KeptVerts.A, MergeInfo.KeptVerts.A, MergeInfo.RemovedVerts.A, .5);
+			SetAttributeFromLerp(MergeInfo.KeptVerts.A, MergeInfo.KeptVerts.A, MergeInfo.RemovedVerts.A, MergeInfo.InterpolationT);
 		}
 		if (MergeInfo.RemovedVerts.B != FDynamicMesh3::InvalidID)
 		{
-			SetAttributeFromLerp(MergeInfo.KeptVerts.B, MergeInfo.KeptVerts.B, MergeInfo.RemovedVerts.B, .5);
+			SetAttributeFromLerp(MergeInfo.KeptVerts.B, MergeInfo.KeptVerts.B, MergeInfo.RemovedVerts.B, MergeInfo.InterpolationT);
 		}
 	}
 
+	void OnMergeVertices(const FDynamicMesh3::FMergeVerticesInfo& MergeInfo) override
+	{
+		SetAttributeFromLerp(MergeInfo.KeptVertex, MergeInfo.KeptVertex, MergeInfo.RemovedVertex, MergeInfo.InterpolationT);
+	}
+
 	/** Update the overlay to reflect a vertex split in the parent */
 	void OnSplitVertex(const FDynamicMesh3::FVertexSplitInfo& SplitInfo, const TArrayView<const int>& TrianglesToUpdate) override
 	{
diff --git a/Engine/Source/Runtime/GeometryCore/Public/DynamicMesh/DynamicVertexSkinWeightsAttribute.h b/Engine/Source/Runtime/GeometryCore/Public/DynamicMesh/DynamicVertexSkinWeightsAttribute.h
--- a/Engine/Source/Runtime/GeometryCore/Public/DynamicMesh/DynamicVertexSkinWeightsAttribute.h
+++ b/Engine/Source/Runtime/GeometryCore/Public/DynamicMesh/DynamicVertexSkinWeightsAttribute.h
@@ -447,14 +447,20 @@ public:
 		// just blend the attributes?
 		if (MergeInfo.RemovedVerts.A != FDynamicMesh3::InvalidID)
 		{
-			SetBoneWeightsFromLerp(MergeInfo.KeptVerts.A, MergeInfo.KeptVerts.A, MergeInfo.RemovedVerts.A, .5);
+			SetBoneWeightsFromLerp(MergeInfo.KeptVerts.A, MergeInfo.KeptVerts.A, MergeInfo.RemovedVerts.A, MergeInfo.InterpolationT);
 		}
 		if (MergeInfo.RemovedVerts.B != FDynamicMesh3::InvalidID)
 		{
-			SetBoneWeightsFromLerp(MergeInfo.KeptVerts.B, MergeInfo.KeptVerts.B, MergeInfo.RemovedVerts.B, .5);
+			SetBoneWeightsFromLerp(MergeInfo.KeptVerts.B, MergeInfo.KeptVerts.B, MergeInfo.RemovedVerts.B, MergeInfo.InterpolationT);
 		}
 	}
 
+	/** Update the overlay to reflect a vertex merge in the parent */
+	void OnMergeVertices(const FDynamicMesh3::FMergeVerticesInfo& MergeInfo) override
+	{
+		SetBoneWeightsFromLerp(MergeInfo.KeptVertex, MergeInfo.KeptVertex, MergeInfo.RemovedVertex, MergeInfo.InterpolationT);
+	}
+
 	/** Update the overlay to reflect a vertex split in the parent */
 	void OnSplitVertex(const FDynamicMesh3::FVertexSplitInfo& SplitInfo, const TArrayView<const int>& TrianglesToUpdate) override
 	{
diff --git a/Engine/Source/Runtime/GeometryCore/Public/DynamicMesh/InfoTypes.h b/Engine/Source/Runtime/GeometryCore/Public/DynamicMesh/InfoTypes.h
--- a/Engine/Source/Runtime/GeometryCore/Public/DynamicMesh/InfoTypes.h
+++ b/Engine/Source/Runtime/GeometryCore/Public/DynamicMesh/InfoTypes.h
@@ -207,6 +207,21 @@ struct FMergeEdgesInfo
 	// Even more Removed and Kept edges, in cases where there were multiple such edges on one or both sides of the merged edge
 	// Only possible if the pre-merge mesh had non-manifold vertices (aka bowties), in almost all meshes these arrays will be empty
 	TArray<int, TInlineAllocator<4>> BowtiesRemovedEdges, BowtiesKeptEdges;
+
+	double InterpolationT = 0;	// Interpolation parameter for each kept vertex in range [0,1] where 0 => KeptVertex and 1 => RemovedVertex
+};
+
+/** Information about mesh elements modified by MergeVertices() */
+struct FMergeVerticesInfo
+{
+	int KeptVertex;				// the vertex that was kept
+	int RemovedVertex;			// the vertex that was removed
+	double InterpolationT = 0;	// Interpolation parameter for the kept vertex in range [0,1] where 0 => KeptVertex and 1 => RemovedVertex
+
+	// If the merge resolves as an edge collapse, the information is stored here
+	TOptional<FEdgeCollapseInfo> EdgeCollapseInfo;
+	// If the merge resolves as an edge weld, the information is stored here
+	TOptional<FMergeEdgesInfo> MergeEdgesInfo;
 };
 
 /** Information about mesh elements modified/created by PokeTriangle() */
diff --git a/Engine/Source/Runtime/GeometryCore/Public/GeometryTypes.h b/Engine/Source/Runtime/GeometryCore/Public/GeometryTypes.h
--- a/Engine/Source/Runtime/GeometryCore/Public/GeometryTypes.h
+++ b/Engine/Source/Runtime/GeometryCore/Public/GeometryTypes.h
@@ -33,6 +33,7 @@ enum class EMeshResult
 	Failed_CollapseTriangle = 26,
 	Failed_NotABoundaryEdge = 27,
 	Failed_SameOrientation = 28,
+	Failed_CollapseQuad = 29,
 
 	Failed_WouldCreateBowtie = 30,
 	Failed_VertexAlreadyExists = 31,
