add helpers for uniform random sampling of positions and barycentric coordinates of triangles

diff --git a/Engine/Source/Runtime/GeometryCore/Public/VectorUtil.h b/Engine/Source/Runtime/GeometryCore/Public/VectorUtil.h
--- a/Engine/Source/Runtime/GeometryCore/Public/VectorUtil.h
+++ b/Engine/Source/Runtime/GeometryCore/Public/VectorUtil.h
@@ -416,6 +416,56 @@ namespace VectorUtil
 		return TVector<RealType>(fBary1, fBary2, fBary3);
 	}
 
+	/**
+	 * Sample uniform random barycentric coordinates on a triangle
+	 * The inputs should be two random floating point numbers in the [0,1] range
+	 * (See eq. 1 of 'Shape Distributions' by Osada, Funkhouser, Chazelle and Dobkin)
+	 * Note the sample positions will not match those of UniformSampleTrianglePoint for the same R1, R2 values
+	 */
+	template<typename RealType>
+	TVector<RealType> UniformSampleTriangleBarycentricCoords(RealType R1, RealType R2)
+	{
+		checkSlow(R1 >= 0);
+		RealType SqrtR1 = FMath::Sqrt(R1);
+		return TVector<RealType>(1 - SqrtR1, SqrtR1 * (1 - R2), SqrtR1 * R2);
+	}
+
+	/**
+	 * Sample a uniform random point on a triangle
+	 * The inputs should be two random floating point numbers in the [0,1] range, and the triangle coordinates
+	 * ref: http://mathworld.wolfram.com/TrianglePointPicking.html
+	 * Note the sample positions will not match those of UniformSampleTriangleBarycentricCoords for the same R1, R2 values
+	 */
+	template<typename RealType>
+	TVector<RealType> UniformSampleTrianglePoint(RealType R1, RealType R2, TVector<RealType>& A, TVector<RealType>& B, TVector<RealType>& C)
+	{
+		// reflect samples that would go outside the triangle
+		if (R1 + R2 > (RealType)1)
+		{
+			R1 = (RealType)1 - R1;
+			R2 = (RealType)1 - R2;
+		}
+		return A + R1 * (B - A) + R2 * (C - A);
+	}
+
+	/**
+	 * Sample a uniform random point on a 2D triangle
+	 * The inputs should be two random floating point numbers in the [0,1] range, and the triangle coordinates
+	 * ref: http://mathworld.wolfram.com/TrianglePointPicking.html
+	 * Note the sample positions will not match those of UniformSampleTriangleBarycentricCoords for the same R1, R2 values
+	 */
+	template<typename RealType>
+	TVector2<RealType> UniformSampleTrianglePoint(RealType R1, RealType R2, TVector2<RealType>& A, TVector2<RealType>& B, TVector2<RealType>& C)
+	{
+		// reflect samples that would go outside the triangle
+		if (R1 + R2 > (RealType)1)
+		{
+			R1 = (RealType)1 - R1;
+			R2 = (RealType)1 - R2;
+		}
+		return A + R1 * (B - A) + R2 * (C - A);
+	}
+
 	/**
 	 * @return solid angle at point P for triangle A,B,C
 	 */
