add DisconnectTrianglesAlongEdges as a mesh-topology version of CreateSeamsAtEdges, and an equivalent ApplyMeshDisconnectFacesAlongEdges geometry script method

diff --git a/Engine/Plugins/Runtime/GeometryScripting/Source/GeometryScriptingCore/Private/MeshModelingFunctions.cpp b/Engine/Plugins/Runtime/GeometryScripting/Source/GeometryScriptingCore/Private/MeshModelingFunctions.cpp
--- a/Engine/Plugins/Runtime/GeometryScripting/Source/GeometryScriptingCore/Private/MeshModelingFunctions.cpp
+++ b/Engine/Plugins/Runtime/GeometryScripting/Source/GeometryScriptingCore/Private/MeshModelingFunctions.cpp
@@ -335,6 +335,37 @@ UDynamicMesh* UGeometryScriptLibrary_MeshModelingFunctions::ApplyMeshDisconnectF
 	return TargetMesh;
 }
 
+UDynamicMesh* UGeometryScriptLibrary_MeshModelingFunctions::ApplyMeshDisconnectFacesAlongEdges(
+	UDynamicMesh* TargetMesh,
+	FGeometryScriptMeshSelection Selection,
+	UGeometryScriptDebug* Debug)
+{
+	if (TargetMesh == nullptr)
+	{
+		UE::Geometry::AppendError(Debug, EGeometryScriptErrorType::InvalidInputs, LOCTEXT("ApplyMeshDisconnectEdges_InvalidInput", "ApplyMeshDisconnectEdges: TargetMesh is Null"));
+		return TargetMesh;
+	}
+
+	TargetMesh->EditMesh([&](FDynamicMesh3& EditMesh)
+	{
+		TArray<int32> Edges;
+		Selection.ConvertToMeshIndexArray(EditMesh, Edges, EGeometryScriptIndexType::Edge);
+		if (Edges.Num() == 0)
+		{
+			UE::Geometry::AppendWarning(Debug, EGeometryScriptErrorType::InvalidInputs, LOCTEXT("ApplyMeshDisconnectFaces_NoEdges", "ApplyMeshDisconnectEdges: Selection contains no edges"));
+		}
+		else
+		{
+			FDynamicMeshEditor Editor(&EditMesh);
+			TSet<int32> EdgeSet(Edges);
+			Editor.DisconnectTrianglesAlongEdges(EdgeSet);
+		}
+
+	}, EDynamicMeshChangeType::GeneralEdit, EDynamicMeshAttributeChangeFlags::Unknown, false);
+
+	return TargetMesh;
+}
+
 
 UDynamicMesh* UGeometryScriptLibrary_MeshModelingFunctions::ApplyMeshDuplicateFaces(
 	UDynamicMesh* TargetMesh,
diff --git a/Engine/Plugins/Runtime/GeometryScripting/Source/GeometryScriptingCore/Public/GeometryScript/MeshModelingFunctions.h b/Engine/Plugins/Runtime/GeometryScripting/Source/GeometryScriptingCore/Public/GeometryScript/MeshModelingFunctions.h
--- a/Engine/Plugins/Runtime/GeometryScripting/Source/GeometryScriptingCore/Public/GeometryScript/MeshModelingFunctions.h
+++ b/Engine/Plugins/Runtime/GeometryScripting/Source/GeometryScriptingCore/Public/GeometryScript/MeshModelingFunctions.h
@@ -296,7 +296,7 @@ public:
 	/**
 	 * Disconnect the triangles of TargetMesh identified by the Selection.
 	 * The input Selection will still identify the same geometric elements after Disconnecting.
-	 * @param bAllowBowtiesInOutput if false, any bowtie vertices resulting created in the Duplicate area will be disconnected into unique vertices
+	 * @param bAllowBowtiesInOutput if false, any bowtie vertices created by the operation will be automatically split
 	 */
 	UFUNCTION(BlueprintCallable, Category = "GeometryScript|Modeling", meta=(ScriptMethod))
 	static UPARAM(DisplayName = "Target Mesh") UDynamicMesh* 
@@ -306,6 +306,20 @@ public:
 		bool bAllowBowtiesInOutput = true,
 		UGeometryScriptDebug* Debug = nullptr );
 
+	/**
+	 * Disconnect triangles of TargetMesh along the edges of the Selection.
+	 * The input Selection will still identify the same geometric elements after Disconnecting.
+	 * 
+	 * @param TargetMesh Mesh to operate on
+	 * @param Selection Which edges to operate on. Non-edge selections will be interpreted as edge selections -- i.e., all selected triangles' edges, or all selected vertices' one-ring edges.
+	 */
+	UFUNCTION(BlueprintCallable, Category = "GeometryScript|Modeling", meta = (ScriptMethod))
+	static UPARAM(DisplayName = "Target Mesh") UDynamicMesh*
+	ApplyMeshDisconnectFacesAlongEdges(
+		UDynamicMesh* TargetMesh,
+		FGeometryScriptMeshSelection Selection,
+		UGeometryScriptDebug* Debug = nullptr);
+
 	/**
 	 * Duplicate the triangles of TargetMesh identified by the Selection
 	 * @param NewTriangles a Mesh Selection of the duplicate triangles is returned here (with type Triangles)
diff --git a/Engine/Source/Runtime/GeometryCore/Private/DynamicMeshEditor.cpp b/Engine/Source/Runtime/GeometryCore/Private/DynamicMeshEditor.cpp
--- a/Engine/Source/Runtime/GeometryCore/Private/DynamicMeshEditor.cpp
+++ b/Engine/Source/Runtime/GeometryCore/Private/DynamicMeshEditor.cpp
@@ -22,6 +22,9 @@ namespace UE::Private::OverlaySeamHelpers
 	template<typename RealType, int OverlayDim>
 	static bool CreateSeamsAtEdges(const TSet<int32>& EidsToMakeIntoSeams, FDynamicMesh3* Mesh, TDynamicMeshOverlay<RealType, OverlayDim>* Overlay, TArray<int32>* AddedElementIDs)
 	{
+		// Note: The mesh boundary edge version of this method is the DisconnectTrianglesAlongEdges method --
+		//  the methods are very similar, so if you change this method consider changing that one as well
+
 		auto FindElementForVertex = [Mesh, Overlay](int32 MeshVertexID, int32 MeshEdgeID) -> int32
 		{
 			FIndex2i Tris = Mesh->GetEdgeT(MeshEdgeID);
@@ -858,7 +861,138 @@ void FDynamicMeshEditor::DisconnectTriangles(const TArray<int>& Triangles, bool
 	}
 }
 
+bool FDynamicMeshEditor::DisconnectTrianglesAlongEdges(const TSet<int32>& EdgeIDs, TArray<int32>* AddedVertexIDs)
+{
+	// Note: The overlay seams version of this method is the CreateSeamsAlongEdges method --
+	//  the methods are very similar, so if you change this method consider changing that one as well
+
+	/**
+	 * @param TIDsUntilNextBoundary TIDs in a fan away from EID around VID until the next boundary. Will
+	 *  be the entire one-ring if the given EID is the only boundary attached to VID.
+	 */
+	auto DoesVertHaveAnotherBoundaryEdgeAttached = 
+		[this, &EdgeIDs, AddedVertexIDs](int32 EID, int32 VID, TArray<int32>& TIDsUntilNextBoundary) -> bool
+	{
+		TIDsUntilNextBoundary.Reset();
+
+		FIndex2i EdgeTIDs = Mesh->GetEdgeT(EID);
+		int32 CurrentTID = EdgeTIDs.A;
+		int32 PreviousTID = EdgeTIDs.B;
+
+		int32 MaxNumTIDs = Mesh->GetVtxEdgeCount(VID);
+		while (true)
+		{
+			TIDsUntilNextBoundary.Add(CurrentTID);
+			if (TIDsUntilNextBoundary.Num() > MaxNumTIDs) // sanity check
+			{
+				ensure(false);
+				return false;
+			}
+
+			FIndex3i NextTriResult = UE::Geometry::FindNextAdjacentTriangleAroundVtx(Mesh, VID, CurrentTID, PreviousTID,
+				[&EdgeIDs, EID](int32 Tri0, int32 Tri1, int32 EIDBetween)
+				{
+					// Treat 'future' boundary edges that we haven't disconnected yet as boundary edges here
+					// (but not the current EID, since we're looking for *another* boundary edge)
+					return !(EIDBetween != EID && EdgeIDs.Contains(EIDBetween)); 
+				}
+			);
+			int32 NextTID = NextTriResult.A;
+
+			// Stop if we've come to a (current or future) new boundary ...
+			if (NextTID == IndexConstants::InvalidID)
+			{
+				return true;
+			}
+			// ... or if we've wrapped around to the start
+			else if (NextTID == EdgeTIDs.A)
+			{
+				return false;
+			}
+			PreviousTID = CurrentTID;
+			CurrentTID = NextTID;
+		}
+	};
+
+	auto SplitEdgeVert = [this, AddedVertexIDs](int32 VID, const TArray<int32>& TIDsToModify) 
+	{
+		DynamicMeshInfo::FVertexSplitInfo SplitInfo;
+		if (Mesh->SplitVertex(VID, TIDsToModify, SplitInfo) == EMeshResult::Ok && AddedVertexIDs != nullptr)
+		{
+			AddedVertexIDs->Add(SplitInfo.NewVertex);
+		}
+	};
+
+	TArray<int32> VertTIDsToUseForSplit[2];
+	for (int32 EID : EdgeIDs)
+	{
+		FIndex2i EdgeVIDs = Mesh->GetEdgeV(EID);
+		bool VertNeedsSplitting[2] = { false, false };
+		VertTIDsToUseForSplit[0].Reset();
+		VertTIDsToUseForSplit[1].Reset();
+
+		if (Mesh->IsBoundaryEdge(EID))
+		{
+			// If this is already a boundary edge, make sure its endpoints are not bowties. We create such edges 
+			// sometimes as we split adjacent edges, and it is also likely to be what the user wants when
+			// selecting edges adjacent to a bowtie.
+			for (int i = 0; i < 2; ++i)
+			{
+				SplitBowties(EdgeVIDs[i], AddedVertexIDs);
+			}
+			continue;
+		}
+
+		// If we're not already a boundary, then it's going to become one. A vert needs to get split if it
+		// has a present or future boundary edge attached, or else it will become a bowtie if we just split the
+		// other vert.
+		for (int i = 0; i < 2; ++i)
+		{
+			VertNeedsSplitting[i] = DoesVertHaveAnotherBoundaryEdgeAttached(EID, EdgeVIDs[i], VertTIDsToUseForSplit[i]);
+		}
+
+		// If neither absolutely has to get split, then one needs to get split anyway so that we
+		// make the edge into a boundary.
+		if (!VertNeedsSplitting[0] && !VertNeedsSplitting[1])
+		{
+			// We'll go halfway around the triangle one-ring, rounding up
+			int32 NumTrisToDisconnect = (VertTIDsToUseForSplit[0].Num() + 1) / 2;
+
+			// In doing this, we are splitting an adjacent edge and therefore might inadvertantly
+			// introduce a bowtie on the other vert of that edge, which we would like to avoid. 
+			// So keep track of that potential bowtie.
+			int32 OtherSplitEdge = Mesh->FindEdgeFromTriPair(
+				VertTIDsToUseForSplit[0][NumTrisToDisconnect - 1],
+				VertTIDsToUseForSplit[0][NumTrisToDisconnect]);
+			int32 PotentialBowtieVID = ensure(OtherSplitEdge != IndexConstants::InvalidID) ?
+				IndexUtil::FindEdgeOtherVertex(Mesh->GetEdgeV(OtherSplitEdge), EdgeVIDs[0])
+				: IndexConstants::InvalidID;
+
+			// Perform the split
+			VertTIDsToUseForSplit[0].SetNum(NumTrisToDisconnect);
+			SplitEdgeVert(EdgeVIDs[0], VertTIDsToUseForSplit[0]);
+
+			// Deal with the bowtie if we created one
+			if (ensure(PotentialBowtieVID != IndexConstants::InvalidID)
+				&& Mesh->IsBowtieVertex(PotentialBowtieVID))
+			{
+				SplitBowties(PotentialBowtieVID, AddedVertexIDs);
+			}
+		}
+		else
+		{
+			for (int i = 0; i < 2; ++i)
+			{
+				if (VertNeedsSplitting[i])
+				{
+					SplitEdgeVert(EdgeVIDs[i], VertTIDsToUseForSplit[i]);
+				}
+			}
+		}//end if at least one side needed splitting
+	}//end for each edge
 
+	return true;
+}
 
 
 void FDynamicMeshEditor::SplitBowties(FDynamicMeshEditResult& ResultOut)
@@ -886,12 +1020,11 @@ void FDynamicMeshEditor::SplitBowties(FDynamicMeshEditResult& ResultOut)
 
 
 
-void FDynamicMeshEditor::SplitBowties(int VertexID, FDynamicMeshEditResult& ResultOut)
+void FDynamicMeshEditor::SplitBowties(int VertexID, TArray<int32>* NewVertices)
 {
 	TArray<int> TrianglesOut, ContiguousGroupLengths;
 	TArray<bool> GroupIsLoop;
 	DynamicMeshInfo::FVertexSplitInfo SplitInfo;
-	check(Mesh->IsVertex(VertexID));
 	if (ensure(EMeshResult::Ok == Mesh->GetVtxContiguousTriangles(VertexID, TrianglesOut, ContiguousGroupLengths, GroupIsLoop)))
 	{
 		if (ContiguousGroupLengths.Num() > 1)
@@ -900,7 +1033,10 @@ void FDynamicMeshEditor::SplitBowties(int VertexID, FDynamicMeshEditResult& Resu
 			for (int GroupIdx = 1, GroupStartIdx = ContiguousGroupLengths[0]; GroupIdx < ContiguousGroupLengths.Num(); GroupStartIdx += ContiguousGroupLengths[GroupIdx++])
 			{
 				ensure(EMeshResult::Ok == Mesh->SplitVertex(VertexID, TArrayView<const int>(TrianglesOut.GetData() + GroupStartIdx, ContiguousGroupLengths[GroupIdx]), SplitInfo));
-				ResultOut.NewVertices.Add(SplitInfo.NewVertex);
+				if (NewVertices)
+				{
+					NewVertices->Add(SplitInfo.NewVertex);
+				}
 			}
 		}
 	}
diff --git a/Engine/Source/Runtime/GeometryCore/Public/DynamicMeshEditor.h b/Engine/Source/Runtime/GeometryCore/Public/DynamicMeshEditor.h
--- a/Engine/Source/Runtime/GeometryCore/Public/DynamicMeshEditor.h
+++ b/Engine/Source/Runtime/GeometryCore/Public/DynamicMeshEditor.h
@@ -223,6 +223,13 @@ public:
 	 */
 	GEOMETRYCORE_API bool DisconnectTriangles(const TSet<int>& TriangleSet, const TArray<FEdgeLoop>& BoundaryLoops, TArray<FLoopPairSet>& LoopSetOut, bool bAllowBoundaryVertices);
 
+	/**
+	 * Disconnect triangles along the given Edges, so that all input edges become boundary edges.
+	 * @param Edges IDs of edges to disconnect
+	 * @param AddedVertexIDs Optional array of vertices added by this operation
+	 */
+	GEOMETRYCORE_API bool DisconnectTrianglesAlongEdges(const TSet<int32>& Edges, TArray<int32>* AddedVertexIDs = nullptr);
+
 
 
 	/**
@@ -238,10 +245,18 @@ public:
 	 */
 	GEOMETRYCORE_API void SplitBowties(FDynamicMeshEditResult& ResultOut);
 
+	/**
+	 * Splits any bowties specifically on the given vertex, and (if not null) updates (does not reset!) NewVertices with any added vertices
+	 */
+	GEOMETRYCORE_API void SplitBowties(int VertexID, TArray<int32>* NewVertices = nullptr);
+
 	/**
 	 * Splits any bowties specifically on the given vertex, and updates (does not reset!) ResultOut with any added vertices
 	 */
-	GEOMETRYCORE_API void SplitBowties(int VertexID, FDynamicMeshEditResult& ResultOut);
+	void SplitBowties(int VertexID, FDynamicMeshEditResult& ResultOut)
+	{
+		SplitBowties(VertexID, &ResultOut.NewVertices);
+	}
 
 	/**
 	 * Splits bowties attached to any of the given triangles, and updates (does not reset!) ResultOut with any added vertices
