Geoscripting: Add an option to reconstruct non-manifold meshes

diff --git a/Engine/Plugins/Runtime/GeometryScripting/Source/GeometryScriptingCore/Private/MeshAssetFunctions.cpp b/Engine/Plugins/Runtime/GeometryScripting/Source/GeometryScriptingCore/Private/MeshAssetFunctions.cpp
--- a/Engine/Plugins/Runtime/GeometryScripting/Source/GeometryScriptingCore/Private/MeshAssetFunctions.cpp
+++ b/Engine/Plugins/Runtime/GeometryScripting/Source/GeometryScriptingCore/Private/MeshAssetFunctions.cpp
@@ -105,6 +105,7 @@ UDynamicMesh*  UGeometryScriptLibrary_StaticMeshFunctions::CopyMeshFromStaticMes
 	ConversionOptions.bIgnoreRemoveDegenerates = AssetOptions.bIgnoreRemoveDegenerates;
 	ConversionOptions.bUseBuildScale = AssetOptions.bUseBuildScale;
 	ConversionOptions.bUseSectionMaterialIndices = bUseSectionMaterials;
+	ConversionOptions.bIncludeNonManifoldSrcInfo = true;
 
 	FText ErrorMessage;
 
@@ -263,6 +264,8 @@ UDynamicMesh*  UGeometryScriptLibrary_StaticMeshFunctions::CopyMeshToStaticMesh(
 		}
 
 		FConversionToMeshDescriptionOptions ConversionOptions;
+		ConversionOptions.bConvertBackToNonManifold = Options.bUseOriginalVertexOrder;
+		
 		FDynamicMeshToMeshDescription Converter(ConversionOptions);
 		if (!bUseSectionMaterials && !Options.bReplaceMaterials)
 		{
@@ -943,6 +946,8 @@ UDynamicMesh* UGeometryScriptLibrary_StaticMeshFunctions::CopyMeshToSkeletalMesh
 	ToSkeletalMeshAsset->ModifyMeshDescription(TargetLOD.LODIndex);
 	
 	FConversionToMeshDescriptionOptions ConversionOptions;
+	ConversionOptions.bConvertBackToNonManifold = Options.bUseOriginalVertexOrder;
+
 	FDynamicMeshToMeshDescription Converter(ConversionOptions);
 	FromDynamicMesh->ProcessMesh([&](const FDynamicMesh3& ReadMesh)
 	{
diff --git a/Engine/Plugins/Runtime/GeometryScripting/Source/GeometryScriptingCore/Public/GeometryScript/MeshAssetFunctions.h b/Engine/Plugins/Runtime/GeometryScripting/Source/GeometryScriptingCore/Public/GeometryScript/MeshAssetFunctions.h
--- a/Engine/Plugins/Runtime/GeometryScripting/Source/GeometryScriptingCore/Public/GeometryScript/MeshAssetFunctions.h
+++ b/Engine/Plugins/Runtime/GeometryScripting/Source/GeometryScriptingCore/Public/GeometryScript/MeshAssetFunctions.h
@@ -78,6 +78,11 @@ public:
 	UPROPERTY(EditAnywhere, BlueprintReadWrite, Category = Options)
 	bool bRemapBoneIndicesToMatchAsset = false;
 
+	/** Use the original vertex order found in the source data. This is useful if the inbound mesh was originally non-manifold, and needs to keep
+	 *  the non-manifold structure when re-created. */
+	UPROPERTY(EditAnywhere, BlueprintReadWrite, Category = Options)
+	bool bUseOriginalVertexOrder = false;
+
 	// Whether to use the build scale on the target asset. If enabled, the inverse scale will be applied when saving to the asset, and the BuildScale will be preserved. Otherwise, BuildScale will be set to 1.0 on the asset BuildSettings.
 	UPROPERTY(EditAnywhere, BlueprintReadWrite, Category = Options)
 	bool bUseBuildScale = true;
diff --git a/Engine/Plugins/Runtime/GeometryScripting/Source/GeometryScriptingEditor/Private/CreateNewAssetUtilityFunctions.cpp b/Engine/Plugins/Runtime/GeometryScripting/Source/GeometryScriptingEditor/Private/CreateNewAssetUtilityFunctions.cpp
--- a/Engine/Plugins/Runtime/GeometryScripting/Source/GeometryScriptingEditor/Private/CreateNewAssetUtilityFunctions.cpp
+++ b/Engine/Plugins/Runtime/GeometryScripting/Source/GeometryScriptingEditor/Private/CreateNewAssetUtilityFunctions.cpp
@@ -218,6 +218,8 @@ UStaticMesh* UGeometryScriptLibrary_CreateNewAssetFunctions::CreateNewStaticMesh
 	AssetOptions.bCreatePhysicsBody = Options.bEnableCollision;
 	AssetOptions.CollisionType = Options.CollisionMode;
 
+	AssetOptions.bConvertBackToNonManifold = Options.bUseOriginalVertexOrder;
+
 	/**
 	 * We are making a copy of each LOD mesh since UDynamicMesh can potentially be editable asynchronously in the future, 
 	 * so we should not hold onto the pointer outside the function.
diff --git a/Engine/Plugins/Runtime/GeometryScripting/Source/GeometryScriptingEditor/Public/GeometryScript/CreateNewAssetUtilityFunctions.h b/Engine/Plugins/Runtime/GeometryScripting/Source/GeometryScriptingEditor/Public/GeometryScript/CreateNewAssetUtilityFunctions.h
--- a/Engine/Plugins/Runtime/GeometryScripting/Source/GeometryScriptingEditor/Public/GeometryScript/CreateNewAssetUtilityFunctions.h
+++ b/Engine/Plugins/Runtime/GeometryScripting/Source/GeometryScriptingEditor/Public/GeometryScript/CreateNewAssetUtilityFunctions.h
@@ -69,6 +69,11 @@ public:
 
 	UPROPERTY(BlueprintReadWrite, Category = Options)
 	TEnumAsByte<ECollisionTraceFlag> CollisionMode = ECollisionTraceFlag::CTF_UseDefault;
+	
+	/** Use the original vertex order found in the source data. This is useful if the inbound mesh was originally non-manifold, and needs to keep
+	 *  the non-manifold structure when re-created. */
+	UPROPERTY(BlueprintReadWrite, Category = Options)
+	bool bUseOriginalVertexOrder = false;
 };
 
 
@@ -89,6 +94,11 @@ public:
 	/** If true, will use the skeleton proportions (if availabale) stored in the dynamic mesh. */
 	UPROPERTY(BlueprintReadWrite, Category = Options)
 	bool bUseMeshBoneProportions = false;
+	
+	/** Use the original vertex order found in the source data. This is useful if the inbound mesh was originally non-manifold, and needs to keep
+	 *  the non-manifold structure when re-created. */
+	UPROPERTY(BlueprintReadWrite, Category = Options)
+	bool bUseOriginalVertexOrder = false;
 };
 
 
diff --git a/Engine/Plugins/Runtime/MeshModelingToolset/Source/ModelingComponents/Private/ConversionUtils/SceneComponentToDynamicMesh.cpp b/Engine/Plugins/Runtime/MeshModelingToolset/Source/ModelingComponents/Private/ConversionUtils/SceneComponentToDynamicMesh.cpp
--- a/Engine/Plugins/Runtime/MeshModelingToolset/Source/ModelingComponents/Private/ConversionUtils/SceneComponentToDynamicMesh.cpp
+++ b/Engine/Plugins/Runtime/MeshModelingToolset/Source/ModelingComponents/Private/ConversionUtils/SceneComponentToDynamicMesh.cpp
@@ -270,6 +270,7 @@ namespace Private::ConversionHelper
 		}
 
 		FMeshDescriptionToDynamicMesh Converter;
+		Converter.bVIDsFromNonManifoldMeshDescriptionAttr = AssetOptions.bIncludeNonManifoldSrcInfo;
 		if (!AssetOptions.bUseSectionMaterialIndices)
 		{
 			Converter.SetPolygonGroupToMaterialIndexMap(PolygonGroupToMaterialMap);
diff --git a/Engine/Plugins/Runtime/MeshModelingToolset/Source/ModelingComponents/Public/ConversionUtils/SceneComponentToDynamicMesh.h b/Engine/Plugins/Runtime/MeshModelingToolset/Source/ModelingComponents/Public/ConversionUtils/SceneComponentToDynamicMesh.h
--- a/Engine/Plugins/Runtime/MeshModelingToolset/Source/ModelingComponents/Public/ConversionUtils/SceneComponentToDynamicMesh.h
+++ b/Engine/Plugins/Runtime/MeshModelingToolset/Source/ModelingComponents/Public/ConversionUtils/SceneComponentToDynamicMesh.h
@@ -103,6 +103,9 @@ namespace Conversion {
 
 		// Whether to directly use the section material indices as output material IDs. Otherwise, material IDs will be indices in to the materials array.
 		bool bUseSectionMaterialIndices = false;
+
+		// Whether in the case of non-manifold source geometry to include a vertex layer to map vertices back to src VertexIDs
+		bool bIncludeNonManifoldSrcInfo = false;
 	};
 
 	/**
diff --git a/Engine/Plugins/Runtime/MeshModelingToolset/Source/ModelingComponentsEditorOnly/Private/AssetUtils/CreateSkeletalMeshUtil.cpp b/Engine/Plugins/Runtime/MeshModelingToolset/Source/ModelingComponentsEditorOnly/Private/AssetUtils/CreateSkeletalMeshUtil.cpp
--- a/Engine/Plugins/Runtime/MeshModelingToolset/Source/ModelingComponentsEditorOnly/Private/AssetUtils/CreateSkeletalMeshUtil.cpp
+++ b/Engine/Plugins/Runtime/MeshModelingToolset/Source/ModelingComponentsEditorOnly/Private/AssetUtils/CreateSkeletalMeshUtil.cpp
@@ -70,7 +70,9 @@ UE::AssetUtils::ECreateSkeletalMeshResult UE::AssetUtils::CreateSkeletalMeshAsse
 			for (const FDynamicMesh3* DynamicMesh : Options.SourceMeshes.DynamicMeshes)
 			{
 				ConstructedMeshDescriptions.AddDefaulted();
-				FDynamicMeshToMeshDescription Converter;
+				FConversionToMeshDescriptionOptions ConverterOptions;
+				ConverterOptions.bConvertBackToNonManifold = Options.bConvertBackToNonManifold;
+				FDynamicMeshToMeshDescription Converter(ConverterOptions);
 				FSkeletalMeshAttributes Attributes(ConstructedMeshDescriptions.Last());
 				Attributes.Register();
 				Converter.Convert(DynamicMesh, ConstructedMeshDescriptions.Last(), !Options.bEnableRecomputeTangents);
diff --git a/Engine/Plugins/Runtime/MeshModelingToolset/Source/ModelingComponentsEditorOnly/Private/AssetUtils/CreateStaticMeshUtil.cpp b/Engine/Plugins/Runtime/MeshModelingToolset/Source/ModelingComponentsEditorOnly/Private/AssetUtils/CreateStaticMeshUtil.cpp
--- a/Engine/Plugins/Runtime/MeshModelingToolset/Source/ModelingComponentsEditorOnly/Private/AssetUtils/CreateStaticMeshUtil.cpp
+++ b/Engine/Plugins/Runtime/MeshModelingToolset/Source/ModelingComponentsEditorOnly/Private/AssetUtils/CreateStaticMeshUtil.cpp
@@ -156,7 +156,9 @@ UE::AssetUtils::ECreateStaticMeshResult UE::AssetUtils::CreateStaticMeshAsset(
 			for (int32 k = 0; k < UseNumSourceModels; ++k)
 			{
 				FMeshDescription* Mesh = NewStaticMesh->GetMeshDescription(k);
-				FDynamicMeshToMeshDescription Converter;
+				FConversionToMeshDescriptionOptions ConverterOptions;
+				ConverterOptions.bConvertBackToNonManifold = Options.bConvertBackToNonManifold;
+				FDynamicMeshToMeshDescription Converter(ConverterOptions);
 				Converter.Convert(Options.SourceMeshes.DynamicMeshes[k], *Mesh, !Options.bEnableRecomputeTangents);
 				MaxNumSections = FMath::Max(MaxNumSections, Mesh->PolygonGroups().Num());
 				NewStaticMesh->CommitMeshDescription(k);
diff --git a/Engine/Plugins/Runtime/MeshModelingToolset/Source/ModelingComponentsEditorOnly/Public/AssetUtils/CreateSkeletalMeshUtil.h b/Engine/Plugins/Runtime/MeshModelingToolset/Source/ModelingComponentsEditorOnly/Public/AssetUtils/CreateSkeletalMeshUtil.h
--- a/Engine/Plugins/Runtime/MeshModelingToolset/Source/ModelingComponentsEditorOnly/Public/AssetUtils/CreateSkeletalMeshUtil.h
+++ b/Engine/Plugins/Runtime/MeshModelingToolset/Source/ModelingComponentsEditorOnly/Public/AssetUtils/CreateSkeletalMeshUtil.h
@@ -65,6 +65,9 @@ namespace UE::AssetUtils
 		// Controls whether or not the RecomputeTangents option will be enabled on the Asset
 		bool bEnableRecomputeTangents = true;
 
+		// Whether to attempt to re-create the original non-manifold mesh from the (optional) non-manifold index map on the dynamic mesh.
+		bool bConvertBackToNonManifold = false;
+
 		// List of skeletal materials. These will be used over the AssetMaterials, if provided.
 		TArray<FSkeletalMaterial> SkeletalMaterials;
 		
diff --git a/Engine/Plugins/Runtime/MeshModelingToolset/Source/ModelingComponentsEditorOnly/Public/AssetUtils/CreateStaticMeshUtil.h b/Engine/Plugins/Runtime/MeshModelingToolset/Source/ModelingComponentsEditorOnly/Public/AssetUtils/CreateStaticMeshUtil.h
--- a/Engine/Plugins/Runtime/MeshModelingToolset/Source/ModelingComponentsEditorOnly/Public/AssetUtils/CreateStaticMeshUtil.h
+++ b/Engine/Plugins/Runtime/MeshModelingToolset/Source/ModelingComponentsEditorOnly/Public/AssetUtils/CreateStaticMeshUtil.h
@@ -99,6 +99,9 @@ namespace AssetUtils
 
 		// Whether to build a reversed index buffer. Will double the memory usage of indices.
 		bool bBuildReversedIndexBuffer = true;
+
+		// Whether to attempt to re-create the original non-manifold mesh from the (optional) non-manifold index map on the dynamic mesh.
+		bool bConvertBackToNonManifold = false;
 	};
 
 	/**
diff --git a/Engine/Source/Runtime/GeometryCore/Private/DynamicMesh/NonManifoldMappingSupport.cpp b/Engine/Source/Runtime/GeometryCore/Private/DynamicMesh/NonManifoldMappingSupport.cpp
--- a/Engine/Source/Runtime/GeometryCore/Private/DynamicMesh/NonManifoldMappingSupport.cpp
+++ b/Engine/Source/Runtime/GeometryCore/Private/DynamicMesh/NonManifoldMappingSupport.cpp
@@ -2,22 +2,268 @@
 
 #include "DynamicMesh/NonManifoldMappingSupport.h"
 
-#include "CoreMinimal.h"
 #include "DynamicMesh/DynamicMesh3.h"
 #include "DynamicMesh/DynamicMeshAttributeSet.h"
 
+namespace UE::Geometry
+{
+
+class FNonManifoldSourceVertexIdAttributeChange final : 
+	public FDynamicMeshAttributeChangeBase
+{
+	friend class FNonManifoldSourceVertexIdAttribute;
+	TArray<TPair<int32, int32>> OldSourceVertexIds, NewSourceVertexIds;
+
+	bool bOldInvalidState = false;
+	bool bNewInvalidState = false;
 
-FName UE::Geometry::FNonManifoldMappingSupport::NonManifoldMeshVIDsAttrName = FName(TEXT("NonManifoldVIDAttr"));
+public:
+	FNonManifoldSourceVertexIdAttributeChange() = default;
 
+	virtual ~FNonManifoldSourceVertexIdAttributeChange() override = default;
 
-UE::Geometry::FNonManifoldMappingSupport::FNonManifoldMappingSupport(const FDynamicMesh3& MeshIn)
+	void SaveInitialVertex(const FDynamicMeshAttributeBase* Attribute, int VertexID) override;
+
+	void StoreAllFinalVertices(const FDynamicMeshAttributeBase* Attribute, const TArray<int>& VertexIDs) override;
+
+	bool Apply(FDynamicMeshAttributeBase* Attribute, bool bRevert) const override;
+};
+
+class FNonManifoldSourceVertexIdAttribute final :
+	public FDynamicMeshAttributeBase
 {
-	Reset(MeshIn);
+	friend class FNonManifoldSourceVertexIdAttributeChange;
+	
+	FDynamicMesh3* Parent = nullptr;
+	TDynamicVector<int32> SourceVertexIds;
+	bool bValid = false;
+	
+public:
+	FNonManifoldSourceVertexIdAttribute() = default;
+
+	FNonManifoldSourceVertexIdAttribute(FDynamicMesh3* InParent) : Parent(InParent)
+	{
+	}
+
+	virtual ~FNonManifoldSourceVertexIdAttribute() override = default;
+	
+	const FDynamicMesh3* GetParent() const
+	{
+		return Parent;
+	}
+	FDynamicMesh3* GetParent()
+	{
+		return Parent;
+	}
+	
+	bool IsValid() const
+	{
+		return bValid;
+	}
+
+	FDynamicMeshAttributeBase* MakeCopy(FDynamicMesh3* ParentIn) const override
+	{
+		FNonManifoldSourceVertexIdAttribute* Attribute = new FNonManifoldSourceVertexIdAttribute(ParentIn);
+		Attribute->Copy(*this);
+		return Attribute;
+	}
+	
+	FDynamicMeshAttributeBase* MakeNew(FDynamicMesh3* ParentIn) const override
+	{
+		FNonManifoldSourceVertexIdAttribute* Attribute = new FNonManifoldSourceVertexIdAttribute(ParentIn);
+		Attribute->Initialize();
+		return Attribute;
+	}
+	
+	void CompactInPlace(const FCompactMaps& CompactMaps) override
+	{
+		for (int32 FromVID = 0, NumVID = CompactMaps.NumVertexMappings(); FromVID < NumVID; FromVID++)
+		{
+			const int32 ToVID = CompactMaps.GetVertexMapping(FromVID);
+			if (ToVID == FCompactMaps::InvalidID)
+			{
+				continue;
+			}
+			if (ensure(ToVID <= FromVID))
+			{
+				SourceVertexIds[ToVID] = SourceVertexIds[FromVID];
+			}
+		}
+		SourceVertexIds.Resize(Parent->MaxVertexID());
+	}
+	
+	void Reparent(FDynamicMesh3* NewParent) override
+	{
+		Parent = NewParent;
+	}
+	
+	bool CopyThroughMapping(const TDynamicAttributeBase<FDynamicMesh3>* Source, const FMeshIndexMappings& Mapping) override
+	{
+		for (const TPair<int32, int32>& MapVID : Mapping.GetVertexMap().GetForwardMap())
+		{
+			int32 SourceVertexId;
+			if (!ensure(Source->CopyOut(MapVID.Key, &SourceVertexId, sizeof(SourceVertexId))))
+			{
+				return false;
+			}
+			SetValue(MapVID.Value, SourceVertexId);
+		}
+		return true;
+	}
+	
+	bool CopyOut(int RawID, void* Buffer, int BufferSize) const override
+	{
+		if (BufferSize != sizeof(int32) || !Parent->IsVertex(RawID))
+		{
+			return false;
+		}
+
+		*static_cast<int32 *>(Buffer) = SourceVertexIds[RawID];
+		return true;
+	}
+	
+	bool CopyIn(int RawID, void* Buffer, int BufferSize) override
+	{
+		if (BufferSize != sizeof(int32) || !Parent->IsVertex(RawID))
+		{
+			return false;
+		}
+
+		SourceVertexIds[RawID] = *static_cast<int32 *>(Buffer);
+		return true;
+	}
+	
+	TUniquePtr<TDynamicAttributeChangeBase<FDynamicMesh3>> NewBlankChange() const override
+	{
+		return MakeUnique<FNonManifoldSourceVertexIdAttributeChange>();
+	}
+	
+	void Initialize()
+	{
+		SourceVertexIds.Resize(Parent->MaxVertexID());
+		SourceVertexIds.Fill(INDEX_NONE);
+	}
+
+	void InitializeFromArray(TConstArrayView<int32> VertexToNonManifoldVertexIDMap)
+	{
+		check(VertexToNonManifoldVertexIDMap.Num() == Parent->MaxVertexID());
+		SourceVertexIds.Resize(Parent->MaxVertexID());
+		
+		for (int32 VertexID : Parent->VertexIndicesItr())
+		{
+			SetValue(VertexID, VertexToNonManifoldVertexIDMap[VertexID]);
+		}
+		bValid = true;
+	}
+
+	void Copy(const FNonManifoldSourceVertexIdAttribute& Copy)
+	{
+		FDynamicMeshAttributeBase::CopyParentClassData(Copy);
+		SourceVertexIds = Copy.SourceVertexIds;
+		bValid = Copy.bValid;
+	}
+	
+	int32 GetValue(int32 InVertexId) const
+	{
+		return SourceVertexIds[InVertexId];
+	}
+	
+	void SetValue(int32 InVertexId, int32 InSourceVertexId)
+	{
+		SourceVertexIds[InVertexId] = InSourceVertexId;
+	}
+
+	// Any topo operation on the mesh will invalidate the non-manifold information.
+	virtual void OnSplitEdge(const DynamicMeshInfo::FEdgeSplitInfo& ) override
+	{
+		bValid = false;
+	}
+
+	virtual void OnFlipEdge(const DynamicMeshInfo::FEdgeFlipInfo& FlipInfo) override
+	{
+		bValid = false;
+	}
+
+	virtual void OnCollapseEdge(const DynamicMeshInfo::FEdgeCollapseInfo& ) override
+	{
+		bValid = false;
+	}
+
+	virtual void OnPokeTriangle(const DynamicMeshInfo::FPokeTriangleInfo& ) override
+	{
+		bValid = false;
+	}
+
+	virtual void OnMergeEdges(const DynamicMeshInfo::FMergeEdgesInfo& ) override
+	{
+		bValid = false;
+	}
+
+	virtual void OnMergeVertices(const DynamicMeshInfo::FMergeVerticesInfo& ) override
+	{
+		bValid = false;
+	}
+
+	virtual void OnSplitVertex(const DynamicMeshInfo::FVertexSplitInfo& , const TArrayView<const int>& ) override
+	{
+		bValid = false;
+	}
+	
+};
+
+void FNonManifoldSourceVertexIdAttributeChange::SaveInitialVertex(const FDynamicMeshAttributeBase* Attribute, int VertexID)
+{
+	const FNonManifoldSourceVertexIdAttribute* NonManifoldAttribute = static_cast<const FNonManifoldSourceVertexIdAttribute*>(Attribute);
+	if (OldSourceVertexIds.IsEmpty())
+	{
+		bOldInvalidState = NonManifoldAttribute->IsValid();
+	}
+	
+	OldSourceVertexIds.Emplace(VertexID, NonManifoldAttribute->GetValue(VertexID));
 }
 
-void UE::Geometry::FNonManifoldMappingSupport::Reset(const FDynamicMesh3& MeshIn)
+void FNonManifoldSourceVertexIdAttributeChange::StoreAllFinalVertices(const FDynamicMeshAttributeBase* Attribute, const TArray<int>& VertexIDs)
 {
+	NewSourceVertexIds.Reserve(NewSourceVertexIds.Num() + VertexIDs.Num());
+	const FNonManifoldSourceVertexIdAttribute* NonManifoldAttribute = static_cast<const FNonManifoldSourceVertexIdAttribute*>(Attribute);
+	for (int32 VertexID: VertexIDs)
+	{
+		NewSourceVertexIds.Emplace(VertexID, NonManifoldAttribute->GetValue(VertexID));
+	}
+
+	// Store the last known valid state.  
+	bNewInvalidState = NonManifoldAttribute->IsValid();
+}
+
+bool FNonManifoldSourceVertexIdAttributeChange::Apply(FDynamicMeshAttributeBase* Attribute, bool bRevert) const
+{
+	const TArray<TPair<int32, int32>>& Changes = bRevert ? OldSourceVertexIds : NewSourceVertexIds;
+	FNonManifoldSourceVertexIdAttribute* NonManifoldAttribute = static_cast<FNonManifoldSourceVertexIdAttribute*>(Attribute);
+	for (const TPair<int32, int32>& Item: Changes)
+	{
+		if (ensure(NonManifoldAttribute->GetParent()->IsVertex(Item.Key)))
+		{
+			NonManifoldAttribute->SetValue(Item.Key, Item.Value);
+		}
+	}
+
+	// Restore the valid state as well.
+	NonManifoldAttribute->bValid = bRevert ? bOldInvalidState : bNewInvalidState;
 	
+	return true;
+}
+
+
+FName FNonManifoldMappingSupport::NonManifoldMeshVIDsAttrName = FName(TEXT("NonManifoldVIDAttr"));
+
+
+FNonManifoldMappingSupport::FNonManifoldMappingSupport(const FDynamicMesh3& MeshIn)
+{
+	Reset(MeshIn);
+}
+
+void FNonManifoldMappingSupport::Reset(const FDynamicMesh3& MeshIn)
+{
 	NonManifoldSrcVIDsAttribute = nullptr;
 	DynamicMesh = &MeshIn;
 
@@ -25,31 +271,29 @@ void UE::Geometry::FNonManifoldMappingSupport::Reset(const FDynamicMesh3& MeshIn
 
 	if (Attributes && Attributes->HasAttachedAttribute(NonManifoldMeshVIDsAttrName))
 	{
-		NonManifoldSrcVIDsAttribute = static_cast<const FDynamicMeshVertexInt32Attribute* >(Attributes->GetAttachedAttribute(NonManifoldMeshVIDsAttrName));
+		NonManifoldSrcVIDsAttribute = static_cast<const FNonManifoldSourceVertexIdAttribute* >(Attributes->GetAttachedAttribute(NonManifoldMeshVIDsAttrName));
 	}
-	
 }
 
-bool  UE::Geometry::FNonManifoldMappingSupport::IsNonManifoldVertexInSource() const
+bool  FNonManifoldMappingSupport::IsNonManifoldVertexInSource() const
 {
-	return (NonManifoldSrcVIDsAttribute != nullptr);
+	return NonManifoldSrcVIDsAttribute != nullptr && NonManifoldSrcVIDsAttribute->IsValid();
 }
 
-int32 UE::Geometry::FNonManifoldMappingSupport::GetOriginalNonManifoldVertexID(const int32 vid) const
+int32 FNonManifoldMappingSupport::GetOriginalNonManifoldVertexID(const int32 vid) const
 {
 	checkSlow(DynamicMesh->IsVertex(vid));
 
 	int32 Result = vid;
 	if (NonManifoldSrcVIDsAttribute)
 	{
-		NonManifoldSrcVIDsAttribute->GetValue(vid, &Result);
+		Result = NonManifoldSrcVIDsAttribute->GetValue(vid);
 	}
 	return Result;
 }
 
-bool UE::Geometry::FNonManifoldMappingSupport::AttachNonManifoldVertexMappingData(const TArray<int32>& VertexToNonManifoldVertexIDMap, FDynamicMesh3& MeshInOut)
+bool FNonManifoldMappingSupport::AttachNonManifoldVertexMappingData(const TArray<int32>& VertexToNonManifoldVertexIDMap, FDynamicMesh3& MeshInOut)
 {
-	
 	const int32 MaxVID = MeshInOut.MaxVertexID();
 	FDynamicMeshAttributeSet* Attributes = MeshInOut.Attributes();
 
@@ -59,35 +303,31 @@ bool UE::Geometry::FNonManifoldMappingSupport::AttachNonManifoldVertexMappingDat
 	}
 
 	// create and attach a vertex ID buffer, removing any pre-exising one.
-	FDynamicMeshVertexInt32Attribute* VIDAttr = [&] {
-													// replace existing buffer.  we *should* be able to re-use it
-													if (Attributes->HasAttachedAttribute(NonManifoldMeshVIDsAttrName))
-													{
-														Attributes->RemoveAttribute(NonManifoldMeshVIDsAttrName);
-													}
-													FDynamicMeshVertexInt32Attribute* SrcMeshVIDAttr = new FDynamicMeshVertexInt32Attribute(&MeshInOut);
-													SrcMeshVIDAttr->SetName(NonManifoldMeshVIDsAttrName);
-													MeshInOut.Attributes()->AttachAttribute(NonManifoldMeshVIDsAttrName, SrcMeshVIDAttr);
-													return SrcMeshVIDAttr;
-												}();
-	
+	FNonManifoldSourceVertexIdAttribute* VIDAttr = [&] {
+		// replace existing buffer.  we *should* be able to re-use it
+		if (Attributes->HasAttachedAttribute(NonManifoldMeshVIDsAttrName))
+		{
+			Attributes->RemoveAttribute(NonManifoldMeshVIDsAttrName);
+		}
+		FNonManifoldSourceVertexIdAttribute* SrcMeshVIDAttr = new FNonManifoldSourceVertexIdAttribute(&MeshInOut);
+		SrcMeshVIDAttr->SetName(NonManifoldMeshVIDsAttrName);
+		MeshInOut.Attributes()->AttachAttribute(NonManifoldMeshVIDsAttrName, SrcMeshVIDAttr);
+		return SrcMeshVIDAttr;
+	}();
 	
-	for (int32 vid : MeshInOut.VertexIndicesItr())
-	{
-		const int32 SrcVID = VertexToNonManifoldVertexIDMap[vid];
-		VIDAttr->SetValue(vid, &SrcVID);
-	}
+	VIDAttr->InitializeFromArray(VertexToNonManifoldVertexIDMap);
 	return true;
 }
 
-void UE::Geometry::FNonManifoldMappingSupport::RemoveNonManifoldVertexMappingData(FDynamicMesh3& MeshInOut)
+void FNonManifoldMappingSupport::RemoveNonManifoldVertexMappingData(FDynamicMesh3& MeshInOut)
 {
 
 	FDynamicMeshAttributeSet* Attributes = MeshInOut.Attributes();
 
 	if (Attributes && Attributes->HasAttachedAttribute(NonManifoldMeshVIDsAttrName))
 	{
-		 Attributes->RemoveAttribute(NonManifoldMeshVIDsAttrName);
+		Attributes->RemoveAttribute(NonManifoldMeshVIDsAttrName);
 	}
 
+}
 }
\ No newline at end of file
diff --git a/Engine/Source/Runtime/GeometryCore/Public/DynamicMesh/NonManifoldMappingSupport.h b/Engine/Source/Runtime/GeometryCore/Public/DynamicMesh/NonManifoldMappingSupport.h
--- a/Engine/Source/Runtime/GeometryCore/Public/DynamicMesh/NonManifoldMappingSupport.h
+++ b/Engine/Source/Runtime/GeometryCore/Public/DynamicMesh/NonManifoldMappingSupport.h
@@ -3,7 +3,8 @@
 
 #pragma once
 
-#include "CoreMinimal.h"
+#include "Containers/Array.h"
+#include "UObject/NameTypes.h"
 
 namespace UE
 {
@@ -12,11 +13,7 @@ namespace Geometry
 
 // forward
 class FDynamicMesh3;
-
-template<typename AttribValueType, int AttribDimension, typename ParentType>
-class TDynamicVertexAttribute;
-
-typedef  TDynamicVertexAttribute<int32, 1, FDynamicMesh3> FDynamicMeshVertexInt32Attribute;
+class FNonManifoldSourceVertexIdAttribute;
 
 
 /*
@@ -113,7 +110,7 @@ public:
 
 protected:
 
-	const FDynamicMeshVertexInt32Attribute*  NonManifoldSrcVIDsAttribute = nullptr;
+	const FNonManifoldSourceVertexIdAttribute*  NonManifoldSrcVIDsAttribute = nullptr;
 	const FDynamicMesh3* DynamicMesh = nullptr;
 };
 
diff --git a/Engine/Source/Runtime/MeshConversion/Private/DynamicMeshToMeshDescription.cpp b/Engine/Source/Runtime/MeshConversion/Private/DynamicMeshToMeshDescription.cpp
--- a/Engine/Source/Runtime/MeshConversion/Private/DynamicMeshToMeshDescription.cpp
+++ b/Engine/Source/Runtime/MeshConversion/Private/DynamicMeshToMeshDescription.cpp
@@ -9,11 +9,11 @@
 #include "DynamicMesh/DynamicVertexAttribute.h"
 #include "MeshDescriptionBuilder.h"
 #include "DynamicMesh/MeshTangents.h"
+#include "DynamicMesh/NonManifoldMappingSupport.h"
 #include "Util/ColorConstants.h"
 
 using namespace UE::Geometry;
 
-
 namespace DynamicMeshToMeshDescriptionConversionHelper
 {
 	// NOTE: assumes the order of triangles in the MeshIn correspond to the ordering over tris on MeshOut
@@ -44,6 +44,76 @@ namespace DynamicMeshToMeshDescriptionConversionHelper
 			}
 		}
 	}
+
+	void ConvertVertices(
+		const FDynamicMesh3& InMesh,
+		FMeshDescriptionBuilder& InBuilder,
+		const bool bReconstructNonManifoldMesh,
+		TArray<FVertexID>& OutVertexMappings
+		)
+	{
+		OutVertexMappings.SetNum(InMesh.MaxVertexID());
+		
+		FNonManifoldMappingSupport ManifoldMapping(InMesh);
+		if (bReconstructNonManifoldMesh && ManifoldMapping.IsNonManifoldVertexInSource())
+		{
+			// Figure out the actual vertex count we need. For non-mapping vertices, we count them separately and
+			// then tally them all up. The new vertices will be appended after the original vertices.
+			int32 MaxOriginalVertexID = 0;
+			int32 NumNewVertices = 0;
+			for (const int VertID : InMesh.VertexIndicesItr())
+			{
+				int32 OriginalVertID = ManifoldMapping.GetOriginalNonManifoldVertexID(VertID);
+				if (OriginalVertID != INDEX_NONE)
+				{
+					MaxOriginalVertexID = FMath::Max(MaxOriginalVertexID, OriginalVertID);
+				}
+				else
+				{
+					NumNewVertices++;
+				}
+			}
+		
+			int32 NewVertexID = MaxOriginalVertexID + 1;
+
+			InBuilder.ReserveNewVertices(NewVertexID + NumNewVertices);
+			TSet<int32> DefinedVertexIDs;
+
+			for (int VertID : InMesh.VertexIndicesItr())
+			{
+				const FVector3d Position = InMesh.GetVertex(VertID);
+				int32 OriginalVertID = ManifoldMapping.GetOriginalNonManifoldVertexID(VertID);
+				if (OriginalVertID != INDEX_NONE)
+				{
+					// FMeshDescription::CreateVertexWithID is not happy if we try to create the same vertex twice.
+					if (!DefinedVertexIDs.Contains(OriginalVertID))
+					{
+						OutVertexMappings[VertID] = InBuilder.AppendVertexWithId(OriginalVertID, Position);
+						DefinedVertexIDs.Add(OriginalVertID);
+					}
+					else
+					{
+						OutVertexMappings[VertID] = FVertexID{OriginalVertID};
+					}
+				}
+				else
+				{
+					OutVertexMappings[VertID] = InBuilder.AppendVertexWithId(NewVertexID++, Position);
+				}
+			}
+		}
+		else
+		{
+			// allocate
+			InBuilder.ReserveNewVertices(InMesh.VertexCount());
+		
+			// create "vertex buffer" in MeshDescription
+			for (int VertID : InMesh.VertexIndicesItr())
+			{
+				OutVertexMappings[VertID] = InBuilder.AppendVertex(InMesh.GetVertex(VertID));
+			}
+		}
+	}
 }
 
 
@@ -452,7 +522,7 @@ void FDynamicMeshToMeshDescription::Convert_NoAttributes(const FDynamicMesh3* Me
 	Builder.SetMeshDescription(&MeshOut);
 
 	Builder.SuspendMeshDescriptionIndexing();
-	const int32 UVLayerIndex = 0;
+	constexpr int32 UVLayerIndex = 0;
 	Builder.SetNumUVLayers(1);
 	Builder.ReserveNewUVs(MeshIn->VertexCount(), UVLayerIndex);
 
@@ -463,15 +533,11 @@ void FDynamicMeshToMeshDescription::Convert_NoAttributes(const FDynamicMesh3* Me
 		bCopyGroupToPolyGroup = true;
 	}
 
-	// create vertices
-	TArray<FVertexID> MapV; 
-	MapV.SetNum(MeshIn->MaxVertexID());
-	Builder.ReserveNewVertices(MeshIn->VertexCount());
-	for (int VertID : MeshIn->VertexIndicesItr())
-	{
-		MapV[VertID] = Builder.AppendVertex((FVector)MeshIn->GetVertex(VertID));
-	}
-
+	// Convert the vertices and construct a mapping table between the dynamic mesh and the MeshDescription, optionally applying
+	// an undoing conversion for the non-manifold fixes applied when converting from a MeshDescription to dynamic mesh.
+	TArray<FVertexID> MapV;
+	DynamicMeshToMeshDescriptionConversionHelper::ConvertVertices(*MeshIn, Builder, ConversionOptions.bConvertBackToNonManifold, MapV);
+	
 	FPolygonGroupID AllGroupID = Builder.AppendPolygonGroup();
 
 	// create new instances when seen
@@ -792,15 +858,10 @@ void FDynamicMeshToMeshDescription::Convert_NoSharedInstances(const FDynamicMesh
 	// disable indexing during the full build of the mesh
 	Builder.SuspendMeshDescriptionIndexing();
 
-	// allocate
-	Builder.ReserveNewVertices(MeshIn->VertexCount());
-
-	// create "vertex buffer" in MeshDescription
-	TArray<FVertexID> MapV; MapV.SetNum(MeshIn->MaxVertexID());
-	for (int VertID : MeshIn->VertexIndicesItr())
-	{
-		MapV[VertID] = Builder.AppendVertex((FVector)MeshIn->GetVertex(VertID));
-	}
+	// Convert the vertices and construct a mapping table between the dynamic mesh and the MeshDescription, optionally applying
+	// an undoing conversion for the non-manifold fixes applied when converting from a MeshDescription to dynamic mesh.
+	TArray<FVertexID> MapV;
+	DynamicMeshToMeshDescriptionConversionHelper::ConvertVertices(*MeshIn, Builder, ConversionOptions.bConvertBackToNonManifold, MapV);
 
 	// create UV vertex buffer in MeshDescription
 	Builder.SetNumUVLayers(NumUVLayers);
@@ -1163,4 +1224,3 @@ void FDynamicMeshToMeshDescription::ApplyVertexColorTransform(FVector4f& Color)
 		LinearColors::SRGBToLinear(Color);
 	}
 }
-
diff --git a/Engine/Source/Runtime/MeshConversion/Private/MeshDescriptionBuilder.cpp b/Engine/Source/Runtime/MeshConversion/Private/MeshDescriptionBuilder.cpp
--- a/Engine/Source/Runtime/MeshConversion/Private/MeshDescriptionBuilder.cpp
+++ b/Engine/Source/Runtime/MeshConversion/Private/MeshDescriptionBuilder.cpp
@@ -75,6 +75,14 @@ FVertexID FMeshDescriptionBuilder::AppendVertex(const FVector& Position)
 	return VertexID;
 }
 
+FVertexID FMeshDescriptionBuilder::AppendVertexWithId(int32 NewVertexID, const FVector& Position)
+{
+	FVertexID VertexID{NewVertexID};
+	MeshDescription->CreateVertexWithID(VertexID);
+	VertexPositions.Set(VertexID, FVector3f(Position));	//LWC_TODO: Precision loss
+	return VertexID;
+}
+
 
 FPolygonGroupID FMeshDescriptionBuilder::AppendPolygonGroup(FName MaterialSlotName)
 {
diff --git a/Engine/Source/Runtime/MeshConversion/Public/MeshConversionOptions.h b/Engine/Source/Runtime/MeshConversion/Public/MeshConversionOptions.h
--- a/Engine/Source/Runtime/MeshConversion/Public/MeshConversionOptions.h
+++ b/Engine/Source/Runtime/MeshConversion/Public/MeshConversionOptions.h
@@ -30,6 +30,9 @@ public:
 	/** Should Vertex Colors of MeshDescription be transformed from SRGB to Linear */
 	bool bTransformVtxColorsSRGBToLinear = true;
 
+	/** Should the mesh be put back in the original non-manifold state, if the information is present */
+	bool bConvertBackToNonManifold = false;
+
 	//
 	// utility functions for common configuration cases
 	//
diff --git a/Engine/Source/Runtime/MeshConversion/Public/MeshDescriptionBuilder.h b/Engine/Source/Runtime/MeshConversion/Public/MeshDescriptionBuilder.h
--- a/Engine/Source/Runtime/MeshConversion/Public/MeshDescriptionBuilder.h
+++ b/Engine/Source/Runtime/MeshConversion/Public/MeshDescriptionBuilder.h
@@ -42,6 +42,9 @@ public:
 	/** Append vertex and return new vertex ID */
 	MESHCONVERSION_API FVertexID AppendVertex(const FVector& Position);
 
+	/** Append vertex with a specific ID and return it */
+	MESHCONVERSION_API FVertexID AppendVertexWithId(int32 NewVertexID, const FVector& Position);
+
 	/** Return position of vertex */
 	MESHCONVERSION_API FVector GetPosition(const FVertexID& VertexID);
 
