make dynamicmesheditor's internal uses of split bowties reuse its temp array storage for one ring triangle traversal

diff --git a/Engine/Source/Runtime/GeometryCore/Private/DynamicMeshEditor.cpp b/Engine/Source/Runtime/GeometryCore/Private/DynamicMeshEditor.cpp
--- a/Engine/Source/Runtime/GeometryCore/Private/DynamicMeshEditor.cpp
+++ b/Engine/Source/Runtime/GeometryCore/Private/DynamicMeshEditor.cpp
@@ -847,16 +847,16 @@ void FDynamicMeshEditor::DisconnectTriangles(const TArray<int>& Triangles, bool
 	}
 	if (bPreventBowties)
 	{
-		FDynamicMeshEditResult Result;
+		// reusable arrays for SplitBowtiesInternal
+		TArray<int> BowtieTriangles, BowtieContiguousGroupLengths;
+		TArray<bool> BowtieGroupIsLoop;
 		for (int VID : OldVertsThatSplit)
 		{
-			SplitBowties(VID, Result);
-			Result.Reset(); // don't actually keep results; they are not used in this fn
+			SplitBowtiesInternal(VID, BowtieTriangles, BowtieContiguousGroupLengths, BowtieGroupIsLoop, nullptr);
 		}
 		for (int VID : NewVerts)
 		{
-			SplitBowties(VID, Result);
-			Result.Reset(); // don't actually keep results; they are not used in this fn
+			SplitBowtiesInternal(VID, BowtieTriangles, BowtieContiguousGroupLengths, BowtieGroupIsLoop, nullptr);
 		}
 	}
 }
@@ -922,7 +922,9 @@ bool FDynamicMeshEditor::DisconnectTrianglesAlongEdges(const TSet<int32>& EdgeID
 			AddedVertexIDs->Add(SplitInfo.NewVertex);
 		}
 	};
-
+	// reusable arrays for bowtie and edge disconnecting logic below
+	TArray<int> BowtieTriangles, BowtieContiguousGroupLengths;
+	TArray<bool> BowtieGroupIsLoop;
 	TArray<int32> VertTIDsToUseForSplit[2];
 	for (int32 EID : EdgeIDs)
 	{
@@ -938,7 +940,7 @@ bool FDynamicMeshEditor::DisconnectTrianglesAlongEdges(const TSet<int32>& EdgeID
 			// selecting edges adjacent to a bowtie.
 			for (int i = 0; i < 2; ++i)
 			{
-				SplitBowties(EdgeVIDs[i], AddedVertexIDs);
+				SplitBowtiesInternal(EdgeVIDs[i], BowtieTriangles, BowtieContiguousGroupLengths, BowtieGroupIsLoop, AddedVertexIDs);
 			}
 			continue;
 		}
@@ -976,7 +978,7 @@ bool FDynamicMeshEditor::DisconnectTrianglesAlongEdges(const TSet<int32>& EdgeID
 			if (ensure(PotentialBowtieVID != IndexConstants::InvalidID)
 				&& Mesh->IsBowtieVertex(PotentialBowtieVID))
 			{
-				SplitBowties(PotentialBowtieVID, AddedVertexIDs);
+				SplitBowtiesInternal(PotentialBowtieVID, BowtieTriangles, BowtieContiguousGroupLengths, BowtieGroupIsLoop, AddedVertexIDs);
 			}
 		}
 		else
@@ -997,6 +999,8 @@ bool FDynamicMeshEditor::DisconnectTrianglesAlongEdges(const TSet<int32>& EdgeID
 
 void FDynamicMeshEditor::SplitBowties(FDynamicMeshEditResult& ResultOut)
 {
+	TArray<int> Triangles, ContiguousGroupLengths;
+	TArray<bool> GroupIsLoop;
 	ResultOut.Reset();
 	TSet<int> AddedVerticesWithIDLessThanMax; // added vertices that we can't filter just by checking against original max id; this will be empty for compact meshes
 	for (int VertexID = 0, OriginalMaxID = Mesh->MaxVertexID(); VertexID < OriginalMaxID; VertexID++)
@@ -1006,8 +1010,7 @@ void FDynamicMeshEditor::SplitBowties(FDynamicMeshEditResult& ResultOut)
 			continue;
 		}
 		int32 NumVertsBefore = ResultOut.NewVertices.Num();
-		// TODO: may be faster to inline this call to reuse the contiguous triangle arrays?
-		SplitBowties(VertexID, ResultOut);
+		SplitBowtiesInternal(VertexID, Triangles, ContiguousGroupLengths, GroupIsLoop, &ResultOut.NewVertices);
 		for (int Idx = NumVertsBefore; Idx < ResultOut.NewVertices.Num(); Idx++)
 		{
 			if (ResultOut.NewVertices[Idx] < OriginalMaxID)
@@ -1020,10 +1023,8 @@ void FDynamicMeshEditor::SplitBowties(FDynamicMeshEditResult& ResultOut)
 
 
 
-void FDynamicMeshEditor::SplitBowties(int VertexID, TArray<int32>* NewVertices)
+void FDynamicMeshEditor::SplitBowtiesInternal(int VertexID, TArray<int>& TrianglesOut, TArray<int>& ContiguousGroupLengths, TArray<bool>& GroupIsLoop, TArray<int32>* NewVertices)
 {
-	TArray<int> TrianglesOut, ContiguousGroupLengths;
-	TArray<bool> GroupIsLoop;
 	DynamicMeshInfo::FVertexSplitInfo SplitInfo;
 	if (ensure(EMeshResult::Ok == Mesh->GetVtxContiguousTriangles(VertexID, TrianglesOut, ContiguousGroupLengths, GroupIsLoop)))
 	{
@@ -1044,6 +1045,8 @@ void FDynamicMeshEditor::SplitBowties(int VertexID, TArray<int32>* NewVertices)
 
 void FDynamicMeshEditor::SplitBowtiesAtTriangles(const TArray<int32>& TriangleIDs, FDynamicMeshEditResult& ResultOut)
 {
+	TArray<int> Triangles, ContiguousGroupLengths;
+	TArray<bool> GroupIsLoop;
 	TSet<int32> VidsProcessed;
 	for (int32 Tid : TriangleIDs)
 	{
@@ -1055,7 +1058,7 @@ void FDynamicMeshEditor::SplitBowtiesAtTriangles(const TArray<int32>& TriangleID
 			VidsProcessed.Add(Vid, &bWasAlreadyProcessed);
 			if (!bWasAlreadyProcessed)
 			{
-				SplitBowties(Vid, ResultOut);
+				SplitBowtiesInternal(Vid, Triangles, ContiguousGroupLengths, GroupIsLoop, &ResultOut.NewVertices);
 			}
 		}
 	}
diff --git a/Engine/Source/Runtime/GeometryCore/Public/DynamicMeshEditor.h b/Engine/Source/Runtime/GeometryCore/Public/DynamicMeshEditor.h
--- a/Engine/Source/Runtime/GeometryCore/Public/DynamicMeshEditor.h
+++ b/Engine/Source/Runtime/GeometryCore/Public/DynamicMeshEditor.h
@@ -248,7 +248,17 @@ public:
 	/**
 	 * Splits any bowties specifically on the given vertex, and (if not null) updates (does not reset!) NewVertices with any added vertices
 	 */
-	GEOMETRYCORE_API void SplitBowties(int VertexID, TArray<int32>* NewVertices = nullptr);
+	void SplitBowties(int VertexID, TArray<int32>* NewVertices = nullptr)
+	{
+		TArray<int> Triangles, ContiguousGroupLengths;
+		TArray<bool> GroupIsLoop;
+		SplitBowtiesInternal(VertexID, Triangles, ContiguousGroupLengths, GroupIsLoop, NewVertices);
+	}
+
+private:
+	// Internal version of split bowties that takes some arrays for reusable storage
+	GEOMETRYCORE_API void SplitBowtiesInternal(int VertexID, TArray<int>& Triangles, TArray<int>& ContiguousGroupLengths, TArray<bool>& GroupIsLoop, TArray<int32>* NewVertices = nullptr);
+public:
 
 	/**
 	 * Splits any bowties specifically on the given vertex, and updates (does not reset!) ResultOut with any added vertices
