make dynamic mesh actually call the OnRemoveVertex method when doing remove vertex operations, + make the attribute set consistently forward all On* method calls to generic attribute layers

diff --git a/Engine/Source/Runtime/GeometryCore/Private/DynamicMesh/DynamicMesh3_Edits.cpp b/Engine/Source/Runtime/GeometryCore/Private/DynamicMesh/DynamicMesh3_Edits.cpp
--- a/Engine/Source/Runtime/GeometryCore/Private/DynamicMesh/DynamicMesh3_Edits.cpp
+++ b/Engine/Source/Runtime/GeometryCore/Private/DynamicMesh/DynamicMesh3_Edits.cpp
@@ -294,6 +294,11 @@ void FDynamicMesh3::RemoveUnusedVertices()
 		if (VertexRefCounts.GetRefCount(VID) == 1)
 		{
 			VertexRefCounts.Decrement(VID);
+
+			if (HasAttributes())
+			{
+				Attributes()->OnRemoveVertex(VID);
+			}
 			checkSlow(VertexRefCounts.IsValid(VID) == false); // vertex should now not be valid
 			checkSlow(VertexEdgeLists.GetCount(VID) == 0); // vertex should not have had any edges attached
 		}
@@ -649,6 +654,11 @@ EMeshResult FDynamicMesh3::RemoveVertex(int vID, bool bPreserveManifold)
 	ensure(VertexRefCounts.IsValid(vID) == false);
 	VertexEdgeLists.Clear(vID);
 
+	if (HasAttributes())
+	{
+		Attributes()->OnRemoveVertex(vID);
+	}
+
 	UpdateChangeStamps(true, true);
 	return EMeshResult::Ok;
 }
diff --git a/Engine/Source/Runtime/GeometryCore/Private/DynamicMesh/DynamicMeshAttributeSet.cpp b/Engine/Source/Runtime/GeometryCore/Private/DynamicMesh/DynamicMeshAttributeSet.cpp
--- a/Engine/Source/Runtime/GeometryCore/Private/DynamicMesh/DynamicMeshAttributeSet.cpp
+++ b/Engine/Source/Runtime/GeometryCore/Private/DynamicMesh/DynamicMeshAttributeSet.cpp
@@ -807,6 +807,11 @@ void FDynamicMeshAttributeSet::OnNewVertex(int VertexID, bool bInserted)
 		float NewWeight = 0.0f;
 		WeightLayer.SetNewValue(VertexID, &NewWeight);
 	}
+
+	for (TPair<FName, TUniquePtr<FDynamicMeshAttributeBase>>& AttribPair : GenericAttributes)
+	{
+		AttribPair.Value->OnNewVertex(VertexID, bInserted);
+	}
 }
 
 
@@ -814,9 +819,11 @@ void FDynamicMeshAttributeSet::OnRemoveVertex(int VertexID)
 {
 	FDynamicMeshAttributeSetBase::OnRemoveVertex(VertexID);
 
-	for (FDynamicMeshWeightAttribute& WeightLayer : WeightLayers)
+	// no standard attributes need to update on vertex removals, but custom generic attributes may need to
+
+	for (TPair<FName, TUniquePtr<FDynamicMeshAttributeBase>>& AttribPair : GenericAttributes)
 	{
-		WeightLayer.OnRemoveVertex(VertexID);
+		AttribPair.Value->OnRemoveVertex(VertexID);
 	}
 }
 
@@ -848,6 +855,11 @@ void FDynamicMeshAttributeSet::OnNewTriangle(int TriangleID, bool bInserted)
 		int32 NewGroup = 0;
 		PolygroupLayer.SetNewValue(TriangleID, &NewGroup);
 	}
+
+	for (TPair<FName, TUniquePtr<FDynamicMeshAttributeBase>>& AttribPair : GenericAttributes)
+	{
+		AttribPair.Value->OnNewTriangle(TriangleID, bInserted);
+	}
 }
 
 
@@ -869,6 +881,11 @@ void FDynamicMeshAttributeSet::OnRemoveTriangle(int TriangleID)
 	}
 
 	// has no effect on MaterialIDAttrib
+
+	for (TPair<FName, TUniquePtr<FDynamicMeshAttributeBase>>& AttribPair : GenericAttributes)
+	{
+		AttribPair.Value->OnRemoveTriangle(TriangleID);
+	}
 }
 
 void FDynamicMeshAttributeSet::OnReverseTriOrientation(int TriangleID)
@@ -888,6 +905,11 @@ void FDynamicMeshAttributeSet::OnReverseTriOrientation(int TriangleID)
 		ColorLayer->OnReverseTriOrientation(TriangleID);
 	}
 	// has no effect on MaterialIDAttrib
+
+	for (TPair<FName, TUniquePtr<FDynamicMeshAttributeBase>>& AttribPair : GenericAttributes)
+	{
+		AttribPair.Value->OnReverseTriOrientation(TriangleID);
+	}
 }
 
 void FDynamicMeshAttributeSet::OnSplitEdge(const FDynamicMesh3::FEdgeSplitInfo& SplitInfo)
@@ -921,6 +943,10 @@ void FDynamicMeshAttributeSet::OnSplitEdge(const FDynamicMesh3::FEdgeSplitInfo&
 		WeightLayer.OnSplitEdge(SplitInfo);
 	}
 
+	for (TPair<FName, TUniquePtr<FDynamicMeshAttributeBase>>& AttribPair : GenericAttributes)
+	{
+		AttribPair.Value->OnSplitEdge(SplitInfo);
+	}
 }
 
 void FDynamicMeshAttributeSet::OnFlipEdge(const FDynamicMesh3::FEdgeFlipInfo & flipInfo)
@@ -954,6 +980,10 @@ void FDynamicMeshAttributeSet::OnFlipEdge(const FDynamicMesh3::FEdgeFlipInfo & f
 		WeightLayer.OnFlipEdge(flipInfo);
 	}
 
+	for (TPair<FName, TUniquePtr<FDynamicMeshAttributeBase>>& AttribPair : GenericAttributes)
+	{
+		AttribPair.Value->OnFlipEdge(flipInfo);
+	}
 }
 
 
@@ -988,6 +1018,10 @@ void FDynamicMeshAttributeSet::OnCollapseEdge(const FDynamicMesh3::FEdgeCollapse
 		WeightLayer.OnCollapseEdge(collapseInfo);
 	}
 
+	for (TPair<FName, TUniquePtr<FDynamicMeshAttributeBase>>& AttribPair : GenericAttributes)
+	{
+		AttribPair.Value->OnCollapseEdge(collapseInfo);
+	}
 }
 
 void FDynamicMeshAttributeSet::OnPokeTriangle(const FDynamicMesh3::FPokeTriangleInfo & pokeInfo)
@@ -1020,6 +1054,11 @@ void FDynamicMeshAttributeSet::OnPokeTriangle(const FDynamicMesh3::FPokeTriangle
 	{
 		WeightLayer.OnPokeTriangle(pokeInfo);
 	}
+
+	for (TPair<FName, TUniquePtr<FDynamicMeshAttributeBase>>& AttribPair : GenericAttributes)
+	{
+		AttribPair.Value->OnPokeTriangle(pokeInfo);
+	}
 }
 
 void FDynamicMeshAttributeSet::OnMergeEdges(const FDynamicMesh3::FMergeEdgesInfo & mergeInfo)
@@ -1052,6 +1091,11 @@ void FDynamicMeshAttributeSet::OnMergeEdges(const FDynamicMesh3::FMergeEdgesInfo
 	{
 		WeightLayer.OnMergeEdges(mergeInfo);
 	}
+
+	for (TPair<FName, TUniquePtr<FDynamicMeshAttributeBase>>& AttribPair : GenericAttributes)
+	{
+		AttribPair.Value->OnMergeEdges(mergeInfo);
+	}
 }
 
 void FDynamicMeshAttributeSet::OnMergeVertices(const DynamicMeshInfo::FMergeVerticesInfo& MergeInfo)
@@ -1084,6 +1128,11 @@ void FDynamicMeshAttributeSet::OnMergeVertices(const DynamicMeshInfo::FMergeVert
 	{
 		WeightLayer.OnMergeVertices(MergeInfo);
 	}
+
+	for (TPair<FName, TUniquePtr<FDynamicMeshAttributeBase>>& AttribPair : GenericAttributes)
+	{
+		AttribPair.Value->OnMergeVertices(MergeInfo);
+	}
 }
 
 void FDynamicMeshAttributeSet::OnSplitVertex(const DynamicMeshInfo::FVertexSplitInfo& SplitInfo, const TArrayView<const int>& TrianglesToUpdate)
@@ -1116,6 +1165,11 @@ void FDynamicMeshAttributeSet::OnSplitVertex(const DynamicMeshInfo::FVertexSplit
 	{
 		WeightLayer.OnSplitVertex(SplitInfo, TrianglesToUpdate);
 	}
+
+	for (TPair<FName, TUniquePtr<FDynamicMeshAttributeBase>>& AttribPair : GenericAttributes)
+	{
+		AttribPair.Value->OnSplitVertex(SplitInfo, TrianglesToUpdate);
+	}
 }
 
 bool FDynamicMeshAttributeSet::IsSameAs(const FDynamicMeshAttributeSet& Other, bool bIgnoreDataLayout) const
