change uniform random sampling method used by mesh surface point sampling to not need per-tri frames and 2D projections, and not need to do rejection sampling for the triangle points or pt->barycentric coordinate conversion

+ add a concept of sampling method version so that users can request the old sampling, to support use cases like PCG where it is important to always generate the same result across UE versions

diff --git a/Engine/Plugins/Runtime/GeometryScripting/Source/GeometryScriptingCore/Private/MeshSamplingFunctions.cpp b/Engine/Plugins/Runtime/GeometryScripting/Source/GeometryScriptingCore/Private/MeshSamplingFunctions.cpp
--- a/Engine/Plugins/Runtime/GeometryScripting/Source/GeometryScriptingCore/Private/MeshSamplingFunctions.cpp
+++ b/Engine/Plugins/Runtime/GeometryScripting/Source/GeometryScriptingCore/Private/MeshSamplingFunctions.cpp
@@ -35,6 +35,7 @@ UDynamicMesh* UGeometryScriptLibrary_MeshSamplingFunctions::ComputePointSampling
 	Sampler.MaxSamples = Options.MaxNumSamples;
 	Sampler.SubSampleDensity = Options.SubSampleDensity;
 	Sampler.RandomSeed = Options.RandomSeed;
+	Sampler.SamplingMethodVersion = Options.SamplingMethodVersion;
 
 	TargetMesh->ProcessMesh([&](const FDynamicMesh3& ReadMesh)
 	{
@@ -78,6 +79,7 @@ UDynamicMesh* UGeometryScriptLibrary_MeshSamplingFunctions::ComputeNonUniformPoi
 	Sampler.MaxSamples = Options.MaxNumSamples;
 	Sampler.SubSampleDensity = Options.SubSampleDensity;
 	Sampler.RandomSeed = Options.RandomSeed;
+	Sampler.SamplingMethodVersion = Options.SamplingMethodVersion;
 	if (NonUniformOptions.MaxSamplingRadius > Options.SamplingRadius)
 	{
 		Sampler.MaxSampleRadius = NonUniformOptions.MaxSamplingRadius;
@@ -131,6 +133,7 @@ UDynamicMesh* UGeometryScriptLibrary_MeshSamplingFunctions::ComputeVertexWeighte
 	Sampler.MaxSamples = Options.MaxNumSamples;
 	Sampler.SubSampleDensity = Options.SubSampleDensity;
 	Sampler.RandomSeed = Options.RandomSeed;
+	Sampler.SamplingMethodVersion = Options.SamplingMethodVersion;
 	bool bIsNonUniform = false;
 	if (NonUniformOptions.MaxSamplingRadius > Options.SamplingRadius)
 	{
diff --git a/Engine/Plugins/Runtime/GeometryScripting/Source/GeometryScriptingCore/Public/GeometryScript/MeshSamplingFunctions.h b/Engine/Plugins/Runtime/GeometryScripting/Source/GeometryScriptingCore/Public/GeometryScript/MeshSamplingFunctions.h
--- a/Engine/Plugins/Runtime/GeometryScripting/Source/GeometryScriptingCore/Public/GeometryScript/MeshSamplingFunctions.h
+++ b/Engine/Plugins/Runtime/GeometryScripting/Source/GeometryScriptingCore/Public/GeometryScript/MeshSamplingFunctions.h
@@ -30,6 +30,16 @@ public:
 	/** Density of subsampling used in Poisson strategy. Larger numbers mean "more accurate" (but slower) results. */
 	UPROPERTY(EditAnywhere, BlueprintReadWrite, Category = Options)
 	double SubSampleDensity = 10;
+
+	/**
+	 * If < 0, the latest, recommended sampling methods will be used. Otherwise, requests a specific sampling method version.
+	 * Set this for more consistent results across UE versions, at the risk of worse performance or quality.
+	 * Valid versions are:
+	 * Method 0: UE 5.5 and earlier sampling method. Slower initial (dense) point sampling, less robust to degenerate triangles.
+	 * Method 1: Currently the default method.
+	 */
+	UPROPERTY(EditAnywhere, BlueprintReadWrite, Category = Options, AdvancedDisplay)
+	int SamplingMethodVersion = INDEX_NONE;
 };
 
 /** Control how active Weights are used to affect point radius */
diff --git a/Engine/Source/Runtime/GeometryCore/Private/Sampling/MeshSurfacePointSampling.cpp b/Engine/Source/Runtime/GeometryCore/Private/Sampling/MeshSurfacePointSampling.cpp
--- a/Engine/Source/Runtime/GeometryCore/Private/Sampling/MeshSurfacePointSampling.cpp
+++ b/Engine/Source/Runtime/GeometryCore/Private/Sampling/MeshSurfacePointSampling.cpp
@@ -20,30 +20,37 @@ struct FTriangleInfoCache
 {
 	TArray<FVector3d> TriNormals;
 	TArray<double> TriAreas;
-	TArray<FFrame3d> TriFrames;
-	TArray<FTriangle2d> UVTriangles;
+	// Note: The below two arrays are only used for the legacy sampling method 0
+	TArray<FFrame3d> Legacy_TriFrames;
+	TArray<FTriangle2d> Legacy_UVTriangles;
 
 	double TotalArea;
 
 	template<typename MeshType>
-	void InitializeForTriangleSet(const MeshType& SampleMesh)
+	void InitializeForTriangleSet(const MeshType& SampleMesh, int32 SamplingMethodVersion)
 	{
 		TriNormals.SetNumUninitialized(SampleMesh.MaxTriangleID());
 		TriAreas.SetNumZeroed(SampleMesh.MaxTriangleID()); // note: zero the memory so we can include skipped triangle IDs in the sum below
-		TriFrames.SetNumUninitialized(SampleMesh.MaxTriangleID());
-		UVTriangles.SetNumUninitialized(SampleMesh.MaxTriangleID());
+		if (SamplingMethodVersion == 0)
+		{
+			Legacy_TriFrames.SetNumUninitialized(SampleMesh.MaxTriangleID());
+			Legacy_UVTriangles.SetNumUninitialized(SampleMesh.MaxTriangleID());
+		}
 				
 		ParallelFor(SampleMesh.MaxTriangleID(), [&](int32 tid)
 		{
 			if (SampleMesh.IsTriangle(tid))
 			{
-				FVector3d A, B, C, Centroid;
+				FVector3d A, B, C;
 				SampleMesh.GetTriVertices(tid, A, B, C);
-				Centroid = (A + B + C) / 3.0;
 				TriNormals[tid] = VectorUtil::NormalArea(A, B, C, TriAreas[tid]);
-				TriFrames[tid] = FFrame3d(Centroid, TriNormals[tid]);
-				UVTriangles[tid] = FTriangle2d(
-					TriFrames[tid].ToPlaneUV(A), TriFrames[tid].ToPlaneUV(B), TriFrames[tid].ToPlaneUV(C));
+				if (SamplingMethodVersion == 0)
+				{
+					FVector3d Centroid;
+					Centroid = (A + B + C) / 3.0;
+					Legacy_TriFrames[tid] = FFrame3d(Centroid, TriNormals[tid]);
+					Legacy_UVTriangles[tid] = FTriangle2d(Legacy_TriFrames[tid].ToPlaneUV(A), Legacy_TriFrames[tid].ToPlaneUV(B), Legacy_TriFrames[tid].ToPlaneUV(C));
+				}
 			}
 		});
 		// Note: This sum includes areas of skipped triangle IDs, which are zero by initialization above
@@ -91,15 +98,6 @@ struct FPerTriangleDensePointSampling
 	};
 	TArray<FTriangleSubArray> TriSubArrays;
 
-	template<typename MeshType>
-	void InitializeForTriangleSet(
-		const MeshType& SampleMesh, 
-		const FTriangleInfoCache& TriInfo, 
-		double DenseSampleArea, 
-		int RandomSeed,
-		FDenseSamplePointSet& PointSetOut)
-	{
-	}
 
 	template<typename MeshType>
 	void InitializeForTriangleSet(
@@ -108,7 +106,8 @@ struct FPerTriangleDensePointSampling
 		double DenseSampleArea, 
 		int RandomSeed,
 		const FNonUniformSamplingConfig& NonUniformConfig,
-		FDenseSamplePointSet& PointSetOut )
+		FDenseSamplePointSet& PointSetOut,
+		int32 SamplingMethodVersion)
 	{
 		// figure out how many samples in each triangle, and assign each triangle a starting index into PointSetOut arrays
 		TriSubArrays.SetNumUninitialized(SampleMesh.MaxTriangleID());
@@ -144,30 +143,64 @@ struct FPerTriangleDensePointSampling
 			int NumSamples = TriSubArrays[tid].NumSamples;
 			if ( NumSamples == 0 ) return;
 			int StartIndex = TriSubArrays[tid].StartIndex;
-
-			const FFrame3d& ProjectFrame = TriInfo.TriFrames[tid];
-			const FTriangle2d& TriUV = TriInfo.UVTriangles[tid];
-
-			// generate uniform random point in 2D quadrilateral  (http://mathworld.wolfram.com/TrianglePointPicking.html)
-			FVector2d V1 = TriUV.V[1] - TriUV.V[0];
-			FVector2d V2 = TriUV.V[2] - TriUV.V[0];
-
-			int NumGenerated = 0;
-			FRandomStream RandomStream(tid + RandomSeed);
-			while (NumGenerated < NumSamples)
+			
+			// Legacy sampling method using rejection sampling, to support any application that
+			// depends on the sampling pattern not changing (e.g., PCG sampling nodes expect deterministic results)
+			// If changing, be sure the sampling pattern remains the same.
+			if (SamplingMethodVersion == 0)
+			{
+				const FFrame3d& ProjectFrame = TriInfo.Legacy_TriFrames[tid];
+				const FTriangle2d& TriUV = TriInfo.Legacy_UVTriangles[tid];
+				
+				// generate uniform random point in 2D quadrilateral  (http://mathworld.wolfram.com/TrianglePointPicking.html)
+				// Note: VectorUtil::UniformSampleTrianglePoint has a better (non-rejection) version of this sampling method
+				// (but it would generate a different pattern, so we can't use it here)
+				FVector2d V1 = TriUV.V[1] - TriUV.V[0];
+				FVector2d V2 = TriUV.V[2] - TriUV.V[0];
+
+				int NumGenerated = 0;
+				FRandomStream RandomStream(tid + RandomSeed);
+				while (NumGenerated < NumSamples)
+				{
+					double a1 = RandomStream.GetFraction();
+					double a2 = RandomStream.GetFraction();
+					FVector2d PointUV = TriUV.V[0] + a1 * V1 + a2 * V2;
+					if (TriUV.IsInside(PointUV))
+					{
+						FVector3d Position = ProjectFrame.FromPlaneUV(PointUV, 2);
+						PointSetOut.DensePoints[StartIndex+NumGenerated] = Position;
+						PointSetOut.Triangles[StartIndex+NumGenerated] = tid;
+
+						if (bComputeWeights)
+						{
+							FVector3d BaryCoords = TriUV.GetBarycentricCoords(PointUV);
+							PointSetOut.Weights[StartIndex+NumGenerated] = 
+								NonUniformConfig.WeightFunction.GetValue()(tid, Position, BaryCoords);
+						}
+
+						NumGenerated++;
+					}
+				}
+			}
+			else // Sampling method > 0, sample uniform random barycentric coordinates directly
 			{
-				double a1 = RandomStream.GetFraction();
-				double a2 = RandomStream.GetFraction();
-				FVector2d PointUV = TriUV.V[0] + a1 * V1 + a2 * V2;
-				if (TriUV.IsInside(PointUV))
+				FVector3d A, B, C;
+				SampleMesh.GetTriVertices(tid, A, B, C);
+
+				int NumGenerated = 0;
+				FRandomStream RandomStream(tid + RandomSeed);
+				while (NumGenerated < NumSamples)
 				{
-					FVector3d Position = ProjectFrame.FromPlaneUV(PointUV, 2);
-					PointSetOut.DensePoints[StartIndex+NumGenerated] = Position;
-					PointSetOut.Triangles[StartIndex+NumGenerated] = tid;
+					double a1 = RandomStream.GetFraction();
+					double a2 = RandomStream.GetFraction();
 
+					FVector3d BaryCoords = VectorUtil::UniformSampleTriangleBarycentricCoords(a1, a2);
+					FVector3d Position = BaryCoords.X * A + BaryCoords.Y * B + BaryCoords.Z * C;
+					PointSetOut.DensePoints[StartIndex + NumGenerated] = Position;
+					PointSetOut.Triangles[StartIndex + NumGenerated] = tid;
+				
 					if (bComputeWeights)
 					{
-						FVector3d BaryCoords = TriUV.GetBarycentricCoords(PointUV);
 						PointSetOut.Weights[StartIndex+NumGenerated] = 
 							NonUniformConfig.WeightFunction.GetValue()(tid, Position, BaryCoords);
 					}
@@ -188,10 +221,11 @@ void ConstructDenseUniformMeshPointSampling(
 	int RandomSeed,
 	const FNonUniformSamplingConfig& NonUniformConfig,
 	int MaxNumDenseSamples,
-	FDenseSamplePointSet& DensePointSetOut )
+	FDenseSamplePointSet& DensePointSetOut,
+	int32 SamplingMethodVersion)
 {
 	FTriangleInfoCache TriInfoCache;
-	TriInfoCache.InitializeForTriangleSet<MeshType>(SampleMesh);
+	TriInfoCache.InitializeForTriangleSet<MeshType>(SampleMesh, SamplingMethodVersion);
 
 	// compute mesh bounds in background thread. ParallelFor could make this faster (eg FDynamicMesh3::GetBounds) but it's going to take longer than the other steps below anyway
 	UE::Tasks::FTask BoundsTask = UE::Tasks::Launch(UE_SOURCE_LOCATION, [&DensePointSetOut, &SampleMesh]()
@@ -219,7 +253,7 @@ void ConstructDenseUniformMeshPointSampling(
 
 	FPerTriangleDensePointSampling DensePerTriangleSampling;
 	DensePerTriangleSampling.InitializeForTriangleSet<MeshType>(SampleMesh, TriInfoCache, DenseSampleArea, RandomSeed, 
-		NonUniformConfig, DensePointSetOut);
+		NonUniformConfig, DensePointSetOut, SamplingMethodVersion);
 
 	//UE_LOG(LogTemp, Warning, TEXT("TotalArea %f  SampleArea %f EstSamples %d  DenseSamples %d   DenseArea %f  NumDensePoints %d"),
 	//	TriInfo.TotalArea, DiscArea, (int)ApproxNumUniformSamples, (int)EstNumDenseSamples, DenseSampleArea, NumDensePoints);
@@ -239,6 +273,7 @@ void UniformMeshPointSampling(
 	double SubSampleDensity,
 	int RandomSeed,
 	int MaxNumDenseSamples,
+	int32 SamplingMethodVersion,
 	FProgressCancel* Progress)
 {
 	TRACE_CPUPROFILER_EVENT_SCOPE(UniformMeshPointSampling);
@@ -254,7 +289,7 @@ void UniformMeshPointSampling(
 		TRACE_CPUPROFILER_EVENT_SCOPE(ComputeDenseSamples);
 		
 		ConstructDenseUniformMeshPointSampling<MeshType>(SampleMesh, SampleRadius, SubSampleDensity, RandomSeed, 
-			FNonUniformSamplingConfig(), MaxNumDenseSamples, DensePointSet);
+			FNonUniformSamplingConfig(), MaxNumDenseSamples, DensePointSet, SamplingMethodVersion);
 	}
 	int32 NumDensePoints = DensePointSet.MaxVertexID();
 
@@ -377,6 +412,7 @@ void NonUniformMeshPointSampling(
 	int RandomSeed,
 	const FNonUniformSamplingConfig& NonUniformConfig,
 	int MaxNumDenseSamples,
+	int32 SamplingMethodVersion,
 	FProgressCancel* Progress)
 {
 	TRACE_CPUPROFILER_EVENT_SCOPE(UniformMeshPointSampling);
@@ -390,7 +426,7 @@ void NonUniformMeshPointSampling(
 	{
 		TRACE_CPUPROFILER_EVENT_SCOPE(ComputeDenseSamples);
 		ConstructDenseUniformMeshPointSampling<MeshType>(SampleMesh, MinSampleRadius, SubSampleDensity, RandomSeed, 
-			NonUniformConfig, MaxNumDenseSamples, DensePointSet);
+			NonUniformConfig, MaxNumDenseSamples, DensePointSet, SamplingMethodVersion);
 	}
 	int32 NumDensePoints = DensePointSet.MaxVertexID();
 	bool bHaveWeights = DensePointSet.Weights.Num() > 0;
@@ -666,18 +702,18 @@ void FMeshSurfacePointSampling::ComputePoissonSampling(const FDynamicMesh3& Mesh
 			NonUniformConfig.WeightFunction = ComputeVertexWeightFunc;
 
 			UELocal::NonUniformMeshPointSampling<FDynamicMesh3>(Mesh, AddSampleFunc,
-				SampleRadius, MaxSampleRadius, MaxSamples, SubSampleDensity, RandomSeed, NonUniformConfig, MaxSubSamplePoints, Progress);
+				SampleRadius, MaxSampleRadius, MaxSamples, SubSampleDensity, RandomSeed, NonUniformConfig, MaxSubSamplePoints, SamplingMethodVersion, Progress);
 		}
 		else
 		{
 			UELocal::NonUniformMeshPointSampling<FDynamicMesh3>(Mesh, AddSampleFunc,
-				SampleRadius, MaxSampleRadius, MaxSamples, SubSampleDensity, RandomSeed, NonUniformConfig, MaxSubSamplePoints, Progress );
+				SampleRadius, MaxSampleRadius, MaxSamples, SubSampleDensity, RandomSeed, NonUniformConfig, MaxSubSamplePoints, SamplingMethodVersion, Progress );
 		}
 	}
 	else
 	{
 		UELocal::UniformMeshPointSampling<FDynamicMesh3>(Mesh, AddSampleFunc,
-			SampleRadius, MaxSamples, SubSampleDensity, RandomSeed, MaxSubSamplePoints, Progress );
+			SampleRadius, MaxSamples, SubSampleDensity, RandomSeed, MaxSubSamplePoints, SamplingMethodVersion, Progress );
 	}
 
 	if (bComputeBarycentrics)
diff --git a/Engine/Source/Runtime/GeometryCore/Public/Sampling/MeshSurfacePointSampling.h b/Engine/Source/Runtime/GeometryCore/Public/Sampling/MeshSurfacePointSampling.h
--- a/Engine/Source/Runtime/GeometryCore/Public/Sampling/MeshSurfacePointSampling.h
+++ b/Engine/Source/Runtime/GeometryCore/Public/Sampling/MeshSurfacePointSampling.h
@@ -43,6 +43,15 @@ public:
 	/** Safety threshold for maximum number of points in subsampling. */
 	int32 MaxSubSamplePoints = 50000000;
 
+	/**
+	 * If < 0, the latest, recommended sampling methods will be used. Otherwise, requests a specific sampling method version.
+	 * Set this for more consistent results across UE versions, at the risk of worse performance or quality.
+	 * Valid versions are:
+	 * Method 0: UE 5.5 and earlier sampling method. Slower initial (dense) point sampling, less robust to degenerate triangles.
+	 * Method 1: Currently the default method.
+	 */
+	int32 SamplingMethodVersion = INDEX_NONE;
+
 	//
 	// Parameters for non-uniform / variable-radius sampling. 
 	// Disabled if MaxSampleRadius <= SampleRadius
