fix 2D voronoi diagram generation failing collinear inputs
+ add a result enum to delaunay triangulation so we can know the reason for triangulation failures

diff --git a/Engine/Plugins/Runtime/GeometryScripting/Source/GeometryScriptingCore/Private/MeshPrimitiveFunctions.cpp b/Engine/Plugins/Runtime/GeometryScripting/Source/GeometryScriptingCore/Private/MeshPrimitiveFunctions.cpp
--- a/Engine/Plugins/Runtime/GeometryScripting/Source/GeometryScriptingCore/Private/MeshPrimitiveFunctions.cpp
+++ b/Engine/Plugins/Runtime/GeometryScripting/Source/GeometryScriptingCore/Private/MeshPrimitiveFunctions.cpp
@@ -1491,16 +1491,16 @@ UDynamicMesh* UGeometryScriptLibrary_MeshPrimitiveFunctions::AppendVoronoiDiagra
 		return TargetMesh;
 	}
 
-	if (VoronoiSites.Num() < 3)
+	if (VoronoiSites.Num() < 1)
 	{
-		UE::Geometry::AppendError(Debug, EGeometryScriptErrorType::InvalidInputs, LOCTEXT("PrimitiveFunctions_AppendVoronoiDiagram2D_TooFewSites", "AppendVoronoiDiagram2D: VoronoiSites array requires at least 3 positions"));
+		UE::Geometry::AppendError(Debug, EGeometryScriptErrorType::InvalidInputs, LOCTEXT("PrimitiveFunctions_AppendVoronoiDiagram2D_TooFewSites", "AppendVoronoiDiagram2D: VoronoiSites array requires at least 1 position"));
 		return TargetMesh;
 	}
 
 	UE::Geometry::FDelaunay2 Delaunay;
-	bool bTriSuccess = Delaunay.Triangulate(VoronoiSites);
+	Delaunay.Triangulate(VoronoiSites);
 
-	if (!bTriSuccess)
+	if (!Delaunay.CanComputeVoronoiCells())
 	{
 		UE::Geometry::AppendError(Debug, EGeometryScriptErrorType::OperationFailed, LOCTEXT("PrimitiveFunctions_AppendVoronoiDiagram2D_GenFailed", "AppendVoronoiDiagram2D: Voronoi diagram generation failed"));
 		return TargetMesh;
diff --git a/Engine/Source/Runtime/GeometryCore/Private/CompGeom/Delaunay2.cpp b/Engine/Source/Runtime/GeometryCore/Private/CompGeom/Delaunay2.cpp
--- a/Engine/Source/Runtime/GeometryCore/Private/CompGeom/Delaunay2.cpp
+++ b/Engine/Source/Runtime/GeometryCore/Private/CompGeom/Delaunay2.cpp
@@ -1099,6 +1099,124 @@ namespace DelaunayInternal
 		return Cells;
 	}
 
+
+	// Handle special case of voronoi diagram for collinear input, where Delaunay triangulation cannot be computed
+	template<typename RealType>
+	TArray<TArray<TVector2<RealType>>> GetVoronoiCells_Collinear(TArrayView<const TVector2<RealType>> Vertices, bool bIncludeBoundary, TAxisAlignedBox2<RealType> BoundsClip, RealType ExpandBounds)
+	{
+		TArray<TArray<TVector2<RealType>>> Result;
+
+		// In a voronoi diagram of collinear input points, all cells will touch the boundary
+		if (!bIncludeBoundary)
+		{
+			return Result;
+		}
+
+		if (!ensure(!Vertices.IsEmpty()))
+		{
+			return Result;
+		}
+
+		TVector2<RealType> V0 = Vertices[0];
+		int32 FarIdx = INDEX_NONE;
+		RealType FarDistSq = (RealType)0;
+		for (int32 Idx = 1; Idx < Vertices.Num(); ++Idx)
+		{ 
+			RealType DSq = TVector2<RealType>::DistSquared(Vertices[Idx], V0);
+			if (DSq > FarDistSq)
+			{
+				FarIdx = Idx;
+				FarDistSq = DSq;
+			}
+		}
+
+
+		TAxisAlignedBox2<RealType> Bounds;
+		for (const TVector2<RealType>& Vert : Vertices)
+		{
+			Bounds.Contain(Vert);
+		}
+		Bounds.Expand(FMath::Max((RealType)UE_SMALL_NUMBER, ExpandBounds));
+		if (!BoundsClip.IsEmpty())
+		{
+			BoundsClip.Expand(ExpandBounds);
+			Bounds.Contain(BoundsClip);
+		}
+		else
+		{
+			BoundsClip = Bounds;
+		}
+
+		Result.SetNum(Vertices.Num());
+
+		// All points had zero squared distance from the first point; treat this as a single-point 'diagram' and just generate a single cell of the entire bounds
+		if (FarIdx == INDEX_NONE)
+		{
+			TArray<TVector2<RealType>>& BoundsPoly = Result[0];
+			BoundsPoly.Add(Bounds.GetCorner(0));
+			BoundsPoly.Add(Bounds.GetCorner(1));
+			BoundsPoly.Add(Bounds.GetCorner(2));
+			BoundsPoly.Add(Bounds.GetCorner(3));
+		}
+		else
+		{
+			// Safe distance that will always move us outside the bounding box from a point inside the bounding box
+			RealType ToOutside = Bounds.DiagonalLength();
+
+			TVector2<RealType> Dir = Vertices[FarIdx] - V0;
+			Dir.Normalize();
+			TVector2<RealType> PerpToOutside = ToOutside * PerpCW(Dir);
+			TArray<int32> VertexOrder;
+			VertexOrder.SetNumUninitialized(Vertices.Num());
+			for (int32 Idx = 0; Idx < Vertices.Num(); ++Idx)
+			{
+				VertexOrder[Idx] = Idx;
+			}
+			VertexOrder.Sort([Dir, &Vertices](const int32& A, const int32& B)
+			{
+				return Vertices[A].Dot(Dir) < Vertices[B].Dot(Dir);
+			});
+			int32 PrevIdx = INDEX_NONE, NextIdx = INDEX_NONE;
+			for (int32 Idx = 0; Idx < VertexOrder.Num(); PrevIdx = Idx, Idx = NextIdx)
+			{
+				for (NextIdx = Idx + 1; NextIdx < VertexOrder.Num(); ++NextIdx)
+				{
+					if (Vertices[VertexOrder[Idx]] != Vertices[VertexOrder[NextIdx]])
+					{
+						break;
+					}
+				}
+				TVector2<RealType> CurV = Vertices[VertexOrder[Idx]];
+				TVector2<RealType> PrevMid, NextMid;
+				if (PrevIdx != INDEX_NONE)
+				{
+					PrevMid = (Vertices[VertexOrder[PrevIdx]] + CurV) * (RealType).5;
+				}
+				else
+				{
+					PrevMid = CurV - Dir * ToOutside;
+				}
+				if (NextIdx < VertexOrder.Num())
+				{
+					NextMid = (Vertices[VertexOrder[NextIdx]] + CurV) * (RealType).5;
+				}
+				else
+				{
+					NextMid = CurV + Dir * ToOutside;
+				}
+
+				TArray<TVector2<RealType>>& Polygon = Result[VertexOrder[Idx]];
+				Polygon.Add(PrevMid - PerpToOutside);
+				Polygon.Add(PrevMid + PerpToOutside);
+				Polygon.Add(NextMid + PerpToOutside);
+				Polygon.Add(NextMid - PerpToOutside);
+				CurveUtil::ClipConvexToBounds<RealType, TVector2<RealType>>(Polygon, BoundsClip.Min, BoundsClip.Max);
+			}
+		}
+
+		return Result;
+	}
+
 	template<typename RealType>
 	bool IsDelaunay(FDelaunay2Connectivity& Connectivity, TArrayView<const TVector2<RealType>> Vertices, TArrayView<const FIndex2i> SkipEdgesIn)
 	{
@@ -1391,12 +1509,15 @@ namespace DelaunayInternal
 
 	template<typename RealType>
 	bool Triangulate(FRandomStream& Random, FDelaunay2Connectivity& Connectivity,
-		TArrayView<const TVector2<RealType>> Vertices, TArrayView<const FIndex2i> Edges, bool bKeepFastEdgeAdjacencyData)
+		TArrayView<const TVector2<RealType>> Vertices, TArrayView<const FIndex2i> Edges, bool bKeepFastEdgeAdjacencyData, FDelaunay2::EResult& OutResult)
 	{
+		OutResult = FDelaunay2::EResult::Unknown;
+
 		Connectivity.Empty(Vertices.Num());
 
 		if (Vertices.Num() < 3)
 		{
+			OutResult = Vertices.IsEmpty() ? FDelaunay2::EResult::EmptyInput : FDelaunay2::EResult::Collinear;
 			return false;
 		}
 
@@ -1419,6 +1540,7 @@ namespace DelaunayInternal
 		}
 		if (BootstrapIndices[1] == -1) // all points were identical; nothing to triagulate
 		{
+			OutResult = FDelaunay2::EResult::Collinear;
 			return false;
 		}
 
@@ -1435,6 +1557,7 @@ namespace DelaunayInternal
 		}
 		if (BootstrapIndices[2] == -1) // all points were colinear; nothing to triangulate
 		{
+			OutResult = FDelaunay2::EResult::Collinear;
 			return false;
 		}
 
@@ -1473,9 +1596,16 @@ namespace DelaunayInternal
 			checkSlow(SearchTri.A != FDelaunay2Connectivity::InvalidIndex);
 		}
 
-		return ConstrainEdges(Random, Connectivity, Vertices, Edges, bKeepFastEdgeAdjacencyData);
+		bool bSuccess = ConstrainEdges(Random, Connectivity, Vertices, Edges, bKeepFastEdgeAdjacencyData);
+		if (bSuccess)
+		{
+			// Note: The ConstrainEdges implementation does not currently detect failure reasons, 
+			// so we rely on the ValidateEdgesResult method to set the result on edge failures
+			OutResult = FDelaunay2::EResult::Success; 
+		}
+		return bSuccess;
 	}
-}
+} // namespace DelaunayInternal
 
 bool FDelaunay2::Triangulate(TArrayView<const FVector2d> Vertices, TArrayView<const FIndex2i> Edges)
 {
@@ -1483,8 +1613,8 @@ bool FDelaunay2::Triangulate(TArrayView<const FVector2d> Vertices, TArrayView<co
 
 	bIsConstrained = Edges.Num() > 0;
 
-	bool bSuccess = DelaunayInternal::Triangulate<double>(RandomStream, *Connectivity, Vertices, Edges, bKeepFastEdgeAdjacencyData);
-	return bSuccess && ValidateResult(Edges);
+	bool bSuccess = DelaunayInternal::Triangulate<double>(RandomStream, *Connectivity, Vertices, Edges, bKeepFastEdgeAdjacencyData, Result);
+	return bSuccess && ValidateEdgesResult(Edges);
 }
 
 bool FDelaunay2::Triangulate(TArrayView<const FVector2f> Vertices, TArrayView<const FIndex2i> Edges)
@@ -1493,8 +1623,8 @@ bool FDelaunay2::Triangulate(TArrayView<const FVector2f> Vertices, TArrayView<co
 
 	bIsConstrained = Edges.Num() > 0;
 
-	bool bSuccess = DelaunayInternal::Triangulate<float>(RandomStream, *Connectivity, Vertices, Edges, bKeepFastEdgeAdjacencyData);
-	return bSuccess && ValidateResult(Edges);
+	bool bSuccess = DelaunayInternal::Triangulate<float>(RandomStream, *Connectivity, Vertices, Edges, bKeepFastEdgeAdjacencyData, Result);
+	return bSuccess && ValidateEdgesResult(Edges);
 }
 
 bool FDelaunay2::ConstrainEdges(TArrayView<const FVector2d> Vertices, TArrayView<const FIndex2i> Edges)
@@ -1504,7 +1634,7 @@ bool FDelaunay2::ConstrainEdges(TArrayView<const FVector2d> Vertices, TArrayView
 	checkSlow(!bAutomaticallyFixEdgesToDuplicateVertices || Connectivity->HasDuplicateTracking());
 
 	bool bSuccess = DelaunayInternal::ConstrainEdges<double>(RandomStream, *Connectivity, Vertices, Edges, bKeepFastEdgeAdjacencyData);
-	return bSuccess && ValidateResult(Edges);
+	return bSuccess && ValidateEdgesResult(Edges);
 }
 
 bool FDelaunay2::ConstrainEdges(TArrayView<const FVector2f> Vertices, TArrayView<const FIndex2i> Edges)
@@ -1514,7 +1644,7 @@ bool FDelaunay2::ConstrainEdges(TArrayView<const FVector2f> Vertices, TArrayView
 	checkSlow(!bAutomaticallyFixEdgesToDuplicateVertices || Connectivity->HasDuplicateTracking());
 
 	bool bSuccess = DelaunayInternal::ConstrainEdges<float>(RandomStream, *Connectivity, Vertices, Edges, bKeepFastEdgeAdjacencyData);
-	return bSuccess && ValidateResult(Edges);
+	return bSuccess && ValidateEdgesResult(Edges);
 }
 
 TArray<FIndex3i> FDelaunay2::GetTriangles() const
@@ -1591,7 +1721,11 @@ bool FDelaunay2::IsDelaunay(TArrayView<const FVector2d> Vertices, TArrayView<con
 
 TArray<TArray<FVector2d>> FDelaunay2::GetVoronoiCells(TArrayView<const FVector2d> Vertices, bool bIncludeBoundary, FAxisAlignedBox2d Bounds, double ExpandBounds) const
 {
-	if (ensureMsgf(Connectivity.IsValid() && !bIsConstrained, TEXT("Voronoi diagram computation requires a valid, unconstrained Delaunay triangulation to be already computed")))
+	if (Result == FDelaunay2::EResult::Collinear)
+	{
+		return DelaunayInternal::GetVoronoiCells_Collinear<double>(Vertices, bIncludeBoundary, Bounds, ExpandBounds);
+	}
+	else if (ensureMsgf(Connectivity.IsValid() && !bIsConstrained, TEXT("Voronoi diagram computation requires an unconstrained Delaunay triangulation to be already computed")))
 	{
 		return DelaunayInternal::GetVoronoiCells<double>(*Connectivity, Vertices, bIncludeBoundary, Bounds, ExpandBounds);
 	}
@@ -1600,7 +1734,11 @@ TArray<TArray<FVector2d>> FDelaunay2::GetVoronoiCells(TArrayView<const FVector2d
 
 TArray<TArray<FVector2f>> FDelaunay2::GetVoronoiCells(TArrayView<const FVector2f> Vertices, bool bIncludeBoundary, FAxisAlignedBox2f Bounds, float ExpandBounds) const
 {
-	if (ensureMsgf(Connectivity.IsValid() && !bIsConstrained, TEXT("Voronoi diagram computation required a valid, unconstrained Delaunay triangulation to be already computed")))
+	if (Result == FDelaunay2::EResult::Collinear)
+	{
+		return DelaunayInternal::GetVoronoiCells_Collinear<float>(Vertices, bIncludeBoundary, Bounds, ExpandBounds);
+	}
+	else if (ensureMsgf(Connectivity.IsValid() && !bIsConstrained, TEXT("Voronoi diagram computation requires an unconstrained Delaunay triangulation to be already computed")))
 	{
 		return DelaunayInternal::GetVoronoiCells<float>(*Connectivity, Vertices, bIncludeBoundary, Bounds, ExpandBounds);
 	}
diff --git a/Engine/Source/Runtime/GeometryCore/Public/CompGeom/Delaunay2.h b/Engine/Source/Runtime/GeometryCore/Public/CompGeom/Delaunay2.h
--- a/Engine/Source/Runtime/GeometryCore/Public/CompGeom/Delaunay2.h
+++ b/Engine/Source/Runtime/GeometryCore/Public/CompGeom/Delaunay2.h
@@ -36,6 +36,22 @@ class FDelaunay2
 {
 public:
 
+	// Indicates result of triangulation
+	enum class EResult
+	{
+		Success,
+		// Has not yet tried to compute a triangulation
+		NotComputed,
+		// No input points
+		EmptyInput,
+		// Input did not span a 2D basis (was 1D or 0D) so could not form a triangle
+		Collinear,
+		// Could not create triangulation containing all requested constraint edges, typically due to intersections between requested edges
+		MissingEdges,
+		// Uncategorized failure
+		Unknown
+	};
+
 	// Options for selecting what triangles to include in the output, for constrained Delaunay triangulation of polygons
 	enum class EFillMode
 	{
@@ -286,16 +302,42 @@ public:
 		return Delaunay.GetVoronoiCells(Sites, bIncludeBoundary, ClipBounds, ExpandBounds);
 	}
 
+	/** @return Last triangulation result -- either Success, NotComputed, or a reason for failure */
+	EResult GetResult() const
+	{
+		return Result;
+	}
+	
+	/** @return Whether the triangulation result supports computing a Voronoi diagram -- i.e., if it was triangulated, or a collinear input */
+	bool CanComputeVoronoiCells() const
+	{
+		return Result == EResult::Collinear || Result == EResult::Success;
+	}
+
 protected:
 	TPimplPtr<FDelaunay2Connectivity> Connectivity;
 
 	bool bIsConstrained = false;
 
 	// helper to perform standard validation on results after Triangulate or ConstrainEdges calls
+	UE_DEPRECATED(5.5, "No longer used by the implementation and may be removed")
 	bool ValidateResult(TArrayView<const FIndex2i> Edges) const
 	{
 		return !bValidateEdges || HasEdges(Edges);
 	}
+
+private:
+	EResult Result = EResult::NotComputed;
+
+	bool ValidateEdgesResult(TArrayView<const FIndex2i> Edges)
+	{
+		if (bValidateEdges && !HasEdges(Edges))
+		{
+			Result = EResult::MissingEdges;
+			return false;
+		}
+		return true;
+	}
 };
 
 } // end namespace UE::Geometry
