serialize the bones in dynamic mesh attribute sets

diff --git a/Engine/Source/Runtime/Core/Public/UObject/FortniteReleaseBranchCustomObjectVersion.h b/Engine/Source/Runtime/Core/Public/UObject/FortniteReleaseBranchCustomObjectVersion.h
--- a/Engine/Source/Runtime/Core/Public/UObject/FortniteReleaseBranchCustomObjectVersion.h
+++ b/Engine/Source/Runtime/Core/Public/UObject/FortniteReleaseBranchCustomObjectVersion.h
@@ -54,6 +54,9 @@ struct FFortniteReleaseBranchCustomObjectVersion
 		// Deprecation of Nav Movement Properties and moving them to a new struct
 		NavMovementComponentMovingPropertiesToStruct,
 
+		// Add bone serialization for dynamic mesh attributes
+		DynamicMeshAttributesSerializeBones,
+
 		// -----<new versions can be added above this line>-------------------------------------------------
 		VersionPlusOne,
 		LatestVersion = VersionPlusOne - 1
diff --git a/Engine/Source/Runtime/GeometryCore/Private/DynamicMesh/DynamicBoneAttribute.cpp b/Engine/Source/Runtime/GeometryCore/Private/DynamicMesh/DynamicBoneAttribute.cpp
--- a/Engine/Source/Runtime/GeometryCore/Private/DynamicMesh/DynamicBoneAttribute.cpp
+++ b/Engine/Source/Runtime/GeometryCore/Private/DynamicMesh/DynamicBoneAttribute.cpp
@@ -88,8 +88,7 @@ void TDynamicBoneAttributeBase<ParentType, AttribValueType>::Append(const Attrib
 template<typename ParentType, typename AttribValueType>
 void TDynamicBoneAttributeBase<ParentType, AttribValueType>::Serialize(FArchive& Ar)
 {	
-	/** TODO: Not supported for now */
-	checkNoEntry();
+	Ar << AttribValues;
 }
 
 // Instantiate for linker
diff --git a/Engine/Source/Runtime/GeometryCore/Private/DynamicMesh/DynamicMesh3_Serialization.cpp b/Engine/Source/Runtime/GeometryCore/Private/DynamicMesh/DynamicMesh3_Serialization.cpp
--- a/Engine/Source/Runtime/GeometryCore/Private/DynamicMesh/DynamicMesh3_Serialization.cpp
+++ b/Engine/Source/Runtime/GeometryCore/Private/DynamicMesh/DynamicMesh3_Serialization.cpp
@@ -12,11 +12,7 @@ namespace FDynamicMesh3Serialization_Local
 enum EDynamicMeshSerializationVersion
 {
 	InitialVersion = 1,
-	CompactAndCompress = 2,
-
-	// ----- new versions to be added above this line -------------------------------------------------
-	VersionPlusOne,
-	LatestVersion = VersionPlusOne - 1
+	CompactAndCompress = 2
 };
 
 // Encapsulates our serialization options, and selects the serialization variant for a given set of options. 
diff --git a/Engine/Source/Runtime/GeometryCore/Private/DynamicMesh/DynamicMeshAttributeSet.cpp b/Engine/Source/Runtime/GeometryCore/Private/DynamicMesh/DynamicMeshAttributeSet.cpp
--- a/Engine/Source/Runtime/GeometryCore/Private/DynamicMesh/DynamicMeshAttributeSet.cpp
+++ b/Engine/Source/Runtime/GeometryCore/Private/DynamicMesh/DynamicMeshAttributeSet.cpp
@@ -8,6 +8,7 @@
 #include "Tasks/Task.h"
 #include "Serialization/NameAsStringProxyArchive.h"
 
+#include "UObject/FortniteReleaseBranchCustomObjectVersion.h"
 
 using namespace UE::Geometry;
 
@@ -1359,6 +1360,7 @@ void FDynamicMeshAttributeSet::Serialize(FArchive& Ar, const FCompactMaps* Compa
 	using namespace FDynamicMeshAttributeSet_Local;
 	
 	Ar.UsingCustomVersion(FUE5MainStreamObjectVersion::GUID);
+	Ar.UsingCustomVersion(FFortniteReleaseBranchCustomObjectVersion::GUID);
 
 	const bool bUseLegacySerialization = Ar.IsLoading() && Ar.CustomVer(FUE5MainStreamObjectVersion::GUID) < FUE5MainStreamObjectVersion::DynamicMeshCompactedSerialization;
 
@@ -1473,8 +1475,23 @@ void FDynamicMeshAttributeSet::Serialize(FArchive& Ar, const FCompactMaps* Compa
 		}
 	}
 
-
-	// TODO: Serialize bone attributes
+	const bool bSerializeBones = !Ar.IsLoading() || Ar.CustomVer(FFortniteReleaseBranchCustomObjectVersion::GUID) >= FFortniteReleaseBranchCustomObjectVersion::DynamicMeshAttributesSerializeBones;
+	if (bSerializeBones)
+	{
+		bool bHasBones = BoneNameAttrib && BoneParentIndexAttrib && BonePoseAttrib && BoneColorAttrib;
+		Ar << bHasBones;
+		if (bHasBones)
+		{
+			if (Ar.IsLoading())
+			{
+				EnableBones(0);
+			}
+			FNameAsStringProxyArchive ProxyArchive(Ar); BoneNameAttrib->Serialize(ProxyArchive);
+			BoneParentIndexAttrib->Serialize(Ar);
+			BonePoseAttrib->Serialize(Ar);
+			BoneColorAttrib->Serialize(Ar);
+		}
+	}
 
 	//Ar << GenericAttributes; // TODO
 }
