add uniform random mesh sampling methods

diff --git a/Engine/Plugins/Runtime/GeometryScripting/Source/GeometryScriptingCore/Private/MeshSamplingFunctions.cpp b/Engine/Plugins/Runtime/GeometryScripting/Source/GeometryScriptingCore/Private/MeshSamplingFunctions.cpp
--- a/Engine/Plugins/Runtime/GeometryScripting/Source/GeometryScriptingCore/Private/MeshSamplingFunctions.cpp
+++ b/Engine/Plugins/Runtime/GeometryScripting/Source/GeometryScriptingCore/Private/MeshSamplingFunctions.cpp
@@ -7,6 +7,8 @@
 #include "Scene/SceneCapturePhotoSet.h"
 #include "GameFramework/Actor.h"
 #include "Async/ParallelFor.h"
+#include "TriangleTypes.h"
+#include "VectorUtil.h"
 
 #include UE_INLINE_GENERATED_CPP_BY_NAME(MeshSamplingFunctions)
 
@@ -15,6 +17,53 @@ using namespace UE::Geometry;
 #define LOCTEXT_NAMESPACE "UGeometryScriptLibrary_MeshSamplingFunctions"
 
 
+const UDynamicMesh* UGeometryScriptLibrary_MeshSamplingFunctions::ComputeUniformRandomPointSampling(
+	const UDynamicMesh* TargetMesh,
+	int32 NumSamples,
+	TArray<FTransform>& Samples,
+	FGeometryScriptIndexList& TriangleIDs,
+	int RandomSeed,
+	UGeometryScriptDebug* Debug)
+{
+	if (TargetMesh == nullptr)
+	{
+		UE::Geometry::AppendError(Debug, EGeometryScriptErrorType::InvalidInputs, LOCTEXT("ComputeUniformRandomPointSampling", "ComputePointSampling: TargetMesh is Null"));
+		return TargetMesh;
+	}
+
+	TargetMesh->ProcessMesh([Debug, NumSamples, RandomSeed, &Samples, &TriangleIDs](const FDynamicMesh3& ReadMesh)
+	{
+		if (ReadMesh.TriangleCount() == 0)
+		{
+			return;
+		}
+		FMeshUniformSurfaceSampling MeshUniformSampling(ReadMesh);
+		if (!MeshUniformSampling.IsValid())
+		{
+			// cannot sample mesh (e.g., if there are no triangles)
+			UE::Geometry::AppendWarning(Debug, EGeometryScriptErrorType::InvalidInputs, LOCTEXT("ComputeUniformRandomPointSampling", "ComputeUniformRandomPointSampling: TargetMesh does not have triangles to sample."));
+			return;
+		}
+		FRandomStream RandomStream(RandomSeed);
+		FTriangle3d Tri;
+		TriangleIDs.Reset(EGeometryScriptIndexType::Triangle);
+		TArray<int>& IndexList = *TriangleIDs.List;
+		for (int32 SampleIdx = 0; SampleIdx < NumSamples; ++SampleIdx)
+		{
+			int32 TID = MeshUniformSampling.DrawTriangleIDSample(RandomStream.GetFraction(), RandomStream.GetFraction());
+			if (!ensure(ReadMesh.IsTriangle(TID))) // guard against sampling invalid IDs (should never happen)
+			{
+				continue;
+			}
+			ReadMesh.GetTriVertices(TID, Tri.V[0], Tri.V[1], Tri.V[2]);
+			Samples.Add(FFrame3d(VectorUtil::UniformSampleTrianglePoint<double>(RandomStream.GetFraction(), RandomStream.GetFraction(), Tri.V[0], Tri.V[1], Tri.V[2]), Tri.Normal()).ToFTransform());
+			IndexList.Add(TID);
+		}
+	});
+
+	return TargetMesh;
+}
+
 UDynamicMesh* UGeometryScriptLibrary_MeshSamplingFunctions::ComputePointSampling(
 	UDynamicMesh* TargetMesh,
 	FGeometryScriptMeshPointSamplingOptions Options,
diff --git a/Engine/Plugins/Runtime/GeometryScripting/Source/GeometryScriptingCore/Public/GeometryScript/MeshSamplingFunctions.h b/Engine/Plugins/Runtime/GeometryScripting/Source/GeometryScriptingCore/Public/GeometryScript/MeshSamplingFunctions.h
--- a/Engine/Plugins/Runtime/GeometryScripting/Source/GeometryScriptingCore/Public/GeometryScript/MeshSamplingFunctions.h
+++ b/Engine/Plugins/Runtime/GeometryScripting/Source/GeometryScriptingCore/Public/GeometryScript/MeshSamplingFunctions.h
@@ -109,6 +109,19 @@ class GEOMETRYSCRIPTINGCORE_API UGeometryScriptLibrary_MeshSamplingFunctions : p
 	GENERATED_BODY()
 public:
 
+	/**
+	 * Compute a 'uniform random' (not 'uniform spacing') point sampling over the mesh surface
+	 */
+	UFUNCTION(BlueprintCallable, Category = "GeometryScript|MeshSampling", meta = (ScriptMethod))
+	static UPARAM(DisplayName = "Target Mesh") const UDynamicMesh*
+	ComputeUniformRandomPointSampling(
+		const UDynamicMesh* TargetMesh,
+		int32 NumSamples,
+		TArray<FTransform>& Samples,
+		FGeometryScriptIndexList& TriangleIDs,
+		int RandomSeed = 0,
+		UGeometryScriptDebug* Debug = nullptr);
+
 	/**
 	 * Compute a set of sample points lying on the surface of TargetMesh based on the provided sampling Options.
 	 * Samples are approximately uniformly distributed, and non-overlapping relative to the provided Options.SamplingRadius,
diff --git a/Engine/Source/Runtime/GeometryCore/Private/Sampling/MeshSurfacePointSampling.cpp b/Engine/Source/Runtime/GeometryCore/Private/Sampling/MeshSurfacePointSampling.cpp
--- a/Engine/Source/Runtime/GeometryCore/Private/Sampling/MeshSurfacePointSampling.cpp
+++ b/Engine/Source/Runtime/GeometryCore/Private/Sampling/MeshSurfacePointSampling.cpp
@@ -729,4 +729,151 @@ void FMeshSurfacePointSampling::ComputePoissonSampling(const FDynamicMesh3& Mesh
 	}
 
 	Result.SetSuccess(true, Progress);
-}
\ No newline at end of file
+}
+
+
+template<class RealType>
+bool TWeightedSamplingAliasTable<RealType>::Init(TConstArrayView<RealType> Weights, RealType SumOfValidWeights, bool bAllowInvalidWeights)
+{
+	Probability.Reset();
+	Alias.Reset();
+
+	const int32 NumWeights = Weights.Num();
+
+	// handle the all-zero weight case specially
+	if (SumOfValidWeights <= (RealType)0)
+	{
+		// in the edge case where we have no positive weights, and some invalid weights, generate valid weights to uniform-sample the zero weight triangles
+		if (bAllowInvalidWeights)
+		{
+			// Make a new weight table where the zero weights are given uniform positive weight,
+			// and the invalid weights remain negative
+			TArray<RealType> UniformWeights;
+			UniformWeights.SetNumUninitialized(NumWeights);
+			RealType NewWeightSum = 0;
+			bool bHaveInvalidWeights = false;
+			for (int32 Idx = 0; Idx < NumWeights; ++Idx)
+			{
+				if (Weights[Idx] < 0)
+				{
+					bHaveInvalidWeights = true;
+					UniformWeights[Idx] = -1;
+				}
+				else
+				{
+					NewWeightSum += 1;
+					UniformWeights[Idx] = 1;
+				}
+			}
+			// If there were invalid weights, use our new uniform weights to build the sampling table
+			// (otherwise, just fall through to the no-invalid-weight uniform case, where we don't need the alias table)
+			if (bHaveInvalidWeights)
+			{
+				if (NewWeightSum == (RealType)0)
+				{
+					// no valid indices to sample, cannot build table
+					return false;
+				}
+				return Init(UniformWeights, SumOfValidWeights, true);
+			}
+		}
+		
+		// in the case of zero weight sum with no invalid weights, we can use a uniform probability table with no aliases
+		Probability.Init(TMathUtilConstants<RealType>::MaxReal, NumWeights);
+		// Leave alias empty, since it will never be picked
+		return IsValid();
+	}
+
+	Probability.SetNumUninitialized(NumWeights);
+
+	TArray<int32> Small, Large;
+	Small.Reserve(NumWeights);
+	Large.Reserve(NumWeights);
+
+	// Transform weights to initial scaled probabilities, and initialize small, large index arrays
+	RealType WeightToProb = (RealType)NumWeights / SumOfValidWeights;
+	int32 KnownValidWeightIdx = INDEX_NONE;
+	for (int32 Idx = NumWeights - 1; Idx >= 0; --Idx)
+	{
+		Probability[Idx] = Weights[Idx] * WeightToProb;
+		if (bAllowInvalidWeights)
+		{
+			if (Probability[Idx] < 0)
+			{
+				Probability[Idx] = 0;
+			}
+			else
+			{
+				KnownValidWeightIdx = Idx;
+			}
+		}
+		else
+		{
+			// Weights must be non-negative if we are not allowing invalid weights
+			checkSlow(Probability[Idx] >= (RealType)0);
+		}
+		if (Probability[Idx] < 1)
+		{
+			Small.Add(Idx);
+		}
+		else
+		{
+			Large.Add(Idx);
+		}
+	}
+	if (bAllowInvalidWeights && KnownValidWeightIdx == INDEX_NONE)
+	{
+		// No valid weights, clear the table and return failure
+		Probability.Empty();
+		return false;
+	}
+	
+	// Build aliases
+	Alias.SetNumUninitialized(NumWeights);
+	while (!Small.IsEmpty() && !Large.IsEmpty())
+	{
+		int32 SmallIdx = Small.Pop(EAllowShrinking::No);
+		int32 LargeIdx = Large.Pop(EAllowShrinking::No);
+		Alias[SmallIdx] = LargeIdx;
+		Probability[LargeIdx] = (Probability[LargeIdx] + Probability[SmallIdx]) - (RealType)1;
+		if (Probability[LargeIdx] < (RealType)1)
+		{
+			Small.Add(LargeIdx);
+		}
+		else
+		{
+			Large.Add(LargeIdx);
+		}
+	}
+
+	// Remaining unmatched large or small indices are assigned probability 1, don't need aliases
+	for (int32 Idx : Large)
+	{
+		Probability[Idx] = TMathUtilConstants<RealType>::MaxReal;
+		// Alias irrelevant
+	}
+	for (int32 Idx : Small)
+	{
+		// Theoretically possible due to numerical error to still have an invalid entry in the small list ...
+		// In this unlikely event, just redirect to a known valid alias
+		if (bAllowInvalidWeights && Weights[Idx] < 0)
+		{
+			Probability[Idx] = 0;
+			Alias[Idx] = KnownValidWeightIdx;
+		}
+		else
+		{
+			Probability[Idx] = TMathUtilConstants<RealType>::MaxReal;
+			// Alias irrelevant
+		}
+	}
+
+	return IsValid();
+}
+
+namespace UE::Geometry
+{
+	template class TWeightedSamplingAliasTable<float>;
+	template class TWeightedSamplingAliasTable<double>;
+}
+
diff --git a/Engine/Source/Runtime/GeometryCore/Public/Sampling/MeshSurfacePointSampling.h b/Engine/Source/Runtime/GeometryCore/Public/Sampling/MeshSurfacePointSampling.h
--- a/Engine/Source/Runtime/GeometryCore/Public/Sampling/MeshSurfacePointSampling.h
+++ b/Engine/Source/Runtime/GeometryCore/Public/Sampling/MeshSurfacePointSampling.h
@@ -16,10 +16,138 @@ namespace Geometry
 
 class FDynamicMesh3;
 
+// class for efficiently drawing weighted samples (when you need to draw more than a very small amount of them)
+// note this is a more flexible version of the algorithm in Engine/WeightedRandomSampler.h, 
+// with support for cases where some IDs cannot be sampled (e.g. when sampling triangles of non-compact meshes)
+template<typename RealType>
+class TWeightedSamplingAliasTable
+{
+public:
+	/**
+	 * Initialize sampling table
+	 * 
+	 * @param Weights The relative frequency that the element should be sampled -- e.g., for uniform mesh sampling, these may be triangle areas.
+	 * @param SumOfValidWeights Sum of all valid weights. Must not include negative (invalid) weights.
+	 * @param bAllowInvalidWeights If true, then negative weights can be used to indicate invalid IDs, which must never be sampled. Otherwise, all weights should be non-negative.
+	 * @return true if successfully created a valid sampling table, false otherwise (e.g., if there were no valid weights)
+	 */
+	GEOMETRYCORE_API bool Init(TConstArrayView<RealType> Weights, RealType SumOfValidWeights, bool bAllowInvalidWeights);
+
+	TWeightedSamplingAliasTable() = default;
+	TWeightedSamplingAliasTable(TConstArrayView<RealType> Weights, RealType SumOfValidWeights, bool bAllowInvalidWeights)
+	{
+		Init(Weights, SumOfValidWeights, bAllowInvalidWeights);
+	}
+	TWeightedSamplingAliasTable(TConstArrayView<RealType> Weights)
+	{
+		bool bHasInvalidWeight = false;
+		RealType WtSum = (RealType)0;
+		for (RealType Wt : Weights)
+		{
+			if (Wt < 0)
+			{
+				bHasInvalidWeight = true;
+				continue;
+			}
+			WtSum += Wt;
+		}
+		Init(Weights, WtSum, bHasInvalidWeight);
+	}
+	
+	inline int32 Num() const
+	{
+		return Probability.Num();
+	}
+
+	inline bool IsValid() const
+	{
+		return !Probability.IsEmpty();
+	}
+	
+	// UniformRandom1, UniformRandom2 must be in the [0,1) range
+	inline int32 Sample(RealType UniformRandom1, RealType UniformRandom2) const
+	{
+		checkSlow(IsValid() && UniformRandom1 < (RealType)1);
+		int32 Idx = FMath::TruncToInt32(UniformRandom1 * (RealType)Probability.Num());
+		return UniformRandom2 < Probability[Idx] ? Idx : Alias[Idx];
+	}
+
+private:
+	// Table of probability that you should keep a given index chosen by unweighted random sampling (can be MaxReal if should always pick the index)
+	TArray<RealType> Probability;
+	// Table indicating the alternative index you should choose, if you don't keep the initial index
+	TArray<int32> Alias;
+};
+
+// Class to help draw uniform-random surface samples from a (possibly non-compact) triangle mesh. 
+// TMeshType must implement: IsTriangle, MaxTriangleID, and GetTriVertices. RealType must be float or double.
+template<typename TMeshType, typename RealType>
+class TMeshUniformSurfaceSampling
+{
+public:
+	
+	TMeshUniformSurfaceSampling() = default;
+	TMeshUniformSurfaceSampling(const TMeshType& Mesh)
+	{
+		Init(Mesh);
+	}
+	
+	// @return true on successful initialization
+	bool Init(const TMeshType& Mesh)
+	{
+		const int32 MaxTID = Mesh.MaxTriangleID();
+		TArray<RealType> Wts;
+		RealType WtSum = (RealType)0;
+		bool bHasInvalidWeights = false;
+		for (int32 TID = 0; TID < MaxTID; ++TID)
+		{
+			if (Mesh.IsTriangle(TID))
+			{
+				TVector<RealType> A, B, C;
+				Mesh.GetTriVertices(TID, A, B, C);
+				RealType Area = (RealType)VectorUtil::Area(A, B, C);
+				WtSum += Area;
+				Wts.Add(Area);
+			}
+			else
+			{
+				bHasInvalidWeights = true;
+				Wts.Add(-1);
+			}
+		}
+		return AliasTable.Init(Wts, WtSum, bHasInvalidWeights);
+	}
+
+	// @return true if we can generate samples (e.g., if the class has been initialized with a non-empty mesh)
+	bool IsValid() const
+	{
+		return AliasTable.IsValid();
+	}
+
+	/**
+	 * Sample a random triangle, with probability proportional to relative triangle area.
+	 * Note: Use VectorUtil::UniformSampleTrianglePoint to generate a uniform random point on a given triangle,
+	 * or VectorUtil::UniformSampleTriangleBarycentricCoords for a uniform random barycentric coordinate
+	 *
+	 * Parameters UniformRandom1, UniformRandom2 must be in the [0,1) range.
+	 */
+	int32 DrawTriangleIDSample(RealType UniformRandom1, RealType UniformRandom2) const
+	{
+		return AliasTable.Sample(UniformRandom1, UniformRandom2);
+	}
+	
+private:
+	TWeightedSamplingAliasTable<RealType> AliasTable;
+};
+
+using FMeshUniformSurfaceSampling = TMeshUniformSurfaceSampling<FDynamicMesh3, double>;
+using FWeightedSamplingAliasTabled = TWeightedSamplingAliasTable<double>;
+using FWeightedSamplingAliasTablef = TWeightedSamplingAliasTable<float>;
+
 
 /**
  * FMeshSurfacePointSampling computes oriented point samples on the surface of
- * a Mesh using various sampling strategies.
+ * a Mesh using various sampling strategies, with controls over sample spacing.
  */
 class FMeshSurfacePointSampling
 {
