Invalidate dynamic mesh non-manifold mappings on remove vertex (matching all other topology changes)

diff --git a/Engine/Source/Runtime/GeometryCore/Private/DynamicMesh/NonManifoldMappingSupport.cpp b/Engine/Source/Runtime/GeometryCore/Private/DynamicMesh/NonManifoldMappingSupport.cpp
--- a/Engine/Source/Runtime/GeometryCore/Private/DynamicMesh/NonManifoldMappingSupport.cpp
+++ b/Engine/Source/Runtime/GeometryCore/Private/DynamicMesh/NonManifoldMappingSupport.cpp
@@ -213,6 +213,11 @@ public:
 	{
 		SourceVertexIds.InsertAt(VertexID, VertexID);
 	}
+
+	virtual void OnRemoveVertex(int VertexID) override
+	{
+		bValid = false;
+	}
 };
 
 void FNonManifoldSourceVertexIdAttributeChange::SaveInitialVertex(const FDynamicMeshAttributeBase* Attribute, int VertexID)
