add optional thickening/inflate parameter to FitKDOP3, and set default to .1, to match previous KDOP fitting behavior and guarantee that flat shapes still generate valid KDOPs.

This fixes an ensure and collision generation failure on asset import.

diff --git a/Engine/Source/Runtime/GeometryCore/Private/CompGeom/FitKDOP3.cpp b/Engine/Source/Runtime/GeometryCore/Private/CompGeom/FitKDOP3.cpp
--- a/Engine/Source/Runtime/GeometryCore/Private/CompGeom/FitKDOP3.cpp
+++ b/Engine/Source/Runtime/GeometryCore/Private/CompGeom/FitKDOP3.cpp
@@ -16,7 +16,7 @@ namespace UE::Private::FitKDop3Helpers
 	template<typename RealType>
 	bool FitKDopImpl3(TArrayView<const UE::Math::TVector<RealType>> PlaneDirections,
 		const int32 NumPoints, TFunctionRef<UE::Math::TVector<RealType>(int32)> GetPointFunc, TFunctionRef<bool(int32)> FilterFunc,
-		TArray<UE::Math::TVector<RealType>>& OutVertices, TArray<UE::Math::TPlane<RealType>>* OptionalOutPlanes, RealType Epsilon, RealType VertexSnapDistance)
+		TArray<UE::Math::TVector<RealType>>& OutVertices, TArray<UE::Math::TPlane<RealType>>* OptionalOutPlanes, RealType Epsilon, RealType VertexSnapDistance, RealType Inflate)
 	{
 		using namespace UE::Math;
 
@@ -43,6 +43,11 @@ namespace UE::Private::FitKDop3Helpers
 			}
 		}
 
+		for (RealType& Dist : PlaneDists)
+		{
+			Dist += Inflate;
+		}
+
 		TArray<TPlane<RealType>> LocalPlanes;
 		TArray<TPlane<RealType>>* UsePlanes = OptionalOutPlanes ? OptionalOutPlanes : &LocalPlanes;
 		UsePlanes->SetNumUninitialized(PlaneDists.Num());
@@ -129,22 +134,22 @@ namespace UE::Geometry
 	bool FitKDOPVertices3(
 		TArrayView<const UE::Math::TVector<RealType>> PlaneDirections,
 		const int32 NumPoints, TFunctionRef<UE::Math::TVector<RealType>(int32)> GetPointFunc, TFunctionRef<bool(int32)> FilterFunc,
-		TArray<UE::Math::TVector<RealType>>& OutVertices, TArray<UE::Math::TPlane<RealType>>* OptionalOutPlanes, RealType Epsilon, RealType VertexSnapDistance)
+		TArray<UE::Math::TVector<RealType>>& OutVertices, TArray<UE::Math::TPlane<RealType>>* OptionalOutPlanes, RealType Epsilon, RealType VertexSnapDistance, RealType Inflate)
 	{
 		return UE::Private::FitKDop3Helpers::FitKDopImpl3(PlaneDirections,
 			NumPoints, GetPointFunc, FilterFunc,
 			OutVertices, OptionalOutPlanes,
-			Epsilon, VertexSnapDistance);
+			Epsilon, VertexSnapDistance, Inflate);
 	}
 
 
 	// explicit instantiations
 	template bool GEOMETRYCORE_API FitKDOPVertices3<float>(
 		TArrayView<const UE::Math::TVector<float>> PlaneDirections, const int32 NumPoints, TFunctionRef<UE::Math::TVector<float>(int32)> GetPointFunc, TFunctionRef<bool(int32)> FilterFunc,
-		TArray<UE::Math::TVector<float>>& OutVertices, TArray<UE::Math::TPlane<float>>* OptionalOutPlanes, float Epsilon, float VertexSnapDistance);
+		TArray<UE::Math::TVector<float>>& OutVertices, TArray<UE::Math::TPlane<float>>* OptionalOutPlanes, float Epsilon, float VertexSnapDistance, float Inflate);
 	template bool GEOMETRYCORE_API FitKDOPVertices3<double>(
 		TArrayView<const UE::Math::TVector<double>> PlaneDirections, const int32 NumPoints, TFunctionRef<UE::Math::TVector<double>(int32)> GetPointFunc, TFunctionRef<bool(int32)> FilterFunc,
-		TArray<UE::Math::TVector<double>>& OutVertices, TArray<UE::Math::TPlane<double>>* OptionalOutPlanes, double Epsilon, double VertexSnapDistance);
+		TArray<UE::Math::TVector<double>>& OutVertices, TArray<UE::Math::TPlane<double>>* OptionalOutPlanes, double Epsilon, double VertexSnapDistance, double Inflate);
 	
 
 } /// namespace UE::Geometry
diff --git a/Engine/Source/Runtime/GeometryCore/Public/CompGeom/FitKDOP3.h b/Engine/Source/Runtime/GeometryCore/Public/CompGeom/FitKDOP3.h
--- a/Engine/Source/Runtime/GeometryCore/Public/CompGeom/FitKDOP3.h
+++ b/Engine/Source/Runtime/GeometryCore/Public/CompGeom/FitKDOP3.h
@@ -27,7 +27,7 @@ namespace Geometry
 		TArrayView<const UE::Math::TVector<RealType>> PlaneDirections,
 		const int32 NumPoints, TFunctionRef<UE::Math::TVector<RealType>(int32)> GetPointFunc, TFunctionRef<bool(int32)> FilterFunc,
 		TArray<UE::Math::TVector<RealType>>& OutVertices, TArray<UE::Math::TPlane<RealType>>* OptionalOutPlanes = nullptr,
-		RealType Epsilon = TMathUtil<RealType>::Epsilon, RealType VertexSnapDistance = (RealType)0.1);
+		RealType Epsilon = TMathUtil<RealType>::Epsilon, RealType VertexSnapDistance = (RealType)0.1, RealType Inflate = (RealType)0.1);
 
 	/**
 	 * Compute the vertices (and optionally planes) of a k-DOP bounding convex hull containing the points returned by GetPointFunc
@@ -37,9 +37,9 @@ namespace Geometry
 		TArrayView<const UE::Math::TVector<RealType>> PlaneDirections,
 		const int32 NumPoints, TFunctionRef<UE::Math::TVector<RealType>(int32)> GetPointFunc,
 		TArray<UE::Math::TVector<RealType>>& OutVertices, TArray<UE::Math::TPlane<RealType>>* OptionalOutPlanes = nullptr,
-		RealType Epsilon = TMathUtil<RealType>::Epsilon, RealType VertexSnapDistance = (RealType)0.1)
+		RealType Epsilon = TMathUtil<RealType>::Epsilon, RealType VertexSnapDistance = (RealType)0.1, RealType Inflate = (RealType)0.1)
 	{
-		return FitKDOPVertices3<RealType>(PlaneDirections, NumPoints, GetPointFunc, [](int32)->bool {return true;}, OutVertices, OptionalOutPlanes, Epsilon);
+		return FitKDOPVertices3<RealType>(PlaneDirections, NumPoints, GetPointFunc, [](int32)->bool {return true;}, OutVertices, OptionalOutPlanes, Epsilon, VertexSnapDistance, Inflate);
 	}
 
 	/**
@@ -50,9 +50,9 @@ namespace Geometry
 		TArrayView<const UE::Math::TVector<RealType>> PlaneDirections,
 		const TArray<UE::Math::TVector<RealType>>& Points,
 		TArray<UE::Math::TVector<RealType>>& OutVertices, TArray<UE::Math::TPlane<RealType>>* OptionalOutPlanes = nullptr,
-		RealType Epsilon = TMathUtil<RealType>::Epsilon, RealType VertexSnapDistance = (RealType)0.1)
+		RealType Epsilon = TMathUtil<RealType>::Epsilon, RealType VertexSnapDistance = (RealType)0.1, RealType Inflate = (RealType)0.1)
 	{
-		return FitKDOPVertices3<RealType>(PlaneDirections, Points.Num(), [&Points](int32 Idx) {return Points[Idx];}, [](int32)->bool {return true;}, OutVertices, OptionalOutPlanes, Epsilon);
+		return FitKDOPVertices3<RealType>(PlaneDirections, Points.Num(), [&Points](int32 Idx) {return Points[Idx];}, [](int32)->bool {return true;}, OutVertices, OptionalOutPlanes, Epsilon, VertexSnapDistance, Inflate);
 	}
 
 } // end namespace Geometry
