fix handling of invalid triangles in mesh aabb tree -- do not leave uninitialized memory at the end of the triangles list, and pass only the count of valid triangles to the build method

diff --git a/Engine/Source/Runtime/GeometryCore/Public/Spatial/MeshAABBTree3.h b/Engine/Source/Runtime/GeometryCore/Public/Spatial/MeshAABBTree3.h
--- a/Engine/Source/Runtime/GeometryCore/Public/Spatial/MeshAABBTree3.h
+++ b/Engine/Source/Runtime/GeometryCore/Public/Spatial/MeshAABBTree3.h
@@ -1207,11 +1207,10 @@ public:
 	{
 		// build list of valid Triangles & Centers. We skip any
 		// Triangles that have infinite/garbage vertices...
-		int i = 0;
 		TArray<int> Triangles;
-		Triangles.SetNumUninitialized(Mesh->TriangleCount());
+		Triangles.Reserve(Mesh->TriangleCount());
 		TArray<FVector3d> Centers;
-		Centers.SetNumUninitialized(Mesh->TriangleCount());
+		Centers.Reserve(Mesh->TriangleCount());
 		for (int ti = 0; ti < Mesh->MaxTriangleID(); ti++)
 		{
 			if (Mesh->IsTriangle(ti) == false)
@@ -1223,14 +1222,13 @@ public:
 			bool bInvalid = FMathd::IsNaN(d2) || (FMathd::IsFinite(d2) == false);
 			if (bInvalid == false)
 			{
-				Triangles[i] = ti;
-				Centers[i] = TMeshQueries<TriangleMeshType>::GetTriCentroid(*Mesh, ti);
-				i++;
+				Triangles.Add(ti);
+				Centers.Add(TMeshQueries<TriangleMeshType>::GetTriCentroid(*Mesh, ti));
 			} // otherwise skip this Tri
 		}
+		checkSlow(Triangles.Num() == Centers.Num());
 
-		// todo: is passing TriangleCount() correct here? what if we skipped some elements above?
-		BuildTopDown(Triangles, Centers, Mesh->TriangleCount());
+		BuildTopDown(Triangles, Centers, Triangles.Num());
 	}
 
 
@@ -1239,11 +1237,10 @@ public:
 	{
 		// build list of valid Triangles & Centers. We skip any
 		// Triangles that have infinite/garbage vertices...
-		int32 i = 0;
 		TArray<int32> Triangles;
-		Triangles.SetNumUninitialized(NumTriangles);
+		Triangles.Reserve(NumTriangles);
 		TArray<FVector3d> Centers;
-		Centers.SetNumUninitialized(NumTriangles);
+		Centers.Reserve(NumTriangles);
 		for (int32 ti : TriangleList)
 		{
 			if ( Mesh->IsTriangle(ti) == false)
@@ -1255,14 +1252,13 @@ public:
 			bool bInvalid = FMathd::IsNaN(d2) || (FMathd::IsFinite(d2) == false);
 			if (bInvalid == false)
 			{
-				Triangles[i] = ti;
-				Centers[i] = TMeshQueries<TriangleMeshType>::GetTriCentroid(*Mesh, ti);
-				i++;
+				Triangles.Add(ti);
+				Centers.Add(TMeshQueries<TriangleMeshType>::GetTriCentroid(*Mesh, ti));
 			} // otherwise skip this Tri
 		}
+		checkSlow(Triangles.Num() == Centers.Num());
 
-		// todo: is passing NumTriangles correct here? what if we skipped some elements above?
-		BuildTopDown(Triangles, Centers, NumTriangles);
+		BuildTopDown(Triangles, Centers, Triangles.Num());
 	}
 
 
