Prevent mesh surface point sampling hanging on degenerate triangles, by switching to non-rejection-sampling in that case. (Leave rejection sampling as-is for all other cases, to preserve exact same sampling behavior otherwise)

diff --git a/Engine/Source/Runtime/GeometryCore/Private/Sampling/MeshSurfacePointSampling.cpp b/Engine/Source/Runtime/GeometryCore/Private/Sampling/MeshSurfacePointSampling.cpp
--- a/Engine/Source/Runtime/GeometryCore/Private/Sampling/MeshSurfacePointSampling.cpp
+++ b/Engine/Source/Runtime/GeometryCore/Private/Sampling/MeshSurfacePointSampling.cpp
@@ -160,17 +160,24 @@ struct FPerTriangleDensePointSampling
 
 				int NumGenerated = 0;
 				FRandomStream RandomStream(tid + RandomSeed);
+
+				// workaround for rejection sampling method failing for degenerate triangles (sampling method 1 does not have this issue)
+				bool bIsDegenerateTri = TriInfo.TriAreas[tid] == 0;
 				while (NumGenerated < NumSamples)
 				{
 					double a1 = RandomStream.GetFraction();
 					double a2 = RandomStream.GetFraction();
+					if (bIsDegenerateTri && a1 + a2 > 1)
+					{
+						a1 = 1-a1;
+						a2 = 1-a2;
+					}
 					FVector2d PointUV = TriUV.V[0] + a1 * V1 + a2 * V2;
-					if (TriUV.IsInside(PointUV))
+					if (bIsDegenerateTri || TriUV.IsInside(PointUV))
 					{
 						FVector3d Position = ProjectFrame.FromPlaneUV(PointUV, 2);
 						PointSetOut.DensePoints[StartIndex+NumGenerated] = Position;
 						PointSetOut.Triangles[StartIndex+NumGenerated] = tid;
-
 						if (bComputeWeights)
 						{
 							FVector3d BaryCoords = TriUV.GetBarycentricCoords(PointUV);
