fix const correctness for vectorutil triangle sampling

diff --git a/Engine/Source/Runtime/GeometryCore/Public/VectorUtil.h b/Engine/Source/Runtime/GeometryCore/Public/VectorUtil.h
--- a/Engine/Source/Runtime/GeometryCore/Public/VectorUtil.h
+++ b/Engine/Source/Runtime/GeometryCore/Public/VectorUtil.h
@@ -437,7 +437,7 @@ namespace VectorUtil
 	 * Note the sample positions will not match those of UniformSampleTriangleBarycentricCoords for the same R1, R2 values
 	 */
 	template<typename RealType>
-	TVector<RealType> UniformSampleTrianglePoint(RealType R1, RealType R2, TVector<RealType>& A, TVector<RealType>& B, TVector<RealType>& C)
+	TVector<RealType> UniformSampleTrianglePoint(RealType R1, RealType R2, const TVector<RealType>& A, const TVector<RealType>& B, const TVector<RealType>& C)
 	{
 		// reflect samples that would go outside the triangle
 		if (R1 + R2 > (RealType)1)
@@ -455,7 +455,7 @@ namespace VectorUtil
 	 * Note the sample positions will not match those of UniformSampleTriangleBarycentricCoords for the same R1, R2 values
 	 */
 	template<typename RealType>
-	TVector2<RealType> UniformSampleTrianglePoint(RealType R1, RealType R2, TVector2<RealType>& A, TVector2<RealType>& B, TVector2<RealType>& C)
+	TVector2<RealType> UniformSampleTrianglePoint(RealType R1, RealType R2, const TVector2<RealType>& A, const TVector2<RealType>& B, const TVector2<RealType>& C)
 	{
 		// reflect samples that would go outside the triangle
 		if (R1 + R2 > (RealType)1)
