TAxisAlignedBox3: fix for implicitly converting an Empty box into an Infinite box

diff --git a/Engine/Source/Runtime/GeometryCore/Public/BoxTypes.h b/Engine/Source/Runtime/GeometryCore/Public/BoxTypes.h
--- a/Engine/Source/Runtime/GeometryCore/Public/BoxTypes.h
+++ b/Engine/Source/Runtime/GeometryCore/Public/BoxTypes.h
@@ -296,6 +296,13 @@ struct TAxisAlignedBox3
 
 	TAxisAlignedBox3(const TAxisAlignedBox3& Box, const FTransformSRT3d& Transform)
 	{
+		if (Box.IsEmpty())
+		{
+			Min = Box.Min;
+			Max = Box.Max;
+			return;
+		}
+
 		TVector<RealType> C0 = Transform.TransformPosition(Box.GetCorner(0));
 		Min = C0;
 		Max = C0;
