GeometryCore: Add cylinder point containment function.

diff --git a/Engine/Source/Runtime/GeometryCore/Private/Intersection/ContainmentQueries3.cpp b/Engine/Source/Runtime/GeometryCore/Private/Intersection/ContainmentQueries3.cpp
--- a/Engine/Source/Runtime/GeometryCore/Private/Intersection/ContainmentQueries3.cpp
+++ b/Engine/Source/Runtime/GeometryCore/Private/Intersection/ContainmentQueries3.cpp
@@ -134,6 +134,22 @@ bool UE::Geometry::IsInside(const TOrientedBox3<RealType>& OuterBox, const TCaps
 
 
 
+template<typename RealType>
+bool UE::Geometry::DoesCylinderContainPoint(const TVector<RealType>& CylinderCenter, const TVector<RealType>& NormalizedCylinderAxis, 
+	RealType CylinderRadius, RealType CylinderHeight, const TVector<RealType>& QueryPoint)
+{
+	// In debug build, verify that the cylinder axis is indeed normalized
+	checkSlow(FMath::Abs(1 - NormalizedCylinderAxis.SquaredLength()) <= KINDA_SMALL_NUMBER);
+
+	RealType TValue = (QueryPoint - CylinderCenter).Dot(NormalizedCylinderAxis);
+	if (FMath::Abs(TValue) > CylinderHeight / 2)
+	{
+		return false;
+	}
+	TVector<RealType> ProjectedPoint = CylinderCenter + TValue * NormalizedCylinderAxis;
+	return (QueryPoint - ProjectedPoint).SquaredLength() <= CylinderRadius * CylinderRadius;
+}
+
 
 
 namespace UE
@@ -160,5 +176,10 @@ namespace UE
 		template bool GEOMETRYCORE_API IsInside(const TOrientedBox3<double>& OuterSphere, const TSphere3<double>& InnerSphere);
 		template bool GEOMETRYCORE_API IsInside(const TOrientedBox3<float>& OuterSphere, const TCapsule3<float>& InnerSphere);
 		template bool GEOMETRYCORE_API IsInside(const TOrientedBox3<double>& OuterSphere, const TCapsule3<double>& InnerSphere);
+
+		template bool GEOMETRYCORE_API DoesCylinderContainPoint(const TVector<float>& CylinderCenter, const TVector<float>& CylinderAxis, 
+			float CylinderRadius, float CylinderHeight, const TVector<float>& QueryPoint);
+		template bool GEOMETRYCORE_API DoesCylinderContainPoint(const TVector<double>& CylinderCenter, const TVector<double>& CylinderAxis, 
+			double CylinderRadius, double CylinderHeight, const TVector<double>& QueryPoint);
 	}
 }
diff --git a/Engine/Source/Runtime/GeometryCore/Public/Intersection/ContainmentQueries3.h b/Engine/Source/Runtime/GeometryCore/Public/Intersection/ContainmentQueries3.h
--- a/Engine/Source/Runtime/GeometryCore/Public/Intersection/ContainmentQueries3.h
+++ b/Engine/Source/Runtime/GeometryCore/Public/Intersection/ContainmentQueries3.h
@@ -225,6 +225,15 @@ namespace UE
 			return IsInside(OuterGrid, OuterGridTransform * InnerGridTransform.Inverse(), TOrientedBox3<RealType>(InnerGrid.Bounds()));
 		}
 
+
+		// Cylinder
+
+		/**
+		 * Test if QueryPoint is inside (or on) the given cylinder.
+		 */
+		template<typename RealType>
+		bool DoesCylinderContainPoint(const TVector<RealType>& CylinderCenter, const TVector<RealType>& NormalizedCylinderAxis,
+			RealType CylinderRadius, RealType CylindeHeight, const TVector<RealType>& QueryPoint);
 	}
 }
 
