Fix ImageBuilder Convert function accessing a private member

diff --git a/Engine/Source/Runtime/GeometryCore/Public/Image/ImageBuilder.h b/Engine/Source/Runtime/GeometryCore/Public/Image/ImageBuilder.h
--- a/Engine/Source/Runtime/GeometryCore/Public/Image/ImageBuilder.h
+++ b/Engine/Source/Runtime/GeometryCore/Public/Image/ImageBuilder.h
@@ -178,7 +178,7 @@ public:
 		int64 NumPixels = Dimensions.Num();
 		for (int64 k = 0; k < NumPixels; ++k)
 		{
-			ConvertedImageOut.Image[k] = ConvertFunc(Image[k]);
+			ConvertedImageOut.SetPixel(k, ConvertFunc(Image[k]));
 		}
 	}
 
