MeshModelingTools: Make FSparseDynamicOctree3 capable of gathering all cells that intersect an arbitrary shape. Also fix some overflow errors and clamp to allocated range in range queries to avoid hanging when given huge bounds to check.

diff --git a/Engine/Source/Runtime/GeometryCore/Private/Spatial/SparseDynamicOctree3.cpp b/Engine/Source/Runtime/GeometryCore/Private/Spatial/SparseDynamicOctree3.cpp
--- a/Engine/Source/Runtime/GeometryCore/Private/Spatial/SparseDynamicOctree3.cpp
+++ b/Engine/Source/Runtime/GeometryCore/Private/Spatial/SparseDynamicOctree3.cpp
@@ -312,6 +312,79 @@ int32 FSparseDynamicOctree3::FindNearestHitObject(const FRay3d& Ray,
 
 
 
+// helper to find the root-level cells that could intersect with a given query object, specialized for range queries
+TArray<const FSparseOctreeCell*, TInlineAllocator<32>> FSparseDynamicOctree3::InitializeQueryQueue(
+	const FAxisAlignedBox3d& RangeBounds,
+	TFunctionRef<bool(const FAxisAlignedBox3d& CellBounds)> BoundsOverlapFn,
+	bool bCheckOverlapFnInRangeBounds) const
+{
+	TArray<const FSparseOctreeCell*, TInlineAllocator<32>> Queue;
+
+	// Skip range iteration if there are not many root cells -- should be faster to just directly iterate all the root cells in that case
+	constexpr int32 MinCountForRangeQuery = 10;
+	if (RootCells.GetCount() > MinCountForRangeQuery)
+	{
+		FVector3d RootBoundExpand(RootDimension * MaxExpandFactor);
+		FVector3i RootMinIndex = PointToIndex(0, RangeBounds.Min - RootBoundExpand);
+		FVector3i RootMaxIndex = PointToIndex(0, RangeBounds.Max + RootBoundExpand);
+		
+		// Clamp our iteration to only the range of allocated cells. Do it only in the appropriate
+		//  direction for the max/min (e.g. don't clamp min to be <= allocated max, to make it easier
+		//  to detect when we were entirely out of allocated bounds).
+		FAxisAlignedBox3i IndexBounds = RootCells.GetBoundsInclusive();
+		RootMinIndex = Max(RootMinIndex, IndexBounds.Min);
+		RootMaxIndex = Min(RootMaxIndex, IndexBounds.Max);
+		
+		// See if the range bounds were outside of the allocated bounds entirely. This check isn't
+		//  strictly necessary since the RangeIteration call below shouldn't do anything if Max < Min,
+		//  but we'll be explicit about this.
+		for (int i = 0; i < 3; ++i)
+		{
+			if (RootMinIndex[i] > RootMaxIndex[i])
+			{
+				return Queue;
+			}
+		}
+		
+		// Double-check that there are enough root cells that we are likely to save time by local iteration,
+		//  since we could have a huge iteration range but only a couple cells allocated in that range.
+		// Use double for the query size here to avoid overflow.
+		FVector3d QuerySize(RootMaxIndex - RootMinIndex + FVector3i(1, 1, 1));
+		if (RootCells.GetCount() > QuerySize.X * QuerySize.Y * QuerySize.Z) // There are enough root cells that we could save time by local iteration
+		{
+			if (bCheckOverlapFnInRangeBounds)
+			{
+				RootCells.RangeIteration(RootMinIndex, RootMaxIndex, [&Queue, &BoundsOverlapFn, this](uint32 RootCellID)
+				{
+					const FSparseOctreeCell* RootCell = &Cells[RootCellID];
+					if (BoundsOverlapFn(GetCellBox(*RootCell, MaxExpandFactor)))
+					{
+						Queue.Add(&Cells[RootCellID]);
+					}
+				});
+			}
+			else
+			{
+				RootCells.RangeIteration(RootMinIndex, RootMaxIndex, [&Queue, this](uint32 RootCellID)
+				{
+					// No need to check bound intersection since we're only iterating over cells that are in bounds
+					Queue.Add(&Cells[RootCellID]);
+				});
+			}
+			return Queue;
+		}
+	}
+
+	RootCells.AllocatedIteration([&Queue, &BoundsOverlapFn, this](const uint32* RootCellID)
+	{
+		const FSparseOctreeCell* RootCell = &Cells[*RootCellID];
+		if (BoundsOverlapFn(GetCellBox(*RootCell, MaxExpandFactor)))
+		{
+			Queue.Add(&Cells[*RootCellID]);
+		}
+	});
+	return Queue;
+}
 
 void FSparseDynamicOctree3::ContainmentQuery(
 	const FVector3d& Point,
@@ -496,19 +569,52 @@ void FSparseDynamicOctree3::ParallelRangeQuery(
 		ObjectIDs.Add(ObjectID);
 	}
 
-	TArray<const FSparseOctreeCell*, TInlineAllocator<32>> Queue = InitializeQueryQueue(Bounds);
+	TArray<const FSparseOctreeCell*, TInlineAllocator<32>> RelevantRootCells = InitializeQueryQueue(Bounds);
+
+	ParallelCollectObjectsInOverlappingCells(RelevantRootCells, 
+		[&Bounds](const FAxisAlignedBox3d& CellBounds) { return CellBounds.Intersects(Bounds); }, 
+		ObjectIDs);
+}
+
+void FSparseDynamicOctree3::ParallelRangeQuery(
+	const FAxisAlignedBox3d& RangeBoundsHint,
+	TFunctionRef<bool(const FAxisAlignedBox3d& CellBounds)> BoundsOverlapFn, 
+	TArray<int32>& ObjectIDs) const
+{
+	// Same as the box version above, except that (aside from using different overlap
+	//  criteria) our root cell initialization still needs to call BoundsOverlapFn
+	//  inside RangeBoundsHint.
+
+	// always collect spill objects
+	for (int ObjectID : SpillObjectSet)
+	{
+		ObjectIDs.Add(ObjectID);
+	}
+
+	TArray<const FSparseOctreeCell*, TInlineAllocator<32>> RelevantRootCells = InitializeQueryQueue(
+		RangeBoundsHint, BoundsOverlapFn, true);
+
+	ParallelCollectObjectsInOverlappingCells(RelevantRootCells, BoundsOverlapFn, ObjectIDs);
+
+}
+
+void FSparseDynamicOctree3::ParallelCollectObjectsInOverlappingCells(
+	TArray<const FSparseOctreeCell*, TInlineAllocator<32>> RelevantRootCells,
+	TFunctionRef<bool(const FAxisAlignedBox3d& CellBounds)> BoundsOverlapFn,
+	TArray<int32>& ObjectIDs) const
+{
 
 	// parallel strategy here is to collect each rootcell subtree into a separate
 	// array, and then merge them. Although this means we do some dynamic memory
 	// allocation, it's much faster than trying to lock ObjectIDs on every Add(),
 	// which results in crazy lock contention
 	FCriticalSection ObjectIDsLock;
-	ParallelFor(Queue.Num(), [&](int32 qi)
+	ParallelFor(RelevantRootCells.Num(), [&](int32 qi)
 	{
 		TArray<int32> LocalObjectIDs;
-		const FSparseOctreeCell* RootCell = Queue[qi];
-		BranchRangeQuery(RootCell, Bounds, LocalObjectIDs);
-		
+		const FSparseOctreeCell* RootCell = RelevantRootCells[qi];
+		BranchRangeQuery(RootCell, BoundsOverlapFn, LocalObjectIDs);
+
 		ObjectIDsLock.Lock();
 		for (int32 ObjectID : LocalObjectIDs)
 		{
@@ -555,7 +661,17 @@ int FSparseDynamicOctree3::ParallelOverlapAnyQuery(const FAxisAlignedBox3d& Shap
 
 void FSparseDynamicOctree3::BranchRangeQuery(
 	const FSparseOctreeCell* ParentCell,
-	const FAxisAlignedBox3d& Bounds,
+	const FAxisAlignedBox3d& BoundsIn,
+	TArray<int>& ObjectIDs) const
+{
+	BranchRangeQuery(ParentCell, 
+		[&BoundsIn](const FAxisAlignedBox3d& CellBounds) { return CellBounds.Intersects(BoundsIn); }, 
+		ObjectIDs);
+}
+
+void FSparseDynamicOctree3::BranchRangeQuery(
+	const FSparseOctreeCell* ParentCell, 
+	TFunctionRef<bool(const FAxisAlignedBox3d&)> BoundsOverlapFn, 
 	TArray<int>& ObjectIDs) const
 {
 	TArray<const FSparseOctreeCell*, TInlineAllocator<32>> Queue;
@@ -576,7 +692,7 @@ void FSparseDynamicOctree3::BranchRangeQuery(
 			if (CurCell->HasChild(k))
 			{
 				const FSparseOctreeCell* ChildCell = &Cells[CurCell->GetChildCellID(k)];
-				if (GetCellBox(*ChildCell, MaxExpandFactor).Intersects(Bounds))
+				if (BoundsOverlapFn(GetCellBox(*ChildCell, MaxExpandFactor)))
 				{
 					Queue.Add(ChildCell);
 				}
diff --git a/Engine/Source/Runtime/GeometryCore/Public/Spatial/SparseDynamicOctree3.h b/Engine/Source/Runtime/GeometryCore/Public/Spatial/SparseDynamicOctree3.h
--- a/Engine/Source/Runtime/GeometryCore/Public/Spatial/SparseDynamicOctree3.h
+++ b/Engine/Source/Runtime/GeometryCore/Public/Spatial/SparseDynamicOctree3.h
@@ -313,6 +313,19 @@ public:
 	GEOMETRYCORE_API void ParallelRangeQuery(const FAxisAlignedBox3d& Bounds,
 		TArray<int>& ObjectIDsOut) const;
 
+	/**
+	 * Collect ObjectIDs from all the leaf cells that have bounding boxes that pass BoundsOverlapFn
+	 * @param RangeBoundsHint A box to potentially limit the iteration range across cells. Note that
+	 *  this might be ignored if the allocation of cells is such that trying to check for cells in that
+	 *  range would take more time than just iterating over all allocated cells.
+	 * @param BoundsOverlapFn Must be thread-safe. Used to determine if the cell is part
+	 *  of the desired range.
+	 * @param ObjectIDsOut collected ObjectIDs are stored here
+	 */
+	GEOMETRYCORE_API void ParallelRangeQuery(
+		const FAxisAlignedBox3d& RangeBoundsHint,
+		TFunctionRef<bool (const FAxisAlignedBox3d& CellBounds)> BoundsOverlapFn,
+		TArray<int32>& ObjectIDs) const;
 
 	/**
 	 * Find any overlap between a caller-defined query and any object ID. Returns the first overlap it finds.
@@ -443,9 +456,20 @@ protected:
 	FVector3i PointToIndex(uint32 Level, const FVector3d& Position) const
 	{
 		double CellWidth = GetCellWidth(Level);
-		int32 i = (int32)FMathd::Floor(Position.X / CellWidth);
-		int32 j = (int32)FMathd::Floor(Position.Y / CellWidth);
-		int32 k = (int32)FMathd::Floor(Position.Z / CellWidth);
+
+		// Converts to an int32 in such a way that values outside of the representable
+		//  range ends up at the min and max values of the representable range.
+		auto SafeDoubleToInt32Floor = [](double Value)
+		{
+			return static_cast<int32>(FMath::Clamp(
+				FMathd::Floor(Value),
+				static_cast<double>(TNumericLimits<int32>::Lowest()), 
+				static_cast<double>(TNumericLimits<int32>::Max())));
+		};
+
+		int32 i = SafeDoubleToInt32Floor(Position.X / CellWidth);
+		int32 j = SafeDoubleToInt32Floor(Position.Y / CellWidth);
+		int32 k = SafeDoubleToInt32Floor(Position.Z / CellWidth);
 		return FVector3i(i, j, k);
 	}
 
@@ -491,80 +515,50 @@ protected:
 		TArray<int>& ObjectIDs) const;
 
 private:
+	// Gather objects in all cells that pass BoundsOverlapFn, starting from a set of root cells
+	void ParallelCollectObjectsInOverlappingCells(
+		TArray<const FSparseOctreeCell*, TInlineAllocator<32>> RelevantRootCells,
+		TFunctionRef<bool(const FAxisAlignedBox3d& CellBounds)> BoundsOverlapFn,
+		TArray<int32>& ObjectIDs) const;
+	// Go down subtree of ParentCell and collect all objects in cells that pass BoundsOverlapFn.
+	void BranchRangeQuery(const FSparseOctreeCell* ParentCell,
+		TFunctionRef<bool(const FAxisAlignedBox3d&)> BoundsOverlapFn,
+		TArray<int>& ObjectIDs) const;
+
 	int BranchCustomOverlapAnyQuery(const FSparseOctreeCell* ParentCell,
 		const FAxisAlignedBox3d& Bounds,
 		TFunctionRef<bool(int32)> ObjectOverlapFn, TFunctionRef<bool(const FAxisAlignedBox3d&)> BoundsOverlapFn) const;
 
+	/**
+	 * Helper to find root-level cells that satisfy BoundsOverlapFn
+	 * 
+	 * @param RangeBounds If the allocation of cells is such that it is worth doing, only cells in these
+	 *   bounds are looked for and iterated over. May be ignored if the allocation is such that it is
+	 *   faster to iterate over all allocated root cells than to check for allocation in a big range.
+	 * @param BoundsOverlapFn Function for determining whether cell should be included in output queue
+	 * @param bCheckOverlapFnInRangeBounds If true, then BoundsOverlapFn is still checked when iterating
+	 *   over RangeBounds (i.e., RangeBounds is an optimization hint, but not sufficient for determining
+	 *   actual overlap).
+	 */
+	TArray<const FSparseOctreeCell*, TInlineAllocator<32>> InitializeQueryQueue(
+		const FAxisAlignedBox3d& RangeBounds, 
+		TFunctionRef<bool(const FAxisAlignedBox3d& CellBounds)> BoundsOverlapFn, 
+		bool bCheckOverlapFnInRangeBounds = true) const;
+
 	// helper to find the root-level cells that could intersect with a given query object, specialized for point containment queries
 	TArray<const FSparseOctreeCell*, TInlineAllocator<32>> InitializeQueryQueue(const FVector3d& Point) const
 	{
-		TArray<const FSparseOctreeCell*, TInlineAllocator<32>> Queue;
-
-		// Skip range iteration if there are not many root cells -- should be faster to just directly iterate all the root cells in that case
-		constexpr int32 MinCountForRangeQuery = 10;
-		if (RootCells.GetCount() > MinCountForRangeQuery)
-		{
-			// because of the cell expand factor, search a range for cells that contain the point
-			FVector3d RootBoundExpand(RootDimension * MaxExpandFactor);
-			FVector3i RootMinIndex = PointToIndex(0, Point - RootBoundExpand);
-			FVector3i RootMaxIndex = PointToIndex(0, Point + RootBoundExpand);
-			// double-check that there are enough root cells that we are likely to save time by local iteration
-			FVector3i QuerySize = RootMaxIndex - RootMinIndex + FVector3i(1, 1, 1);
-			if (RootCells.GetCount() > QuerySize.X * QuerySize.Y * QuerySize.Z)
-			{
-				RootCells.RangeIteration(RootMinIndex, RootMaxIndex, [&](uint32 RootCellID)
-				{
-					Queue.Add(&Cells[RootCellID]);
-				});
-				return Queue;
-			}
-		}
-
-		RootCells.AllocatedIteration([&](const uint32* RootCellID)
-		{
-			const FSparseOctreeCell* RootCell = &Cells[*RootCellID];
-			if (GetCellBox(*RootCell, MaxExpandFactor).Contains(Point))
-			{
-				Queue.Add(&Cells[*RootCellID]);
-			}
-		});
-		return Queue;
+		return InitializeQueryQueue(FAxisAlignedBox3d(Point, Point),
+			[&Point](const FAxisAlignedBox3d& CellBounds) { return CellBounds.Contains(Point); }, 
+			false);
 	}
 
-	// helper to find the root-level cells that could intersect with a given query object, specialized for range queries
+	// helper to find the root-level cells that could intersect with a given query object, specialized for box queries
 	TArray<const FSparseOctreeCell*, TInlineAllocator<32>> InitializeQueryQueue(const FAxisAlignedBox3d& Bounds) const
 	{
-		TArray<const FSparseOctreeCell*, TInlineAllocator<32>> Queue;
-
-		// Skip range iteration if there are not many root cells -- should be faster to just directly iterate all the root cells in that case
-		constexpr int32 MinCountForRangeQuery = 10;
-		if (RootCells.GetCount() > MinCountForRangeQuery)
-		{
-			FVector3d RootBoundExpand(RootDimension * MaxExpandFactor);
-			FVector3i RootMinIndex = PointToIndex(0, Bounds.Min - RootBoundExpand);
-			FVector3i RootMaxIndex = PointToIndex(0, Bounds.Max + RootBoundExpand);
-			// double-check that there are enough root cells that we are likely to save time by local iteration
-			FVector3i QuerySize = RootMaxIndex - RootMinIndex + FVector3i(1, 1, 1);
-			if (RootCells.GetCount() > QuerySize.X * QuerySize.Y * QuerySize.Z) // There are enough root cells that we could save time by local iteration
-			{
-				RootCells.RangeIteration(RootMinIndex, RootMaxIndex, [&](uint32 RootCellID)
-				{
-					// No need to check bound intersection since we're only iterating over cells that are in bounds
-					Queue.Add(&Cells[RootCellID]);
-				});
-				return Queue;
-			}
-		}
-
-		RootCells.AllocatedIteration([&](const uint32* RootCellID)
-		{
-			const FSparseOctreeCell* RootCell = &Cells[*RootCellID];
-			if (GetCellBox(*RootCell, MaxExpandFactor).Intersects(Bounds))
-			{
-				Queue.Add(&Cells[*RootCellID]);
-			}
-		});
-		return Queue;
+		return InitializeQueryQueue(Bounds,
+			[&Bounds](const FAxisAlignedBox3d& CellBounds) { return CellBounds.Intersects(Bounds); },
+			false);
 	}
 };
 
