add 'nearest point' and 'k nearest points' queries to sparse dynamic point octree

diff --git a/Engine/Source/Runtime/GeometryCore/Private/Spatial/SparseDynamicPointOctree3.cpp b/Engine/Source/Runtime/GeometryCore/Private/Spatial/SparseDynamicPointOctree3.cpp
new file mode 100644
--- /dev/null
+++ b/Engine/Source/Runtime/GeometryCore/Private/Spatial/SparseDynamicPointOctree3.cpp
@@ -0,0 +1,219 @@
+// Copyright Epic Games, Inc. All Rights Reserved.
+
+
+#include "Spatial/SparseDynamicPointOctree3.h"
+
+namespace UE::Geometry
+{
+
+
+void FSparseDynamicPointOctree3::IterateRootsInRadius(FVector3d QueryPt, double InitialRadius, TFunctionRef<double(double, const FSparsePointOctreeCell&)> RootFn) const
+{
+	FAxisAlignedBox3d Bounds(QueryPt, InitialRadius);
+	double ThresholdSq = InitialRadius * InitialRadius;
+
+	FVector3i RootMinIndex = PointToIndex(0, Bounds.Min);
+	FVector3i RootMaxIndex = PointToIndex(0, Bounds.Max);
+	
+	
+	// Clamp our iteration to only the range of allocated cells. Do it only in the appropriate
+	//  direction for the max/min (e.g. don't clamp min to be <= allocated max, to make it easier
+	//  to detect when we were entirely out of allocated bounds).
+	FAxisAlignedBox3i IndexBounds = RootCells.GetBoundsInclusive();
+	RootMinIndex = Max(RootMinIndex, IndexBounds.Min);
+	RootMaxIndex = Min(RootMaxIndex, IndexBounds.Max);
+	
+	// See if the range bounds were outside of the allocated bounds entirely. This check isn't
+	//  strictly necessary since the RangeIteration call below shouldn't do anything if Max < Min,
+	//  but we'll be explicit about this.
+	for (int i = 0; i < 3; ++i)
+	{
+		if (RootMinIndex[i] > RootMaxIndex[i])
+		{
+			return;
+		}
+	}
+	
+	// Double-check that there are enough root cells that we are likely to save time by local iteration,
+	//  since we could have a huge iteration range but only a couple cells allocated in that range.
+	// Use double for the query size here to avoid overflow.
+	FVector3d QuerySize(RootMaxIndex - RootMinIndex + FVector3i(1, 1, 1));
+	if (RootCells.GetCount() > QuerySize.X * QuerySize.Y * QuerySize.Z) // There are enough root cells that we could save time by local iteration
+	{
+		RootCells.RangeIteration(RootMinIndex, RootMaxIndex, [this, QueryPt, &ThresholdSq, &RootFn](uint32 RootCellID)
+		{
+			const FSparsePointOctreeCell* RootCell = &Cells[RootCellID];
+			if (GetCellBox(*RootCell).DistanceSquared(QueryPt) < ThresholdSq)
+			{
+				ThresholdSq = RootFn(ThresholdSq, *RootCell);
+			}
+		});
+	}
+	else
+	{
+		RootCells.AllocatedIteration([this, QueryPt, &ThresholdSq, &RootFn](const uint32* RootCellID)
+		{
+			const FSparsePointOctreeCell* RootCell = &Cells[*RootCellID];
+			if (GetCellBox(*RootCell).DistanceSquared(QueryPt) < ThresholdSq)
+			{
+				ThresholdSq = RootFn(ThresholdSq, *RootCell);
+			}
+		});
+	}
+
+}
+
+void FSparseDynamicPointOctree3::FindKClosestPoints(FVector3d QueryPt, double DistanceThreshold, int32 NumToFind,
+	TArray<TPair<int32, double>>& FoundPoints, TFunctionRef<bool(int32)> PredicateFunc, TFunctionRef<double(int32)> DistSqFunc,
+	TArray<const FSparsePointOctreeCell*>* TempBuffer) const
+{
+	FoundPoints.Reset();
+	if (NumToFind <= 0)
+	{
+		return;
+	}
+	FoundPoints.Reserve(NumToFind);
+
+	TArray<const FSparsePointOctreeCell*> InternalBuffer;
+	TArray<const FSparsePointOctreeCell*>& Queue =
+		(TempBuffer == nullptr) ? InternalBuffer : *TempBuffer;
+	Queue.Reset(128);
+
+	IterateRootsInRadius(QueryPt, DistanceThreshold, [this, QueryPt, NumToFind, &FoundPoints, &PredicateFunc, &DistSqFunc, &Queue](double CurThresholdSq, const FSparsePointOctreeCell& Cell)
+	{
+		// Set a heap predicate so we can quickly access the largest point in the array
+		auto HeapPredicate = [](const TPair<int32, double>& A, const TPair<int32, double>& B) -> bool
+		{
+			return A.Value > B.Value;
+		};
+
+		Queue.Add(&Cell);
+
+		while (Queue.Num() > 0)
+		{
+			const FSparsePointOctreeCell* CurCell = Queue.Pop(EAllowShrinking::No);
+
+			// check if we can skip the cell based on its bounding box distance
+			// (since BestDistSq may have been reduced after we added it to the queue)
+			if (GetCellBox(*CurCell).DistanceSquared(QueryPt) >= CurThresholdSq)
+			{
+				continue;
+			}
+
+			// process points on the cell
+			if (CellPointLists.IsAllocated(CurCell->CellID))
+			{
+				CellPointLists.Enumerate(CurCell->CellID, [&](int32 PointID)
+				{
+					if (PredicateFunc(PointID))
+					{
+						double DSq = DistSqFunc(PointID);
+						if (DSq < CurThresholdSq)
+						{
+							if (FoundPoints.Num() < NumToFind)
+							{
+								FoundPoints.HeapPush(TPair<int32, double>(PointID, DSq), HeapPredicate);
+							}
+							else
+							{
+								// take the farthest point off the heap, and put the new point on
+								TPair<int32, double> BiggestItem_Unused;
+								FoundPoints.HeapPop(BiggestItem_Unused, HeapPredicate, EAllowShrinking::No);
+								FoundPoints.HeapPush(TPair<int32, double>(PointID, DSq), HeapPredicate);
+							}
+
+							// once we've found our first K points, we can refine the search radius based on the farthest of them
+							if (FoundPoints.Num() == NumToFind)
+							{
+								CurThresholdSq = FoundPoints[0].Value;
+							}
+						}
+					}
+				});
+			}
+
+			// add child cells to queue (if within current threshold distance)
+			for (int k = 0; k < 8; ++k)
+			{
+				if (CurCell->HasChild(k))
+				{
+					const FSparsePointOctreeCell* ChildCell = &Cells[CurCell->GetChildCellID(k)];
+					if (GetCellBox(*ChildCell).DistanceSquared(QueryPt) < CurThresholdSq)
+					{
+						Queue.Add(ChildCell);
+					}
+				}
+			}
+		}
+
+		return CurThresholdSq;
+	});
+}
+
+
+int32 FSparseDynamicPointOctree3::FindClosestPoint(FVector3d QueryPt, double DistanceThreshold,
+	TFunctionRef<bool(int32)> PredicateFunc, TFunctionRef<double(int32)> DistSqFunc,
+	TArray<const FSparsePointOctreeCell*>* TempBuffer) const
+{
+	TArray<const FSparsePointOctreeCell*> InternalBuffer;
+	TArray<const FSparsePointOctreeCell*>& Queue =
+		(TempBuffer == nullptr) ? InternalBuffer : *TempBuffer;
+	
+	double BestDistSq = DistanceThreshold * DistanceThreshold;
+	int32 BestPtIdx = INDEX_NONE;
+
+	Queue.Reset(128);
+
+	IterateRootsInRadius(QueryPt, DistanceThreshold, [this, QueryPt, &BestPtIdx, &PredicateFunc, &DistSqFunc, &Queue](double BestDistSq, const FSparsePointOctreeCell& Cell)
+	{
+		Queue.Add(&Cell);
+		while (Queue.Num() > 0)
+		{
+			const FSparsePointOctreeCell* CurCell = Queue.Pop(EAllowShrinking::No);
+
+			// check if we can skip the cell based on its bounding box distance
+			// (since BestDistSq may have been reduced after we added it to the queue)
+			if (GetCellBox(*CurCell).DistanceSquared(QueryPt) >= BestDistSq)
+			{
+				continue;
+			}
+
+			// process points on the cell
+			if (CellPointLists.IsAllocated(CurCell->CellID))
+			{
+				CellPointLists.Enumerate(CurCell->CellID, [&](int32 PointID)
+				{
+					if (PredicateFunc(PointID))
+					{
+						double DSq = DistSqFunc(PointID);
+						if (DSq < BestDistSq)
+						{
+							BestPtIdx = PointID;
+							BestDistSq = DSq;
+						}
+					}
+				});
+			}
+
+			// add child cells to queue (if close enough)
+			for (int k = 0; k < 8; ++k)
+			{
+				if (CurCell->HasChild(k))
+				{
+					const FSparsePointOctreeCell* ChildCell = &Cells[CurCell->GetChildCellID(k)];
+					if (GetCellBox(*ChildCell).DistanceSquared(QueryPt) < BestDistSq)
+					{
+						Queue.Add(ChildCell);
+					}
+				}
+			}
+		}
+
+		return BestDistSq;
+	});
+
+	return BestPtIdx;
+}
+
+} // namespace UE::Geometry
+
diff --git a/Engine/Source/Runtime/GeometryCore/Public/Spatial/SparseDynamicPointOctree3.h b/Engine/Source/Runtime/GeometryCore/Public/Spatial/SparseDynamicPointOctree3.h
--- a/Engine/Source/Runtime/GeometryCore/Public/Spatial/SparseDynamicPointOctree3.h
+++ b/Engine/Source/Runtime/GeometryCore/Public/Spatial/SparseDynamicPointOctree3.h
@@ -211,6 +211,35 @@ public:
 		TFunctionRef<double(int, const FRay3d&)> HitPointDistFunc,
 		double MaxDistance = TNumericLimits<double>::Max()) const;
 
+	/**
+	 * Find closest point in the tree, relative to a query point, within a search distance
+	 * 
+	 * @param QueryPt Find the point closest to this query point
+	 * @param DistanceThreshold Limit the search to this distance
+	 * @param PredicateFunc Return true if point should be included (i.e., this is a filter)
+	 * @param DistSqFunc User-provided squared distance between the point with the given ID and the QueryPt
+	 * @param TempBuffer Optional temporary buffer to store a queue of cells (to avoid memory allocations)
+	 * @return The ID of the closest point, or INDEX_NONE if no point was found
+	 */
+	int32 FindClosestPoint(FVector3d QueryPt, double DistanceThreshold,
+		TFunctionRef<bool(int32)> PredicateFunc, TFunctionRef<double(int32)> DistSqFunc,
+		TArray<const FSparsePointOctreeCell*>* TempBuffer = nullptr) const;
+
+	/**
+	 * Find the K closest point in the tree, relative to a query point, within a search distance
+	 *
+	 * @param QueryPt Find the points closest to this query point
+	 * @param DistanceThreshold Limit the search to this distance
+	 * @param NumToFind The maximum number of closest points to find
+	 * @param FoundPoints The found closest points will be added to this array
+	 * @param PredicateFunc Return true if point should be included (i.e., this is a filter)
+	 * @param DistSqFunc User-provided squared distance between the point with the given ID and the QueryPt
+	 * @param TempBuffer Optional temporary buffer to store a queue of cells (to avoid memory allocations)
+	 */
+	void FindKClosestPoints(FVector3d QueryPt, double DistanceThreshold, int32 NumToFind,
+		TArray<TPair<int32, double>>& FoundPoints, TFunctionRef<bool(int32)> PredicateFunc, TFunctionRef<double(int32)> DistSqFunc,
+		TArray<const FSparsePointOctreeCell*>* TempBuffer = nullptr) const;
+
 	/**
 	 * Collect PointIDs from all the cells with bounding boxes that intersect Bounds,
 	 * where PredicateFunc passes
@@ -335,9 +364,20 @@ protected:
 	inline FVector3i PointToIndex(uint32 Level, const FVector3d& Position) const
 	{
 		double CellWidth = GetCellWidth(Level);
-		int32 i = (int32)FMathd::Floor(Position.X / CellWidth);
-		int32 j = (int32)FMathd::Floor(Position.Y / CellWidth);
-		int32 k = (int32)FMathd::Floor(Position.Z / CellWidth);
+
+		// Converts to an int32 in such a way that values outside of the representable
+		//  range ends up at the min and max values of the representable range.
+		auto SafeDoubleToInt32Floor = [](double Value)
+		{
+			return static_cast<int32>(FMath::Clamp(
+				FMathd::Floor(Value),
+				static_cast<double>(TNumericLimits<int32>::Lowest()),
+				static_cast<double>(TNumericLimits<int32>::Max())));
+		};
+		
+		int32 i = SafeDoubleToInt32Floor(Position.X / CellWidth);
+		int32 j = SafeDoubleToInt32Floor(Position.Y / CellWidth);
+		int32 k = SafeDoubleToInt32Floor(Position.Z / CellWidth);
 		return FVector3i(i, j, k);
 	}
 
@@ -379,6 +419,10 @@ protected:
 
 	inline double FindNearestRayCellIntersection(const FSparsePointOctreeCell& Cell, const FRay3d& Ray) const;
 
+private:
+	// Helper to iterate over the root cells within a radius. Note the RootFn can return a refined search radius squared, which will let the iteration visit fewer root cells.
+	// RootFn is passed the current squared search radius, and returns the (possibly smaller) squared search radius.
+	void IterateRootsInRadius(FVector3d QueryPt, double InitialRadius, TFunctionRef<double(double RadSq, const FSparsePointOctreeCell& Root)> RootFn) const;
 };
 
 
