add test for unset triangle in TDynamicMeshOverlay GetElementIDAtVertex method

diff --git a/Engine/Source/Runtime/GeometryCore/Private/DynamicMesh/DynamicMeshOverlay.cpp b/Engine/Source/Runtime/GeometryCore/Private/DynamicMesh/DynamicMeshOverlay.cpp
--- a/Engine/Source/Runtime/GeometryCore/Private/DynamicMesh/DynamicMeshOverlay.cpp
+++ b/Engine/Source/Runtime/GeometryCore/Private/DynamicMesh/DynamicMeshOverlay.cpp
@@ -919,6 +919,11 @@ template<typename RealType, int ElementSize>
 int TDynamicMeshOverlay<RealType, ElementSize>::GetElementIDAtVertex(int TriangleID, int VertexID) const
 {
 	FIndex3i Triangle = GetTriangle(TriangleID);
+	// Check first element ID to see if triangle is unset (Note: triangles must be fully set or fully unset, so checking first element should be enough)
+	if (Triangle.A == FDynamicMesh3::InvalidID)
+	{
+		return FDynamicMesh3::InvalidID;
+	}
 	for (int IDX = 0; IDX < 3; ++IDX)
 	{	
 		int ElementID = Triangle[IDX];
