fix crash on compact in place of a dynamic mesh overlay when the mesh has zero triangles

diff --git a/Engine/Source/Runtime/GeometryCore/Public/DynamicMesh/DynamicMeshOverlay.h b/Engine/Source/Runtime/GeometryCore/Public/DynamicMesh/DynamicMeshOverlay.h
--- a/Engine/Source/Runtime/GeometryCore/Public/DynamicMesh/DynamicMeshOverlay.h
+++ b/Engine/Source/Runtime/GeometryCore/Public/DynamicMesh/DynamicMeshOverlay.h
@@ -209,7 +209,7 @@ public:
 		ParentVertices.Resize(ElementCount());
 
 		// Remap and compact triangle element indices.
-		int32 MaxNewTID = 0;
+		int32 MaxNewTID = -1;
 		for (int TID = 0, OldMaxTID = ElementTriangles.Num() / 3; TID < OldMaxTID; TID++)
 		{
 			const int32 OldStart = TID * 3;
