make dynamic mesh overlay more robust to the case where a seam edge is flipped (e.g., due to not setting constraints correctly before remeshing)

diff --git a/Engine/Source/Runtime/GeometryCore/Private/DynamicMesh/DynamicMeshOverlay.cpp b/Engine/Source/Runtime/GeometryCore/Private/DynamicMesh/DynamicMeshOverlay.cpp
--- a/Engine/Source/Runtime/GeometryCore/Private/DynamicMesh/DynamicMeshOverlay.cpp
+++ b/Engine/Source/Runtime/GeometryCore/Private/DynamicMesh/DynamicMeshOverlay.cpp
@@ -1080,11 +1080,19 @@ void TDynamicMeshOverlay<RealType, ElementSize>::OnFlipEdge(const FDynamicMesh3:
 	int orig_t0 = FlipInfo.Triangles.A;
 	int orig_t1 = FlipInfo.Triangles.B;
 	bool bT0Set = IsSetTriangle(orig_t0), bT1Set = IsSetTriangle(orig_t1);
-	if (!bT0Set)
+	int32 NumSet = (int32)bT0Set + (int32)bT1Set;
+	if (NumSet == 0)
 	{
-		ensure(!bT1Set); // flipping across a set/unset boundary is not allowed?
 		return; // nothing to do on the overlay if both triangles are unset
 	}
+	else if (NumSet == 1)
+	{
+		ensure(false); // flipping across a set/unset boundary is not allowed?
+		// recover by just unsetting both triangles, since it's too late to prevent the flip
+		UnsetTriangle(orig_t0);
+		UnsetTriangle(orig_t1);
+		return;
+	}
 
 	int base_a = FlipInfo.OriginalVerts.A;
 	int base_b = FlipInfo.OriginalVerts.B;
@@ -1110,7 +1118,7 @@ void TDynamicMeshOverlay<RealType, ElementSize>::OnFlipEdge(const FDynamicMesh3:
 	checkSlow(idx_base_d == BaseTriangle1.IndexOf(base_d));
 
 	// we should not have been called on a non-shared edge!!
-	bool bHasSharedUVEdge = IndexUtil::SamePairUnordered(Triangle0[idx_base_a1], Triangle0[idx_base_b1], Triangle1[idx_base_a2], Triangle1[idx_base_b2]);
+	bool bHasSharedUVEdge = (Triangle0[idx_base_a1] == Triangle1[idx_base_a2]) && (Triangle0[idx_base_b1] == Triangle1[idx_base_b2]);
 	checkSlow(bHasSharedUVEdge);
 
 	int A = Triangle0[idx_base_a1];
@@ -1125,10 +1133,35 @@ void TDynamicMeshOverlay<RealType, ElementSize>::OnFlipEdge(const FDynamicMesh3:
 	ElementTriangles[i1] = D; ElementTriangles[i1+1] = C; ElementTriangles[i1+2] = A;
 
 	// update reference counts
-	ElementsRefCounts.Decrement(A);
-	ElementsRefCounts.Decrement(B);
 	ElementsRefCounts.Increment(C);
 	ElementsRefCounts.Increment(D);
+	if (bHasSharedUVEdge)
+	{
+		ElementsRefCounts.Decrement(A);
+		ElementsRefCounts.Decrement(B);
+	}
+	else // flipped a seam edge case
+	{
+		// ideally we wouldn't flip a seam edge, but if we have done so, may need to clean up a now-unreferenced element
+		// (note it's too late to prevent the flip here)
+		int A2 = Triangle1[idx_base_a2];
+		int B2 = Triangle1[idx_base_b2];
+
+		ElementsRefCounts.Decrement(A2);
+		// free A2 if it's no longer referenced (can happen in the A != A2 case)
+		if (ElementsRefCounts.GetRefCount(A2) == 1)
+		{
+			ElementsRefCounts.Decrement(A2);
+			ParentVertices[A2] = FDynamicMesh3::InvalidID;
+		}
+		ElementsRefCounts.Decrement(B2);
+		// free B2 if it's no longer referenced (can happen in the B != B2 case)
+		if (ElementsRefCounts.GetRefCount(B2) == 1)
+		{
+			ElementsRefCounts.Decrement(B2);
+			ParentVertices[B2] = FDynamicMesh3::InvalidID;
+		}
+	}
 }
 
 
