remove redundant calls to On[MeshOperation] methods for generic attributes -- since these methods will be called by the base class via the FDynamicMeshAttributeSetBase::On[MeshOp] calls

diff --git a/Engine/Source/Runtime/GeometryCore/Private/DynamicMesh/DynamicMeshAttributeSet.cpp b/Engine/Source/Runtime/GeometryCore/Private/DynamicMesh/DynamicMeshAttributeSet.cpp
--- a/Engine/Source/Runtime/GeometryCore/Private/DynamicMesh/DynamicMeshAttributeSet.cpp
+++ b/Engine/Source/Runtime/GeometryCore/Private/DynamicMesh/DynamicMeshAttributeSet.cpp
@@ -856,11 +856,6 @@ void FDynamicMeshAttributeSet::OnNewTriangle(int TriangleID, bool bInserted)
 		int32 NewGroup = 0;
 		PolygroupLayer.SetNewValue(TriangleID, &NewGroup);
 	}
-
-	for (TPair<FName, TUniquePtr<FDynamicMeshAttributeBase>>& AttribPair : GenericAttributes)
-	{
-		AttribPair.Value->OnNewTriangle(TriangleID, bInserted);
-	}
 }
 
 
@@ -881,12 +876,6 @@ void FDynamicMeshAttributeSet::OnRemoveTriangle(int TriangleID)
 		ColorLayer->OnRemoveTriangle(TriangleID);
 	}
 
-	// has no effect on MaterialIDAttrib
-
-	for (TPair<FName, TUniquePtr<FDynamicMeshAttributeBase>>& AttribPair : GenericAttributes)
-	{
-		AttribPair.Value->OnRemoveTriangle(TriangleID);
-	}
 }
 
 void FDynamicMeshAttributeSet::OnReverseTriOrientation(int TriangleID)
@@ -905,12 +894,6 @@ void FDynamicMeshAttributeSet::OnReverseTriOrientation(int TriangleID)
 	{
 		ColorLayer->OnReverseTriOrientation(TriangleID);
 	}
-	// has no effect on MaterialIDAttrib
-
-	for (TPair<FName, TUniquePtr<FDynamicMeshAttributeBase>>& AttribPair : GenericAttributes)
-	{
-		AttribPair.Value->OnReverseTriOrientation(TriangleID);
-	}
 }
 
 void FDynamicMeshAttributeSet::OnSplitEdge(const FDynamicMesh3::FEdgeSplitInfo& SplitInfo)
@@ -943,11 +926,6 @@ void FDynamicMeshAttributeSet::OnSplitEdge(const FDynamicMesh3::FEdgeSplitInfo&
 	{
 		WeightLayer.OnSplitEdge(SplitInfo);
 	}
-
-	for (TPair<FName, TUniquePtr<FDynamicMeshAttributeBase>>& AttribPair : GenericAttributes)
-	{
-		AttribPair.Value->OnSplitEdge(SplitInfo);
-	}
 }
 
 void FDynamicMeshAttributeSet::OnFlipEdge(const FDynamicMesh3::FEdgeFlipInfo & flipInfo)
@@ -980,11 +958,6 @@ void FDynamicMeshAttributeSet::OnFlipEdge(const FDynamicMesh3::FEdgeFlipInfo & f
 	{
 		WeightLayer.OnFlipEdge(flipInfo);
 	}
-
-	for (TPair<FName, TUniquePtr<FDynamicMeshAttributeBase>>& AttribPair : GenericAttributes)
-	{
-		AttribPair.Value->OnFlipEdge(flipInfo);
-	}
 }
 
 
@@ -1018,11 +991,6 @@ void FDynamicMeshAttributeSet::OnCollapseEdge(const FDynamicMesh3::FEdgeCollapse
 	{
 		WeightLayer.OnCollapseEdge(collapseInfo);
 	}
-
-	for (TPair<FName, TUniquePtr<FDynamicMeshAttributeBase>>& AttribPair : GenericAttributes)
-	{
-		AttribPair.Value->OnCollapseEdge(collapseInfo);
-	}
 }
 
 void FDynamicMeshAttributeSet::OnPokeTriangle(const FDynamicMesh3::FPokeTriangleInfo & pokeInfo)
@@ -1055,11 +1023,6 @@ void FDynamicMeshAttributeSet::OnPokeTriangle(const FDynamicMesh3::FPokeTriangle
 	{
 		WeightLayer.OnPokeTriangle(pokeInfo);
 	}
-
-	for (TPair<FName, TUniquePtr<FDynamicMeshAttributeBase>>& AttribPair : GenericAttributes)
-	{
-		AttribPair.Value->OnPokeTriangle(pokeInfo);
-	}
 }
 
 void FDynamicMeshAttributeSet::OnMergeEdges(const FDynamicMesh3::FMergeEdgesInfo & mergeInfo)
@@ -1092,11 +1055,6 @@ void FDynamicMeshAttributeSet::OnMergeEdges(const FDynamicMesh3::FMergeEdgesInfo
 	{
 		WeightLayer.OnMergeEdges(mergeInfo);
 	}
-
-	for (TPair<FName, TUniquePtr<FDynamicMeshAttributeBase>>& AttribPair : GenericAttributes)
-	{
-		AttribPair.Value->OnMergeEdges(mergeInfo);
-	}
 }
 
 void FDynamicMeshAttributeSet::OnMergeVertices(const DynamicMeshInfo::FMergeVerticesInfo& MergeInfo)
@@ -1129,11 +1087,6 @@ void FDynamicMeshAttributeSet::OnMergeVertices(const DynamicMeshInfo::FMergeVert
 	{
 		WeightLayer.OnMergeVertices(MergeInfo);
 	}
-
-	for (TPair<FName, TUniquePtr<FDynamicMeshAttributeBase>>& AttribPair : GenericAttributes)
-	{
-		AttribPair.Value->OnMergeVertices(MergeInfo);
-	}
 }
 
 void FDynamicMeshAttributeSet::OnSplitVertex(const DynamicMeshInfo::FVertexSplitInfo& SplitInfo, const TArrayView<const int>& TrianglesToUpdate)
@@ -1166,11 +1119,6 @@ void FDynamicMeshAttributeSet::OnSplitVertex(const DynamicMeshInfo::FVertexSplit
 	{
 		WeightLayer.OnSplitVertex(SplitInfo, TrianglesToUpdate);
 	}
-
-	for (TPair<FName, TUniquePtr<FDynamicMeshAttributeBase>>& AttribPair : GenericAttributes)
-	{
-		AttribPair.Value->OnSplitVertex(SplitInfo, TrianglesToUpdate);
-	}
 }
 
 bool FDynamicMeshAttributeSet::IsSameAs(const FDynamicMeshAttributeSet& Other, bool bIgnoreDataLayout) const
