fix dynamic mesh attribute set's handling of vertex attribute changes, and make code paths for reparenting and compacting use the registered attributes instead of the generic attributes (to prepare for more attributes using the registered attribute array)

diff --git a/Engine/Source/Runtime/GeometryCore/Private/DynamicMesh/DynamicMeshAttributeSet.cpp b/Engine/Source/Runtime/GeometryCore/Private/DynamicMesh/DynamicMeshAttributeSet.cpp
--- a/Engine/Source/Runtime/GeometryCore/Private/DynamicMesh/DynamicMeshAttributeSet.cpp
+++ b/Engine/Source/Runtime/GeometryCore/Private/DynamicMesh/DynamicMeshAttributeSet.cpp
@@ -226,9 +226,9 @@ void FDynamicMeshAttributeSet::CompactInPlace(const FCompactMaps& CompactMaps)
 		AttribPair.Value->CompactInPlace(CompactMaps);
 	}
 	
-	for (const TPair<FName, TUniquePtr<FDynamicMeshAttributeBase>>& AttribPair : GenericAttributes)
+	for (FDynamicMeshAttributeBase* RegAttrib : RegisteredAttributes)
 	{
-		AttribPair.Value->CompactInPlace(CompactMaps);
+		RegAttrib->CompactInPlace(CompactMaps);
 	}
 }
 
@@ -471,9 +471,9 @@ void FDynamicMeshAttributeSet::Reparent(FDynamicMesh3* NewParent)
 		BoneColorAttrib->Reparent(NewParent);
 	}
 	
-	for (const TPair<FName, TUniquePtr<FDynamicMeshAttributeBase>>& AttribPair : GenericAttributes)
+	for (FDynamicMeshAttributeBase* RegAttrib : RegisteredAttributes)
 	{
-		AttribPair.Value->Reparent(NewParent);
+		RegAttrib->Reparent(NewParent);
 	}
 }
 
diff --git a/Engine/Source/Runtime/GeometryCore/Private/DynamicMesh/DynamicMeshChangeTracker.cpp b/Engine/Source/Runtime/GeometryCore/Private/DynamicMesh/DynamicMeshChangeTracker.cpp
--- a/Engine/Source/Runtime/GeometryCore/Private/DynamicMesh/DynamicMeshChangeTracker.cpp
+++ b/Engine/Source/Runtime/GeometryCore/Private/DynamicMesh/DynamicMeshChangeTracker.cpp
@@ -636,7 +636,7 @@ void FDynamicMeshAttributeSetChangeTracker::StoreAllFinalTriangles(const TArray<
 }
 
 
-void FDynamicMeshAttributeSetChangeTracker::StoreAllFinalVertices(const TArray<int>& VertexIDs)
+void FDynamicMeshAttributeSetChangeTracker::StoreAllFinalVertices(const TSet<int>& VertexIDs)
 {
 	for (int k = 0, NumRegAttrib = Attribs->NumRegisteredAttributes(); k < NumRegAttrib; k++)
 	{
@@ -850,6 +850,7 @@ TUniquePtr<FDynamicMeshChange> FDynamicMeshChangeTracker::EndChange()
 	if (AttribChangeTracker != nullptr)
 	{
 		AttribChangeTracker->StoreAllFinalTriangles(StoredTriangles);
+		AttribChangeTracker->StoreAllFinalVertices(StoredVertices);
 		Result->AttachAttributeChanges(AttribChangeTracker->EndChange());
 	}
 
diff --git a/Engine/Source/Runtime/GeometryCore/Private/DynamicMesh/NonManifoldMappingSupport.cpp b/Engine/Source/Runtime/GeometryCore/Private/DynamicMesh/NonManifoldMappingSupport.cpp
--- a/Engine/Source/Runtime/GeometryCore/Private/DynamicMesh/NonManifoldMappingSupport.cpp
+++ b/Engine/Source/Runtime/GeometryCore/Private/DynamicMesh/NonManifoldMappingSupport.cpp
@@ -24,7 +24,7 @@ public:
 
 	void SaveInitialVertex(const FDynamicMeshAttributeBase* Attribute, int VertexID) override;
 
-	void StoreAllFinalVertices(const FDynamicMeshAttributeBase* Attribute, const TArray<int>& VertexIDs) override;
+	void StoreAllFinalVertices(const FDynamicMeshAttributeBase* Attribute, const TSet<int>& VertexIDs) override;
 
 	bool Apply(FDynamicMeshAttributeBase* Attribute, bool bRevert) const override;
 };
@@ -231,7 +231,7 @@ void FNonManifoldSourceVertexIdAttributeChange::SaveInitialVertex(const FDynamic
 	OldSourceVertexIds.Emplace(VertexID, NonManifoldAttribute->GetValue(VertexID));
 }
 
-void FNonManifoldSourceVertexIdAttributeChange::StoreAllFinalVertices(const FDynamicMeshAttributeBase* Attribute, const TArray<int>& VertexIDs)
+void FNonManifoldSourceVertexIdAttributeChange::StoreAllFinalVertices(const FDynamicMeshAttributeBase* Attribute, const TSet<int>& VertexIDs)
 {
 	NewSourceVertexIds.Reserve(NewSourceVertexIds.Num() + VertexIDs.Num());
 	const FNonManifoldSourceVertexIdAttribute* NonManifoldAttribute = static_cast<const FNonManifoldSourceVertexIdAttribute*>(Attribute);
diff --git a/Engine/Source/Runtime/GeometryCore/Public/DynamicMesh/DynamicAttribute.h b/Engine/Source/Runtime/GeometryCore/Public/DynamicMesh/DynamicAttribute.h
--- a/Engine/Source/Runtime/GeometryCore/Public/DynamicMesh/DynamicAttribute.h
+++ b/Engine/Source/Runtime/GeometryCore/Public/DynamicMesh/DynamicAttribute.h
@@ -39,7 +39,12 @@ public:
 	virtual void StoreAllFinalTriangles(const TDynamicAttributeBase<ParentType>* Attribute, const TArray<int>& TriangleIDs)
 	{
 	}
-	virtual void StoreAllFinalVertices(const TDynamicAttributeBase<ParentType>* Attribute, const TArray<int>& VertexIDs)
+	UE_DEPRECATED(5.6, "Use the version of this method that takes a TSet instead")
+	virtual void StoreAllFinalVertices(const TDynamicAttributeBase<ParentType>* Attribute, const TArray<int>& TriangleIDs)
+	{
+		ensure(false);
+	}
+	virtual void StoreAllFinalVertices(const TDynamicAttributeBase<ParentType>* Attribute, const TSet<int>& TriangleIDs)
 	{
 	}
 
diff --git a/Engine/Source/Runtime/GeometryCore/Public/DynamicMesh/DynamicMeshChangeTracker.h b/Engine/Source/Runtime/GeometryCore/Public/DynamicMesh/DynamicMeshChangeTracker.h
--- a/Engine/Source/Runtime/GeometryCore/Public/DynamicMesh/DynamicMeshChangeTracker.h
+++ b/Engine/Source/Runtime/GeometryCore/Public/DynamicMesh/DynamicMeshChangeTracker.h
@@ -234,7 +234,12 @@ public:
 	GEOMETRYCORE_API void SaveInitialVertex(int VertexID);
 
 	/** store the final state of a set of vertices */
-	GEOMETRYCORE_API void StoreAllFinalVertices(const TArray<int>& VertexIDs);
+	UE_DEPRECATED(5.6, "Use the version of this method that takes a TSet instead")
+	GEOMETRYCORE_API void StoreAllFinalVertices(const TArray<int>& VertexIDs)
+	{
+		ensure(false);
+	}
+	GEOMETRYCORE_API void StoreAllFinalVertices(const TSet<int>& VertexIDs);
 
 
 protected:
diff --git a/Engine/Source/Runtime/GeometryCore/Public/DynamicMesh/DynamicMeshTriangleAttribute.h b/Engine/Source/Runtime/GeometryCore/Public/DynamicMesh/DynamicMeshTriangleAttribute.h
--- a/Engine/Source/Runtime/GeometryCore/Public/DynamicMesh/DynamicMeshTriangleAttribute.h
+++ b/Engine/Source/Runtime/GeometryCore/Public/DynamicMesh/DynamicMeshTriangleAttribute.h
@@ -268,6 +268,16 @@ public:
 		}
 	}
 
+	/** Set the element at a given index with a scalar value (the same value for each dimension) */
+	inline void SetScalarValue(int TriangleID, const AttribValueType& SingleValue)
+	{
+		int k = TriangleID * AttribDimension;
+		for (int i = 0; i < AttribDimension; ++i)
+		{
+			AttribValues[k + i] = SingleValue;
+		}
+	}
+
 	/**
 	 * Copy the attribute value at FromTriangleID to ToTriangleID
 	 */
@@ -379,6 +389,7 @@ public:
 	virtual void OnNewTriangle(int TriangleID, bool bInserted) override
 	{
 		ResizeAttribStoreIfNeeded(TriangleID);
+		SetScalarValue(TriangleID, GetDefaultAttributeValue());
 	}
 
 public:
diff --git a/Engine/Source/Runtime/GeometryCore/Public/DynamicMesh/DynamicVertexAttribute.h b/Engine/Source/Runtime/GeometryCore/Public/DynamicMesh/DynamicVertexAttribute.h
--- a/Engine/Source/Runtime/GeometryCore/Public/DynamicMesh/DynamicVertexAttribute.h
+++ b/Engine/Source/Runtime/GeometryCore/Public/DynamicMesh/DynamicVertexAttribute.h
@@ -35,7 +35,7 @@ public:
 
 	inline virtual void SaveInitialVertex(const TDynamicAttributeBase<ParentType>* Attribute, int VertexID) override;
 
-	inline virtual void StoreAllFinalVertices(const TDynamicAttributeBase<ParentType>* Attribute, const TArray<int>& VertexIDs) override;
+	inline virtual void StoreAllFinalVertices(const TDynamicAttributeBase<ParentType>* Attribute, const TSet<int>& VertexIDs) override;
 
 	inline virtual bool Apply(TDynamicAttributeBase<ParentType>* Attribute, bool bRevert) const override;
 };
@@ -257,6 +257,16 @@ public:
 		}
 	}
 
+	/** Set the element at a given index with a scalar value (the same value for each dimension) */
+	inline void SetScalarValue(int32 VertexID, const AttribValueType& SingleValue)
+	{
+		int k = VertexID * AttribDimension;
+		for (int i = 0; i < AttribDimension; ++i)
+		{
+			AttribValues[k + i] = SingleValue;
+		}
+	}
+
 	/**
 	 * Copy the attribute value at FromVertexID to ToVertexID
 	 */
@@ -313,6 +323,7 @@ public:
 	virtual void OnNewVertex(int VertexID, bool bInserted) override
 	{
 		ResizeAttribStoreIfNeeded(VertexID);
+		SetScalarValue(VertexID, GetDefaultAttributeValue());
 	}
 
 	/** Update the overlay to reflect a face poke in the parent */
@@ -556,7 +567,7 @@ void TDynamicVertexAttributeChange<AttribValueType, AttribDimension, ParentType>
 }
 
 template<typename AttribValueType, int AttribDimension, typename ParentType>
-void TDynamicVertexAttributeChange<AttribValueType, AttribDimension, ParentType>::StoreAllFinalVertices(const TDynamicAttributeBase<ParentType>* Attribute, const TArray<int>& VertexIDs)
+void TDynamicVertexAttributeChange<AttribValueType, AttribDimension, ParentType>::StoreAllFinalVertices(const TDynamicAttributeBase<ParentType>* Attribute, const TSet<int>& VertexIDs)
 {
 	NewVertexAttributes.Reserve(NewVertexAttributes.Num() + VertexIDs.Num());
 	const TDynamicVertexAttribute<AttribValueType, AttribDimension, ParentType>* AttribCast = static_cast<const TDynamicVertexAttribute<AttribValueType, AttribDimension, ParentType>*>(Attribute);
diff --git a/Engine/Source/Runtime/GeometryCore/Public/DynamicMesh/DynamicVertexSkinWeightsAttribute.h b/Engine/Source/Runtime/GeometryCore/Public/DynamicMesh/DynamicVertexSkinWeightsAttribute.h
--- a/Engine/Source/Runtime/GeometryCore/Public/DynamicMesh/DynamicVertexSkinWeightsAttribute.h
+++ b/Engine/Source/Runtime/GeometryCore/Public/DynamicMesh/DynamicVertexSkinWeightsAttribute.h
@@ -33,7 +33,7 @@ public:
 
 	void SaveInitialVertex(const TDynamicAttributeBase<ParentType>* Attribute, int VertexID) override;
 
-	void StoreAllFinalVertices(const TDynamicAttributeBase<ParentType>* Attribute, const TArray<int>& VertexIDs) override;
+	void StoreAllFinalVertices(const TDynamicAttributeBase<ParentType>* Attribute, const TSet<int>& VertexIDs) override;
 
 	bool Apply(TDynamicAttributeBase<ParentType>* Attribute, bool bRevert) const override;
 };
@@ -705,7 +705,7 @@ void TDynamicVertexSkinWeightsAttributeChange<ParentType>::SaveInitialVertex(con
 }
 
 template<typename ParentType>
-void TDynamicVertexSkinWeightsAttributeChange<ParentType>::StoreAllFinalVertices(const TDynamicAttributeBase<ParentType>* Attribute, const TArray<int>& VertexIDs)
+void TDynamicVertexSkinWeightsAttributeChange<ParentType>::StoreAllFinalVertices(const TDynamicAttributeBase<ParentType>* Attribute, const TSet<int>& VertexIDs)
 {
 	NewVertexBoneWeights.Reserve(NewVertexBoneWeights.Num() + VertexIDs.Num());
 	const TDynamicVertexSkinWeightsAttribute<ParentType>* AttribCast = static_cast<const TDynamicVertexSkinWeightsAttribute<ParentType>*>(Attribute);
