GeometryCore: Fix a crash in FDynamicMeshAttributeSet::CompactInPlace due to compacting the same skin weight attributes twice. SkinWeightAttributes are also RegisteredAttributes so just compacting the RegisteredAttributes should take care of it

diff --git a/Engine/Source/Runtime/GeometryCore/Private/DynamicMesh/DynamicMeshAttributeSet.cpp b/Engine/Source/Runtime/GeometryCore/Private/DynamicMesh/DynamicMeshAttributeSet.cpp
--- a/Engine/Source/Runtime/GeometryCore/Private/DynamicMesh/DynamicMeshAttributeSet.cpp
+++ b/Engine/Source/Runtime/GeometryCore/Private/DynamicMesh/DynamicMeshAttributeSet.cpp
@@ -220,11 +220,6 @@ void FDynamicMeshAttributeSet::CompactInPlace(const FCompactMaps& CompactMaps)
 	{
 		WeightLayers[WeightLayerIdx].CompactInPlace(CompactMaps);
 	}
-
-	for (const TPair<FName, TUniquePtr<FDynamicMeshVertexSkinWeightsAttribute>>& AttribPair : SkinWeightAttributes)
-	{
-		AttribPair.Value->CompactInPlace(CompactMaps);
-	}
 	
 	for (FDynamicMeshAttributeBase* RegAttrib : RegisteredAttributes)
 	{
@@ -1474,6 +1469,11 @@ bool FDynamicMeshAttributeSet::CheckValidity(bool bAllowNonmanifold, EValidityCh
 		{
 			bValid = Pair.Value->CheckValidity(bAllowNonmanifold, FailMode) && bValid;
 		}
+
+		// Skin Weight Attributes are in the RegisteredAttributes array
+		int32 FoundIndex;
+		const bool bFound = RegisteredAttributes.Find(Pair.Value.Get(), FoundIndex);
+		bValid = bFound && bValid;
 	}
 
 	bValid = CheckBoneValidity(FailMode) && bValid;
