Add optional sculpt layer data to dynamic mesh attribute set

diff --git a/Engine/Source/Runtime/GeometryCore/Private/DynamicMesh/DynamicMeshAttributeSet.cpp b/Engine/Source/Runtime/GeometryCore/Private/DynamicMesh/DynamicMeshAttributeSet.cpp
--- a/Engine/Source/Runtime/GeometryCore/Private/DynamicMesh/DynamicMeshAttributeSet.cpp
+++ b/Engine/Source/Runtime/GeometryCore/Private/DynamicMesh/DynamicMeshAttributeSet.cpp
@@ -91,6 +91,8 @@ void FDynamicMeshAttributeSet::Copy(const FDynamicMeshAttributeSet& Copy)
 		AttachAttribute(AttribPair.Key, AttribPair.Value->MakeCopy(ParentMesh));
 	}
 
+	SculptLayers.Copy(this, Copy.SculptLayers);
+
 	// parent mesh is *not* copied!
 }
 
@@ -186,6 +188,8 @@ void FDynamicMeshAttributeSet::CompactCopy(const FCompactMaps& CompactMaps, cons
 		AttachAttribute(AttribPair.Key, AttribPair.Value->MakeCompactCopy(CompactMaps, ParentMesh));
 	}
 
+	SculptLayers.CompactCopy(this, CompactMaps, Copy.SculptLayers);
+
 	// parent mesh is *not* copied!
 }
 
@@ -405,6 +409,8 @@ void FDynamicMeshAttributeSet::EnableMatchingAttributes(const FDynamicMeshAttrib
 	}
 
 	EnableMatchingBoneAttributes(ToMatch, bClearExisting, bDiscardExtraAttributes);
+
+	SculptLayers.EnableMatching(this, ToMatch.SculptLayers, bClearExisting, bDiscardExtraAttributes);
 }
 
 
@@ -808,6 +814,15 @@ void FDynamicMeshAttributeSet::OnNewVertex(int VertexID, bool bInserted)
 	{
 		AttribPair.Value->OnNewVertex(VertexID, bInserted);
 	}
+
+	// When adding a new vertex (note: not via a edge split or triangle poke!) want to default layer 0 to use the new position
+	// even if we're operating on a different layer, so we don't end up with the base geometry collapsed at the origin
+	// (need to handle this here instead of in the layer itself since it depends on the layer index / active layer info)
+	if (NumSculptLayers() > 0)
+	{
+		checkSlow(ParentMesh);
+		SculptLayers.Layers[0].SetValue(VertexID, ParentMesh->GetVertex(VertexID));
+	}
 }
 
 
@@ -1436,9 +1451,28 @@ void FDynamicMeshAttributeSet::Serialize(FArchive& Ar, const FCompactMaps* Compa
 		}
 	}
 
+	// TODO: Serialize sculpt layers
+	
 	//Ar << GenericAttributes; // TODO
 }
 
+
+void FDynamicMeshAttributeSet::EnableSculptLayers(int32 MinLayers)
+{
+	if (!ensure(ParentMesh))
+	{
+		return;
+	}
+
+	SculptLayers.Enable(this, MinLayers);
+}
+
+void FDynamicMeshAttributeSet::DiscardSculptLayers()
+{
+	SculptLayers.Discard(this);
+}
+
+
 bool FDynamicMeshAttributeSet::CheckValidity(bool bAllowNonmanifold, EValidityCheckFailMode FailMode) const
 {
 	bool bValid = FDynamicMeshAttributeSetBase::CheckValidity(bAllowNonmanifold, FailMode);
@@ -1478,6 +1512,8 @@ bool FDynamicMeshAttributeSet::CheckValidity(bool bAllowNonmanifold, EValidityCh
 
 	bValid = CheckBoneValidity(FailMode) && bValid;
 
+	bValid = SculptLayers.CheckValidity(this, bAllowNonmanifold, FailMode) && bValid;
+
 	return bValid;
 }
 
diff --git a/Engine/Source/Runtime/GeometryCore/Private/DynamicMesh/DynamicMeshSculptLayers.cpp b/Engine/Source/Runtime/GeometryCore/Private/DynamicMesh/DynamicMeshSculptLayers.cpp
new file mode 100644
--- /dev/null
+++ b/Engine/Source/Runtime/GeometryCore/Private/DynamicMesh/DynamicMeshSculptLayers.cpp
@@ -0,0 +1,309 @@
+// Copyright Epic Games, Inc. All Rights Reserved.
+
+#include "DynamicMesh/DynamicMeshSculptLayers.h"
+#include "DynamicMesh/DynamicMeshAttributeSet.h"
+
+using namespace UE::Geometry;
+
+
+
+int32 FDynamicMeshSculptLayers::SetActiveLayer(int32 LayerIndex)
+{
+	UpdateLayersFromMesh();
+	ActiveLayer = LayerIndex;
+	ValidateActiveLayer();
+	return ActiveLayer;
+}
+
+bool FDynamicMeshSculptLayers::DiscardSculptLayer(int32 LayerIndex)
+{
+	if (!Layers.IsValidIndex(LayerIndex))
+	{
+		return false;
+	}
+	UpdateLayersFromMesh();
+	Layers.RemoveAt(LayerIndex);
+	LayerWeights.RemoveAt(LayerIndex);
+	UpdateMeshFromLayers();
+	ValidateActiveLayer();
+	return true;
+}
+
+bool FDynamicMeshSculptLayers::MergeSculptLayers(int32 StartIndex, int32 EndIndex, bool bUseWeights)
+{
+	if (!Layers.IsValidIndex(StartIndex) || !Layers.IsValidIndex(EndIndex) || StartIndex >= EndIndex)
+	{
+		return false;
+	}
+	double ActiveWeight = bUseWeights ? LayerWeights[StartIndex] : 1.0;
+	if (ActiveWeight == 0)
+	{
+		// cannot bake into a zero-weight layer
+		return false;
+	}
+	double ActiveWeightInv = 1.0 / ActiveWeight;
+	UpdateLayersFromMesh();
+
+	FDynamicMesh3* Mesh = Layers[StartIndex].GetParent();
+	for (int32 VID : Mesh->VertexIndicesItr())
+	{
+		FVector3d MergedSum(0, 0, 0);
+		for (int32 LayerIdx = StartIndex + 1; LayerIdx <= EndIndex; ++LayerIdx)
+		{
+			FVector3d LayerOffset;
+			Layers[LayerIdx].GetValue(VID, LayerOffset);
+			double Wt = bUseWeights ? LayerWeights[LayerIdx] : 1.0;
+			MergedSum += LayerOffset * Wt;
+		}
+		FVector3d StartInitial;
+		Layers[StartIndex].GetValue(VID, StartInitial);
+		Layers[StartIndex].SetValue(VID, StartInitial + MergedSum * ActiveWeightInv);
+	}
+	Layers.RemoveAt(StartIndex + 1, EndIndex - StartIndex);
+	LayerWeights.RemoveAt(StartIndex + 1, EndIndex - StartIndex);
+	if (!bUseWeights)
+	{
+		UpdateMeshFromLayers();
+	}
+	ValidateActiveLayer();
+	return true;
+}
+
+void FDynamicMeshSculptLayers::UpdateLayerWeights(TConstArrayView<double> InLayerWeights)
+{
+	UpdateLayersFromMesh();
+	for (int32 LayerIdx = 0, NumLayers = FMath::Min(InLayerWeights.Num(), LayerWeights.Num()); LayerIdx < NumLayers; ++LayerIdx)
+	{
+		LayerWeights[LayerIdx] = InLayerWeights[LayerIdx];
+	}
+	UpdateMeshFromLayers();
+	ValidateActiveLayer();
+}
+
+void FDynamicMeshSculptLayers::Enable(FDynamicMeshAttributeSet* AttributeSet, int32 MinLayerCount)
+{
+	check(AttributeSet);
+
+	int32 UseMinLayers = FMath::Max(MinLayerCount, 1);
+	bool bWasEmpty = Layers.IsEmpty();
+
+	if (UseMinLayers > Layers.Num())
+	{
+		FDynamicMesh3* ParentMesh = AttributeSet->GetParentMesh();
+		Layers.Reserve(UseMinLayers);
+		LayerWeights.Reserve(UseMinLayers);
+		while (Layers.Num() < UseMinLayers)
+		{
+			FDynamicMeshSculptLayerAttribute* LayerToAdd = new FDynamicMeshSculptLayerAttribute(ParentMesh);
+			Layers.Add(LayerToAdd);
+			AttributeSet->RegisterExternalAttribute(LayerToAdd);
+			LayerWeights.Add(1.0);
+		}
+
+		// Initialize base layer from parent mesh (if it was just created)
+		if (bWasEmpty)
+		{
+			ActiveLayer = 0;
+			for (int32 VID : ParentMesh->VertexIndicesItr())
+			{
+				Layers[0].SetValue(VID, ParentMesh->GetVertex(VID));
+			}
+		}
+	}
+}
+
+void FDynamicMeshSculptLayers::Discard(FDynamicMeshAttributeSet* AttributeSet)
+{
+	check(AttributeSet);
+	for (FDynamicMeshSculptLayerAttribute& Layer : Layers)
+	{
+		AttributeSet->UnregisterExternalAttribute(&Layer);
+	}
+	Layers.Reset();
+	LayerWeights.Reset();
+	ActiveLayer = -1;
+}
+
+void FDynamicMeshSculptLayers::Copy(FDynamicMeshAttributeSet* AttributeSet, const FDynamicMeshSculptLayers& Copy)
+{
+	Discard(AttributeSet);
+	for (const FDynamicMeshSculptLayerAttribute& Layer : Copy.Layers)
+	{
+		FDynamicMeshSculptLayerAttribute* NewLayer = static_cast<FDynamicMeshSculptLayerAttribute*>(Layer.MakeCopy(AttributeSet->GetParentMesh()));
+		Layers.Add(NewLayer);
+		AttributeSet->RegisterExternalAttribute(NewLayer);
+	}
+	LayerWeights = Copy.LayerWeights;
+	ActiveLayer = Copy.ActiveLayer;
+}
+
+void FDynamicMeshSculptLayers::CompactCopy(FDynamicMeshAttributeSet* AttributeSet, const FCompactMaps& CompactMaps, const FDynamicMeshSculptLayers& Copy)
+{
+	Discard(AttributeSet);
+	for (const FDynamicMeshSculptLayerAttribute& Layer : Copy.Layers)
+	{
+		FDynamicMeshSculptLayerAttribute* NewLayer = static_cast<FDynamicMeshSculptLayerAttribute*>(Layer.MakeCompactCopy(CompactMaps, AttributeSet->GetParentMesh()));
+		Layers.Add(NewLayer);
+		AttributeSet->RegisterExternalAttribute(NewLayer);
+	}
+	LayerWeights = Copy.LayerWeights;
+	ActiveLayer = Copy.ActiveLayer;
+}
+
+void FDynamicMeshSculptLayers::EnableMatching(FDynamicMeshAttributeSet* AttributeSet, const FDynamicMeshSculptLayers& ToMatch, bool bClearExisting, bool bDiscardExtraAttributes)
+{
+	if (bClearExisting || (bDiscardExtraAttributes && !ToMatch.IsEnabled()))
+	{
+		Discard(AttributeSet);
+	}
+	if (ToMatch.IsEnabled())
+	{
+		FDynamicMesh3* ParentMesh = AttributeSet->GetParentMesh();
+		Layers.Reserve(ToMatch.NumLayers());
+		LayerWeights.Reserve(ToMatch.NumLayers());
+		while (Layers.Num() < ToMatch.NumLayers())
+		{
+			FDynamicMeshSculptLayerAttribute* LayerToAdd = new FDynamicMeshSculptLayerAttribute(ParentMesh);
+			int32 LayerIdx = Layers.Add(LayerToAdd);
+			AttributeSet->RegisterExternalAttribute(LayerToAdd);
+			// Initialize weight to also match the other layer, for new layers
+			LayerWeights.Add(ToMatch.LayerWeights[LayerIdx]);
+		}
+		if (bDiscardExtraAttributes && NumLayers() > ToMatch.NumLayers())
+		{
+			int32 RemoveStart = ToMatch.NumLayers(), RemoveCount = NumLayers() - ToMatch.NumLayers();
+			Layers.RemoveAt(RemoveStart, RemoveCount);
+			LayerWeights.RemoveAt(RemoveStart, RemoveCount);
+		}
+	}
+	ValidateActiveLayer();
+}
+
+void FDynamicMeshSculptLayers::CopyThroughMapping(const FDynamicMeshSculptLayers& Other, const FMeshIndexMappings& Mapping)
+{
+	const int32 NumLayers = FMath::Min(Layers.Num(), Other.Layers.Num());
+	for (int32 LayerIdx = 0; LayerIdx < NumLayers; ++LayerIdx)
+	{
+		Layers[LayerIdx].CopyThroughMapping(&Other.Layers[LayerIdx], Mapping);
+	}
+}
+
+bool FDynamicMeshSculptLayers::CheckValidity(const FDynamicMeshAttributeSet* AttributeSet, bool bAllowNonmanifold, EValidityCheckFailMode FailMode) const
+{
+	bool bValid = (LayerWeights.Num() == Layers.Num());
+	for (const FDynamicMeshSculptLayerAttribute& Layer : Layers)
+	{
+		// validate layer is registered
+		bValid = AttributeSet->RegisteredAttributes.Contains(&Layer) && bValid;
+		// Note: We don't need to call Layer.CheckValidity here, since it is already called on registered attributes
+
+		// validate parent mesh is set consistent with attribute set
+		bValid = (AttributeSet->GetParentMesh() == Layer.GetParent()) && bValid;
+	}
+
+	switch (FailMode)
+	{
+	case EValidityCheckFailMode::Check:
+		check(false);
+	case EValidityCheckFailMode::Ensure:
+		ensure(false);
+	default:
+		return false;
+	}
+}
+
+bool FDynamicMeshSculptLayers::UpdateLayersFromMesh()
+{
+	if (Layers.IsEmpty() || !ensure(HasValidLayers()))
+	{
+		return false;
+	}
+
+	double ActiveWeight = LayerWeights[ActiveLayer];
+	if (ActiveWeight == 0)
+	{
+		return false;
+	}
+
+	FDynamicMesh3* Mesh = Layers[0].GetParent();
+	check(Mesh);
+	for (int32 VID : Mesh->VertexIndicesItr())
+	{
+		FVector3d Sum(0, 0, 0);
+		for (int32 LayerIdx = 0; LayerIdx < Layers.Num(); ++LayerIdx)
+		{
+			FVector3d Offset;
+			Layers[LayerIdx].GetValue(VID, Offset);
+			Sum += Offset * LayerWeights[LayerIdx];
+		}
+		FVector3d Pos = Mesh->GetVertex(VID);
+
+		FVector3d Delta = (Pos - Sum) / ActiveWeight;
+		FVector3d OrigActiveOffset;
+		Layers[ActiveLayer].GetValue(VID, OrigActiveOffset);
+		Layers[ActiveLayer].SetValue(VID, OrigActiveOffset + Delta);
+	}
+
+	return true;
+}
+
+bool FDynamicMeshSculptLayers::UpdateMeshFromLayers()
+{
+	if (Layers.IsEmpty() || !ensure(HasValidLayers()))
+	{
+		return false;
+	}
+
+	FDynamicMesh3* Mesh = Layers[0].GetParent();
+	check(Mesh);
+	for (int32 VID : Mesh->VertexIndicesItr())
+	{
+		FVector3d Sum(0, 0, 0);
+		for (int32 LayerIdx = 0; LayerIdx < Layers.Num(); ++LayerIdx)
+		{
+			FVector3d Offset;
+			Layers[LayerIdx].GetValue(VID, Offset);
+			Sum += Offset * LayerWeights[LayerIdx];
+		}
+		Mesh->SetVertex(VID, Sum);
+	}
+
+	return true;
+}
+
+
+bool FDynamicMeshSculptLayers::ValidateActiveLayer()
+{
+	if (!IsEnabled() || !ensure(Layers.Num() == LayerWeights.Num()))
+	{
+		ActiveLayer = INDEX_NONE;
+		return false;
+	}
+	ActiveLayer = FMath::Clamp(ActiveLayer, 0, LayerWeights.Num());
+
+	// If active layer has zero weight, try to find an active layer with non-zero weight instead
+	if (LayerWeights[ActiveLayer] == 0)
+	{
+		for (int32 Idx = ActiveLayer + 1; Idx < LayerWeights.Num(); ++Idx)
+		{
+			if (LayerWeights[Idx] != 0)
+			{
+				ActiveLayer = Idx;
+				return true;
+			}
+		}
+		for (int32 Idx = ActiveLayer - 1; Idx >= 0; --Idx)
+		{
+			if (LayerWeights[ActiveLayer] != 0)
+			{
+				ActiveLayer = Idx;
+				return true;
+			}
+		}
+		return false;
+	}
+	else
+	{
+		return true;
+	}
+}
diff --git a/Engine/Source/Runtime/GeometryCore/Private/DynamicMeshEditor.cpp b/Engine/Source/Runtime/GeometryCore/Private/DynamicMeshEditor.cpp
--- a/Engine/Source/Runtime/GeometryCore/Private/DynamicMeshEditor.cpp
+++ b/Engine/Source/Runtime/GeometryCore/Private/DynamicMeshEditor.cpp
@@ -2111,6 +2111,11 @@ void FDynamicMeshEditor::AppendMesh(const FDynamicMesh3* AppendMesh,
 				ToAttrib->CopyThroughMapping(AttribPair.Value.Get(), IndexMapsOut);
 			}
 		}
+
+		if (Mesh->Attributes()->NumSculptLayers() > 0 && AppendMesh->Attributes()->NumSculptLayers() > 0)
+		{
+			Mesh->Attributes()->GetSculptLayers()->CopyThroughMapping(*AppendMesh->Attributes()->GetSculptLayers(), IndexMapsOut);
+		}
 	}
 
 	// Flip tri orientations if requested, after all appends -- note we do this after mapping across all the attributes, as that logic assumes triangle vertices have matching vertex order
@@ -2463,6 +2468,11 @@ static void AppendVertexAttributes(const FDynamicMesh3* FromMesh, FDynamicMesh3*
 			ToAttrib->CopyThroughMapping(AttribPair.Value.Get(), IndexMaps);
 		}
 	}
+
+	if (FromMesh->Attributes()->NumSculptLayers() > 0 && ToMesh->Attributes()->NumSculptLayers() > 0)
+	{
+		ToMesh->Attributes()->GetSculptLayers()->CopyThroughMapping(*FromMesh->Attributes()->GetSculptLayers(), IndexMaps);
+	}
 }
 
 
diff --git a/Engine/Source/Runtime/GeometryCore/Public/DynamicMesh/DynamicMeshAttributeSet.h b/Engine/Source/Runtime/GeometryCore/Public/DynamicMesh/DynamicMeshAttributeSet.h
--- a/Engine/Source/Runtime/GeometryCore/Public/DynamicMesh/DynamicMeshAttributeSet.h
+++ b/Engine/Source/Runtime/GeometryCore/Public/DynamicMesh/DynamicMeshAttributeSet.h
@@ -11,6 +11,7 @@
 #include "DynamicMesh/DynamicMeshOverlay.h"
 #include "DynamicMesh/DynamicMeshTriangleAttribute.h"
 #include "DynamicMesh/DynamicVertexAttribute.h"
+#include "DynamicMesh/DynamicMeshSculptLayers.h"
 #include "GeometryTypes.h"
 #include "HAL/PlatformCrt.h"
 #include "InfoTypes.h"
@@ -487,6 +488,30 @@ public:
 		return BoneColorAttrib.Get();
 	}
 
+	const FDynamicMeshSculptLayers* GetSculptLayers() const
+	{
+		return &SculptLayers;
+	}
+	
+	FDynamicMeshSculptLayers* GetSculptLayers()
+	{
+		return &SculptLayers;
+	}
+
+	// Number of sculpt layers enabled on the mesh
+	inline int32 NumSculptLayers() const
+	{
+		return SculptLayers.Layers.Num();
+	}
+
+	// Enable sculpt layers on the mesh, if they are not already enabled.
+	// @param MinLayers Number of layers to enabled. If more layers already exist, no layers will be removed.
+	GEOMETRYCORE_API void EnableSculptLayers(int32 MinLayers);
+
+	// Discard sculpt layer data from the mesh.
+	GEOMETRYCORE_API void DiscardSculptLayers();
+
+
 	// Attach a new attribute (and transfer ownership of it to the attribute set)
 	void AttachAttribute(FName AttribName, FDynamicMeshAttributeBase* Attribute)
 	{
@@ -583,6 +608,13 @@ protected:
 
 	using GenericAttributesMap = TMap<FName, TUniquePtr<FDynamicMeshAttributeBase>>;
 	GenericAttributesMap GenericAttributes;
+
+private:
+
+	// Class to manage sculpt layer data, separated out mainly for clearer organization
+	FDynamicMeshSculptLayers SculptLayers;
+	friend class FDynamicMeshSculptLayers;
+	
 	
 protected:
 	friend class FDynamicMesh3;
diff --git a/Engine/Source/Runtime/GeometryCore/Public/DynamicMesh/DynamicMeshSculptLayers.h b/Engine/Source/Runtime/GeometryCore/Public/DynamicMesh/DynamicMeshSculptLayers.h
new file mode 100644
--- /dev/null
+++ b/Engine/Source/Runtime/GeometryCore/Public/DynamicMesh/DynamicMeshSculptLayers.h
@@ -0,0 +1,104 @@
+// Copyright Epic Games, Inc. All Rights Reserved.
+
+#pragma once
+
+#include "Containers/Array.h"
+#include "Containers/ArrayView.h"
+#include "Containers/IndirectArray.h"
+#include "DynamicMesh/DynamicAttribute.h"
+#include "DynamicMesh/DynamicVertexAttribute.h"
+
+
+namespace UE::Geometry
+{
+
+
+/** Per-vertex position offsets */
+typedef TDynamicMeshVertexAttribute<double, 3> FDynamicMeshSculptLayerAttribute;
+
+
+// Manages a dynamic mesh attribute set's sculpt layer data
+class FDynamicMeshSculptLayers final
+{
+public:
+
+	inline int32 NumLayers() const
+	{
+		return Layers.Num();
+	}
+
+	inline bool IsEnabled() const
+	{
+		return !Layers.IsEmpty();
+	}
+
+	inline int32 GetActiveLayer() const
+	{
+		return ActiveLayer;
+	}
+
+	// Attempt to set the active sculpt layer to the requested index. Will clamp to a valid layer range. Note: If the requested layer has zero weight, will attempt to use a layer with non-zero weight instead.
+	// @return The actual current active layer (may be different from the requested layer!)
+	GEOMETRYCORE_API int32 SetActiveLayer(int32 LayerIndex);
+
+	// Remove the given sculpt layer, discarding its contribution to the shape.
+	// Note: Shifts all sculpt layers above this index downward, which may invalidate externally-held sculpt layer indices. May also change the active layer.
+	GEOMETRYCORE_API bool DiscardSculptLayer(int32 LayerIndex);
+	
+	// Merge the contribution of a range of layers.
+	// Note: Shifts all sculpt layers above this range downward, which may invalidate externally-held sculpt layer indices. May also change the active layer.
+	// @param StartIndex Index of the first layer to merge
+	// @param EndIndex Index of the last layer to merge
+	// @param bUseWeights Whether to merge layers based on their current weight strength. Requires StartIndex have non-zero weight. If false, the mesh vertices may move after merge.
+	GEOMETRYCORE_API bool MergeSculptLayers(int32 StartIndex, int32 EndIndex, bool bUseWeights);
+
+	// Set new sculpt layer weights
+	GEOMETRYCORE_API void UpdateLayerWeights(TConstArrayView<double> InLayerWeights);
+
+	// Get the current sculpt layer weights
+	TConstArrayView<double> GetLayerWeights() const
+	{
+		return TConstArrayView<double>(LayerWeights);
+	}
+
+	// Rebuild mesh from sculpt layer offsets + weights, ignoring current mesh positions
+	void RebuildMesh()
+	{
+		UpdateMeshFromLayers();
+	}
+
+	// Copy across sculpt layer data via a vertex mapping, for all layers that exist on both this and the other layers
+	GEOMETRYCORE_API void CopyThroughMapping(const FDynamicMeshSculptLayers& Other, const FMeshIndexMappings& Mapping);
+
+private:
+
+	// Sculpt layers are stored as vertex position offsets from the previous layer, with layer zero storing initial positions
+	TIndirectArray<FDynamicMeshSculptLayerAttribute> Layers;
+	// Weights per sculpt layer
+	TArray<double> LayerWeights;
+	// Indicates which layer is currently being edited (and is therefore reflected in the mesh vertices, rather than the layer data)
+	int32 ActiveLayer = -1;
+
+	friend class FDynamicMeshAttributeSet;
+
+
+	// Internal helpers used by FDynamicMeshAttributeSet to manage layers
+
+	GEOMETRYCORE_API void Copy(FDynamicMeshAttributeSet* AttributeSet, const FDynamicMeshSculptLayers& Copy);
+	GEOMETRYCORE_API void CompactCopy(FDynamicMeshAttributeSet* AttributeSet, const FCompactMaps& CompactMaps, const FDynamicMeshSculptLayers& Copy);
+	GEOMETRYCORE_API void EnableMatching(FDynamicMeshAttributeSet* AttributeSet, const FDynamicMeshSculptLayers& ToMatch, bool bClearExisting, bool bDiscardExtraAttributes);
+	GEOMETRYCORE_API void Enable(FDynamicMeshAttributeSet* AttributeSet, int32 MinLayerCount = -1);
+	GEOMETRYCORE_API void Discard(FDynamicMeshAttributeSet* AttributeSet);
+	GEOMETRYCORE_API bool CheckValidity(const FDynamicMeshAttributeSet* AttributeSet, bool bAllowNonmanifold, EValidityCheckFailMode FailMode) const;
+	GEOMETRYCORE_API bool UpdateLayersFromMesh();
+	GEOMETRYCORE_API bool UpdateMeshFromLayers();
+	GEOMETRYCORE_API bool ValidateActiveLayer();
+
+	inline bool HasValidLayers() const
+	{
+		return Layers.Num() == LayerWeights.Num() && Layers.IsValidIndex(ActiveLayer);
+	}
+};
+
+
+} // namespace UE::Geometry
