geometry processing unit test CIS fix for unregistered weights after dynamic mesh serialization

diff --git a/Engine/Source/Runtime/GeometryCore/Private/DynamicMesh/DynamicMeshAttributeSet.cpp b/Engine/Source/Runtime/GeometryCore/Private/DynamicMesh/DynamicMeshAttributeSet.cpp
--- a/Engine/Source/Runtime/GeometryCore/Private/DynamicMesh/DynamicMeshAttributeSet.cpp
+++ b/Engine/Source/Runtime/GeometryCore/Private/DynamicMesh/DynamicMeshAttributeSet.cpp
@@ -1412,6 +1412,7 @@ void FDynamicMeshAttributeSet::Serialize(FArchive& Ar, const FCompactMaps* Compa
 				{
 					Value.Reset(new FDynamicMeshVertexSkinWeightsAttribute(ParentMesh, false));
 					Value.Get()->Serialize(Ar, CompactMaps, bUseCompression);
+					RegisterExternalAttribute(Value.Get());
 				}
 			}
 		}
@@ -1502,12 +1503,12 @@ bool FDynamicMeshAttributeSet::CheckValidity(bool bAllowNonmanifold, EValidityCh
 		if (Pair.Value)
 		{
 			bValid = Pair.Value->CheckValidity(bAllowNonmanifold, FailMode) && bValid;
-		}
 
-		// Skin Weight Attributes are in the RegisteredAttributes array
-		int32 FoundIndex;
-		const bool bFound = RegisteredAttributes.Find(Pair.Value.Get(), FoundIndex);
-		bValid = bFound && bValid;
+			// If non-null, Skin Weight Attributes are in the RegisteredAttributes array
+			int32 FoundIndex;
+			const bool bFound = RegisteredAttributes.Find(Pair.Value.Get(), FoundIndex);
+			bValid = bFound && bValid;
+		}
 	}
 
 	bValid = CheckBoneValidity(FailMode) && bValid;
