CIS fix for sculpt layer validity checks

diff --git a/Engine/Source/Runtime/GeometryCore/Private/DynamicMesh/DynamicMeshSculptLayers.cpp b/Engine/Source/Runtime/GeometryCore/Private/DynamicMesh/DynamicMeshSculptLayers.cpp
--- a/Engine/Source/Runtime/GeometryCore/Private/DynamicMesh/DynamicMeshSculptLayers.cpp
+++ b/Engine/Source/Runtime/GeometryCore/Private/DynamicMesh/DynamicMeshSculptLayers.cpp
@@ -201,15 +201,20 @@ bool FDynamicMeshSculptLayers::CheckValidity(const FDynamicMeshAttributeSet* Att
 		bValid = (AttributeSet->GetParentMesh() == Layer.GetParent()) && bValid;
 	}
 
-	switch (FailMode)
+	if (!bValid)
 	{
-	case EValidityCheckFailMode::Check:
-		check(false);
-	case EValidityCheckFailMode::Ensure:
-		ensure(false);
-	default:
-		return false;
+		switch (FailMode)
+		{
+		case EValidityCheckFailMode::Check:
+			check(false);
+		case EValidityCheckFailMode::Ensure:
+			ensure(false);
+		default:
+			return false;
+		}
 	}
+
+	return true;
 }
 
 bool FDynamicMeshSculptLayers::UpdateLayersFromMesh()
