Make sure new intermediate IDs are set to 'invalid' instead of 0 when expanding the point to cell ID mapping in dynamic sparse point octree

+ add test coverage for out of order insertion

diff --git a/Engine/Source/Runtime/GeometryCore/Public/Spatial/SparseDynamicPointOctree3.h b/Engine/Source/Runtime/GeometryCore/Public/Spatial/SparseDynamicPointOctree3.h
--- a/Engine/Source/Runtime/GeometryCore/Public/Spatial/SparseDynamicPointOctree3.h
+++ b/Engine/Source/Runtime/GeometryCore/Public/Spatial/SparseDynamicPointOctree3.h
@@ -510,7 +510,7 @@ void FSparseDynamicPointOctree3::Insert_NewChildCell(int32 PointID, const FVecto
 	NewChildCell.CellID = CellRefCounts.Allocate();
 	Cells.InsertAt(NewChildCell, NewChildCell.CellID);
 
-	PointIDToCellMap.InsertAt(NewChildCell.CellID, PointID);
+	PointIDToCellMap.InsertAt(NewChildCell.CellID, PointID, InvalidCellID);
 
 	CellPointLists.AllocateAt(NewChildCell.CellID);
 	CellPointLists.Insert(NewChildCell.CellID, PointID);
@@ -527,7 +527,7 @@ void FSparseDynamicPointOctree3::Insert_ToCell(int32 PointID, const FVector3d& P
 {
 	checkSlow(CellRefCounts.IsValid(ExistingCell.CellID));
 
-	PointIDToCellMap.InsertAt(ExistingCell.CellID, PointID);
+	PointIDToCellMap.InsertAt(ExistingCell.CellID, PointID, InvalidCellID);
 
 	CellPointLists.Insert(ExistingCell.CellID, PointID);
 
@@ -541,7 +541,7 @@ void FSparseDynamicPointOctree3::Insert_NewRoot(int32 PointID, const FVector3d&
 {
 	uint32 NewRootCellID = CreateNewRootCell(NewRootCell, true);
 
-	PointIDToCellMap.InsertAt(NewRootCellID, PointID);
+	PointIDToCellMap.InsertAt(NewRootCellID, PointID, InvalidCellID);
 	CellPointLists.Insert(NewRootCellID, PointID);
 }
 
