dynamic mesh fix serialization of meshes where attributes where disabled when written, but enabled on the mesh before load

diff --git a/Engine/Source/Runtime/GeometryCore/Private/DynamicMesh/DynamicMesh3_Serialization.cpp b/Engine/Source/Runtime/GeometryCore/Private/DynamicMesh/DynamicMesh3_Serialization.cpp
--- a/Engine/Source/Runtime/GeometryCore/Private/DynamicMesh/DynamicMesh3_Serialization.cpp
+++ b/Engine/Source/Runtime/GeometryCore/Private/DynamicMesh/DynamicMesh3_Serialization.cpp
@@ -201,13 +201,12 @@ void SerializeAttributeSet(FArchive& Ar, FDynamicMesh3* Mesh, const FCompactMaps
 {
 	bool bHasAttributes = Mesh->Attributes() != nullptr;
 	Ar << bHasAttributes;
+	if (Ar.IsLoading())
+	{
+		bHasAttributes ? Mesh->EnableAttributes() : Mesh->DiscardAttributes();
+	}
 	if (bHasAttributes)
 	{
-		if (Ar.IsLoading())
-		{
-			Mesh->EnableAttributes();
-		}
-
 		Mesh->Attributes()->Serialize(Ar, CompactMaps, Options.bUseCompression);
 	}
 }
diff --git a/Engine/Source/Runtime/GeometryCore/Private/DynamicMesh/DynamicMeshAttributeSet.cpp b/Engine/Source/Runtime/GeometryCore/Private/DynamicMesh/DynamicMeshAttributeSet.cpp
--- a/Engine/Source/Runtime/GeometryCore/Private/DynamicMesh/DynamicMeshAttributeSet.cpp
+++ b/Engine/Source/Runtime/GeometryCore/Private/DynamicMesh/DynamicMeshAttributeSet.cpp
@@ -1370,23 +1370,23 @@ void FDynamicMeshAttributeSet::Serialize(FArchive& Ar, const FCompactMaps* Compa
 	// Use int32 here to future-proof for multiple color layers.
 	int32 NumColorLayers = HasPrimaryColors() ? 1 : 0;
 	Ar << NumColorLayers;
+	if (Ar.IsLoading())
+	{
+		NumColorLayers > 0 ? EnablePrimaryColors() : DisablePrimaryColors();
+	}
 	if (NumColorLayers > 0)
 	{
-		if (Ar.IsLoading())
-		{
-			EnablePrimaryColors();
-		}
 		ColorLayer->Serialize(Ar, CompactMaps, bUseCompression);
 	}
 
 	bool bHasMaterialID = HasMaterialID();
 	Ar << bHasMaterialID;
+	if (Ar.IsLoading())
+	{
+		bHasMaterialID ? EnableMaterialID() : DisableMaterialID();
+	}
 	if (bHasMaterialID)
 	{
-		if (Ar.IsLoading())
-		{
-			EnableMaterialID();
-		}
 		MaterialIDAttrib->Serialize(Ar, CompactMaps, bUseCompression);
 	}
 
@@ -1439,12 +1439,13 @@ void FDynamicMeshAttributeSet::Serialize(FArchive& Ar, const FCompactMaps* Compa
 	{
 		bool bHasBones = BoneNameAttrib && BoneParentIndexAttrib && BonePoseAttrib && BoneColorAttrib;
 		Ar << bHasBones;
+
+		if (Ar.IsLoading())
+		{
+			bHasBones ? EnableBones(0) : DisableBones();
+		}
 		if (bHasBones)
 		{
-			if (Ar.IsLoading())
-			{
-				EnableBones(0);
-			}
 			FNameAsStringProxyArchive ProxyArchive(Ar); BoneNameAttrib->Serialize(ProxyArchive);
 			BoneParentIndexAttrib->Serialize(Ar);
 			BonePoseAttrib->Serialize(Ar);
