Add methods for direct access to sculpt layer attribute data

diff --git a/Engine/Source/Runtime/GeometryCore/Public/DynamicMesh/DynamicMeshSculptLayers.h b/Engine/Source/Runtime/GeometryCore/Public/DynamicMesh/DynamicMeshSculptLayers.h
--- a/Engine/Source/Runtime/GeometryCore/Public/DynamicMesh/DynamicMeshSculptLayers.h
+++ b/Engine/Source/Runtime/GeometryCore/Public/DynamicMesh/DynamicMeshSculptLayers.h
@@ -55,6 +55,19 @@ public:
 	// Set new sculpt layer weights
 	GEOMETRYCORE_API void UpdateLayerWeights(TConstArrayView<double> InLayerWeights);
 
+	// Directly access the sculpt layer data
+	const FDynamicMeshSculptLayerAttribute* GetLayer(int32 LayerIndex) const
+	{
+		return Layers.IsValidIndex(LayerIndex) ? &Layers[LayerIndex] : nullptr;
+	}
+	
+	// Directly access the sculpt layer data
+	// Note: Must explicitly call "RebuildMesh" for modifications to be applied to the mesh vertex positions
+	FDynamicMeshSculptLayerAttribute* GetLayer(int32 LayerIndex)
+	{
+		return Layers.IsValidIndex(LayerIndex) ? &Layers[LayerIndex] : nullptr;
+	}
+
 	// Get the current sculpt layer weights
 	TConstArrayView<double> GetLayerWeights() const
 	{
