make dynamic mesh serialize sculpt layers

diff --git a/Engine/Source/Runtime/Core/Public/UObject/FortniteMainBranchObjectVersion.h b/Engine/Source/Runtime/Core/Public/UObject/FortniteMainBranchObjectVersion.h
--- a/Engine/Source/Runtime/Core/Public/UObject/FortniteMainBranchObjectVersion.h
+++ b/Engine/Source/Runtime/Core/Public/UObject/FortniteMainBranchObjectVersion.h
@@ -556,6 +556,9 @@ struct FFortniteMainBranchObjectVersion
 		// Added PreprocessedFontGeometry to FFontFaceData
 		AddedPreprocessedFontGeometry,
 
+		// Added Dynamic Mesh Sculpt Layer serialization
+		DynamicMeshSerializeSculptLayers,
+
 		// -----<new versions can be added above this line>-------------------------------------------------
 		VersionPlusOne,
 		LatestVersion = VersionPlusOne - 1
diff --git a/Engine/Source/Runtime/GeometryCore/Private/DynamicMesh/DynamicMeshAttributeSet.cpp b/Engine/Source/Runtime/GeometryCore/Private/DynamicMesh/DynamicMeshAttributeSet.cpp
--- a/Engine/Source/Runtime/GeometryCore/Private/DynamicMesh/DynamicMeshAttributeSet.cpp
+++ b/Engine/Source/Runtime/GeometryCore/Private/DynamicMesh/DynamicMeshAttributeSet.cpp
@@ -9,6 +9,7 @@
 #include "Serialization/NameAsStringProxyArchive.h"
 
 #include "UObject/FortniteReleaseBranchCustomObjectVersion.h"
+#include "UObject/FortniteMainBranchObjectVersion.h"
 
 using namespace UE::Geometry;
 
@@ -1453,7 +1454,30 @@ void FDynamicMeshAttributeSet::Serialize(FArchive& Ar, const FCompactMaps* Compa
 		}
 	}
 
-	// TODO: Serialize sculpt layers
+	const bool bSerializeSculptLayers = !Ar.IsLoading() || Ar.CustomVer(FFortniteMainBranchObjectVersion::GUID) >= FFortniteMainBranchObjectVersion::DynamicMeshSerializeSculptLayers;
+	if (bSerializeSculptLayers)
+	{
+		int32 NumLayers = NumSculptLayers();
+		Ar << NumLayers;
+
+		if (Ar.IsLoading())
+		{
+			DiscardSculptLayers();
+			if (NumLayers > 0)
+			{
+				EnableSculptLayers(NumLayers);
+			}
+		}
+		if (NumLayers > 0)
+		{
+			Ar << SculptLayers.ActiveLayer;
+			Ar << SculptLayers.LayerWeights;
+			for (int32 LayerIdx = 0; LayerIdx < NumLayers; ++LayerIdx)
+			{
+				SculptLayers.Layers[LayerIdx].Serialize(Ar, CompactMaps, bUseCompression);
+			}
+		}
+	}
 	
 	//Ar << GenericAttributes; // TODO
 }
