PR #12301: Pass in missing QueryOptions in FindNearestHitTriangle

diff --git a/Engine/Source/Runtime/GeometryCore/Private/DynamicMesh/ColliderMesh.cpp b/Engine/Source/Runtime/GeometryCore/Private/DynamicMesh/ColliderMesh.cpp
--- a/Engine/Source/Runtime/GeometryCore/Private/DynamicMesh/ColliderMesh.cpp
+++ b/Engine/Source/Runtime/GeometryCore/Private/DynamicMesh/ColliderMesh.cpp
@@ -109,15 +109,16 @@ TMeshAABBTree3<FColliderMesh>* FColliderMesh::GetRawAABBTreeUnsafe()
 }
 
 
-bool FColliderMesh::FindNearestHitTriangle(const FRay3d& Ray, double& RayParameterOut, int& HitTriangleIDOut, FVector3d& BaryCoordsOut) const
+bool FColliderMesh::FindNearestHitTriangle(const FRay3d& Ray, double& RayParameterOut, int& HitTriangleIDOut, FVector3d& BaryCoordsOut, const IMeshSpatial::FQueryOptions& Options) const
 {
 	if ( AABBTree.IsValid(false) )
 	{
-		return AABBTree.FindNearestHitTriangle(Ray, RayParameterOut, HitTriangleIDOut, BaryCoordsOut);
+		return AABBTree.FindNearestHitTriangle(Ray, RayParameterOut, HitTriangleIDOut, BaryCoordsOut, Options);
 	}
 	return false;
 }
 
+
 int FColliderMesh::FindNearestTriangle(const FVector3d& Point, double& NearestDistSqrOut, const IMeshSpatial::FQueryOptions& Options) const
 {
 	if ( AABBTree.IsValid(false) )
diff --git a/Engine/Source/Runtime/GeometryCore/Public/DynamicMesh/ColliderMesh.h b/Engine/Source/Runtime/GeometryCore/Public/DynamicMesh/ColliderMesh.h
--- a/Engine/Source/Runtime/GeometryCore/Public/DynamicMesh/ColliderMesh.h
+++ b/Engine/Source/Runtime/GeometryCore/Public/DynamicMesh/ColliderMesh.h
@@ -110,7 +110,8 @@ public:
 	 * @param HitTriangleIDOut TriangleID on collider mesh - use GetSourceTriangleID() to map back to source mesh
 	 * @return true if a triangle was hit and Out parameters are initialized
 	 */
-	GEOMETRYCORE_API bool FindNearestHitTriangle( const FRay3d& Ray, double& RayParameterOut, int& HitTriangleIDOut, FVector3d& BaryCoordsOut ) const;
+	GEOMETRYCORE_API bool FindNearestHitTriangle( const FRay3d& Ray, double& RayParameterOut, int& HitTriangleIDOut, FVector3d& BaryCoordsOut, const IMeshSpatial::FQueryOptions& Options = IMeshSpatial::FQueryOptions() ) const;
+
 
 	/**
 	 * Find the TriangleID closest to P on the collider mesh, and distance to it, within distance MaxDist, or return InvalidID.
