Make convex decomposition produce a decomposition for planar inputs, by applying the 'thicken on hull failure' pass to the input mesh

diff --git a/Engine/Plugins/Runtime/GeometryProcessing/Source/DynamicMesh/Private/ShapeApproximation/MeshSimpleShapeApproximation.cpp b/Engine/Plugins/Runtime/GeometryProcessing/Source/DynamicMesh/Private/ShapeApproximation/MeshSimpleShapeApproximation.cpp
--- a/Engine/Plugins/Runtime/GeometryProcessing/Source/DynamicMesh/Private/ShapeApproximation/MeshSimpleShapeApproximation.cpp
+++ b/Engine/Plugins/Runtime/GeometryProcessing/Source/DynamicMesh/Private/ShapeApproximation/MeshSimpleShapeApproximation.cpp
@@ -390,6 +390,16 @@ void FMeshSimpleShapeApproximation::Generate_ConvexHullDecompositions(FSimpleSha
 		const FDynamicMesh3& SourceMesh = *SourceMeshes[idx];
 		FConvexDecomposition3::FPreprocessMeshOptions PreprocessOptions;
 		PreprocessOptions.bMergeEdges = true;
+
+		// Hull thickening allows us to compute decompositions planar inputs
+		PreprocessOptions.ThickenInputAfterHullFailure = ConvexDecompositionMinPartThickness;
+		// If protecting negative space, we also clamp thickening to at most half of the tolerance value so that we keep within the tolerance distance of the input shape
+		if (bConvexDecompositionProtectNegativeSpace)
+		{
+			PreprocessOptions.ThickenInputAfterHullFailure = FMath::Min(PreprocessOptions.ThickenInputAfterHullFailure, NegativeSpaceTolerance * .5);
+		}
+		PreprocessOptions.ThickenInputAfterHullFailure = FMath::Max(PreprocessOptions.ThickenInputAfterHullFailure, FMathd::ZeroTolerance);
+
 		PreprocessOptions.CustomPreprocess = [this](FDynamicMesh3& ProcessMesh, const FAxisAlignedBox3d& Bounds) -> void
 		{
 			// for solid inputs, flip orientation if the initial volume is negative
@@ -436,7 +446,7 @@ void FMeshSimpleShapeApproximation::Generate_ConvexHullDecompositions(FSimpleSha
 			Decomposition.MaxConvexEdgePlanes = 4;
 			Decomposition.bSplitDisconnectedComponents = false;
 			Decomposition.ConvexEdgeAngleMoreSamplesThreshold = 180;
-			Decomposition.ThickenAfterHullFailure = FMath::Max(FMathd::ZeroTolerance, NegativeSpaceTolerance * .01);
+			Decomposition.ThickenAfterHullFailure = PreprocessOptions.ThickenInputAfterHullFailure;
 			
 			Decomposition.InitializeNegativeSpace(Settings);
 
diff --git a/Engine/Source/Runtime/GeometryCore/Private/CompGeom/ConvexDecomposition3.cpp b/Engine/Source/Runtime/GeometryCore/Private/CompGeom/ConvexDecomposition3.cpp
--- a/Engine/Source/Runtime/GeometryCore/Private/CompGeom/ConvexDecomposition3.cpp
+++ b/Engine/Source/Runtime/GeometryCore/Private/CompGeom/ConvexDecomposition3.cpp
@@ -25,6 +25,80 @@ namespace UE
 namespace Geometry
 {
 
+static bool ComputeThickenedHull(double ThickenHullAfterFailureInLocalSpace, FConvexHull3d& HullCompute, const FDynamicMesh3& Mesh, TArray<FVector3d>& CutVertices, TArray<FVector3d>& AddedVertices, TFunctionRef<bool(int32)> UseMeshVertex = [](int32) {return true;})
+{
+	if (ThickenHullAfterFailureInLocalSpace <= 0)
+	{
+		return false;
+	}
+
+	const int32 OrigMaxID = Mesh.MaxVertexID();
+	FMeshNormals Normals(&Mesh);
+	bool bUseNormalsForOffset = HullCompute.GetDimension() >= 2;
+	if (bUseNormalsForOffset)
+	{
+		Normals.ComputeVertexNormals();
+	}
+	// Compute offsets to use based on the dimension of the failed convex hull
+	FVector3d OffsetBasis[3];
+	FVector3d FallbackOffset = FVector3d(1, 1, 1) * FMathd::InvSqrt3;
+	int32 NumOffsets = 0;
+	if (HullCompute.GetDimension() == 0)
+	{
+		OffsetBasis[0] = FVector3d(1, 0, 0);
+		OffsetBasis[1] = FVector3d(0, 1, 0);
+		OffsetBasis[2] = FVector3d(0, 0, 1);
+		NumOffsets = 3;
+	}
+	else if (HullCompute.GetDimension() == 1)
+	{
+		FLine3d Line = HullCompute.GetLine();
+		VectorUtil::MakePerpVectors(Line.Direction, OffsetBasis[0], OffsetBasis[1]);
+		NumOffsets = 2;
+	}
+	else if (HullCompute.GetDimension() == 2)
+	{
+		FPlane3d HullPlane = HullCompute.GetPlane();
+		OffsetBasis[0] = HullPlane.Normal;
+		FallbackOffset = HullPlane.Normal;
+		NumOffsets = 1;
+	}
+	const double OffsetFactor = ThickenHullAfterFailureInLocalSpace;
+	for (int32 VID = 0; VID < OrigMaxID; ++VID)
+	{
+		if (Mesh.IsVertex(VID) && UseMeshVertex(VID))
+		{
+
+			if (bUseNormalsForOffset)
+			{
+				FVector3d Offset = Normals[VID];
+				if (Offset == FVector::ZeroVector)
+				{
+					Offset = FallbackOffset;
+				}
+				AddedVertices.Add(Mesh.GetVertex(VID) - Offset * OffsetFactor);
+			}
+			else
+			{
+				for (int32 OffsetIdx = 0; OffsetIdx < NumOffsets; ++OffsetIdx)
+				{
+					AddedVertices.Add(Mesh.GetVertex(VID) - OffsetBasis[OffsetIdx] * OffsetFactor);
+				}
+			}
+
+		}
+	}
+	return HullCompute.Solve(Mesh.MaxVertexID() + CutVertices.Num() + AddedVertices.Num(),
+		[&Mesh, &CutVertices, &AddedVertices, &UseMeshVertex](int32 Index)
+		{
+			const int32 MaxVID = Mesh.MaxVertexID();
+			const int32 MaxCutVID = MaxVID + CutVertices.Num();
+			return Index < MaxVID ? Mesh.GetVertex(Index) : Index < MaxCutVID ? CutVertices[Index - MaxVID] : AddedVertices[Index - MaxCutVID];
+		},
+		[&Mesh, &UseMeshVertex](int32 Index) { return Index < Mesh.MaxVertexID() ? UseMeshVertex(Index) : true; });
+
+}
+
 bool FConvexDecomposition3::ConvexPartVsSphereOverlap(const FConvexDecomposition3::FConvexPart& Part, FVector3d Center, double Radius, const FTransform* TransformIntoSphereSpace, double* OutDistanceSq)
 {
 	if (OutDistanceSq)
@@ -752,7 +826,7 @@ struct FPartialCutResult
 		return HullMinusGeoVolume;
 	}
 
-	void ApplyToGeo(TIndirectArray<FConvexDecomposition3::FConvexPart>& Decomposition, int32 OrigPartIdx, const FPlane3d& Plane, int32& OtherSideStartIdxOut, double PlaneTol, double ConnectedComponentTolerance, const FSphereCovering& NegativeSpace, bool bTreatAsSolid, bool bSplitDisconnectedComponents)
+	void ApplyToGeo(TIndirectArray<FConvexDecomposition3::FConvexPart>& Decomposition, int32 OrigPartIdx, const FPlane3d& Plane, int32& OtherSideStartIdxOut, double PlaneTol, double ConnectedComponentTolerance, const FSphereCovering& NegativeSpace, bool bTreatAsSolid, bool bSplitDisconnectedComponents, double ThickenAfterHullFailureLocalSpace)
 	{
 		FConvexDecomposition3::FConvexPart& OrigPart = Decomposition[OrigPartIdx];
 		bool bSourceGeometryVolumeUnreliable = OrigPart.bGeometryVolumeUnreliable;
@@ -1045,7 +1119,7 @@ struct FPartialCutResult
 		// TODO: We may be able to skip hole filling if we're using an error metric that is not volume-based
 		bool bHoleFillResultUnreliable = !bTreatAsSolid || !HoleFill();
 
-		auto ComputeHullsIfMultipleComponents = [&Decomposition, &ConnectedComponentTolerance](int32 PartIdx)->int32
+		auto ComputeHullsIfMultipleComponents = [&Decomposition, &ConnectedComponentTolerance, ThickenAfterHullFailureLocalSpace](int32 PartIdx)->int32
 		{
 			FConvexDecomposition3::FConvexPart& Part = Decomposition[PartIdx];
 			FDynamicMesh3& Mesh = Part.InternalGeo;
@@ -1106,7 +1180,7 @@ struct FPartialCutResult
 			for (const TPair<int32, int32>& ToNewPartIdx : ComponentToIdxMap)
 			{
 				int32 NewPartIdx = ToNewPartIdx.Value;
-				bHasFailedComponents = !NewParts[NewPartIdx]->ComputeHull() || bHasFailedComponents;
+				bHasFailedComponents = !NewParts[NewPartIdx]->ComputeHull(true, ThickenAfterHullFailureLocalSpace) || bHasFailedComponents;
 				if (bHasFailedComponents)
 				{
 					break;
@@ -1303,7 +1377,7 @@ struct FPartialCutResult
 
 	FPartialCutResult() {}
 
-	FPartialCutResult(const FConvexDecomposition3::FConvexPart& Convex, const FPlane3d& Plane, double PlaneTol, bool bTreatAsSolid, double ThickenHullAfterFailure)
+	FPartialCutResult(const FConvexDecomposition3::FConvexPart& Convex, const FPlane3d& Plane, double PlaneTol, bool bTreatAsSolid, double ThickenHullAfterFailureInLocalSpace)
 	{
 		SignDists.SetNumZeroed(Convex.InternalGeo.MaxVertexID());
 		TArray<bool> OnHull; // indicates if vertex was on the original hull
@@ -1430,75 +1504,8 @@ struct FPartialCutResult
 			if (!bOK)
 			{
 				// optionally re-try w/ minimal offsets in degen directions
-				if (ThickenHullAfterFailure > 0)
-				{
-					int32 Dimension = HullCompute.GetDimension();
-					const int32 OrigMaxID = Convex.InternalGeo.MaxVertexID();
-					FMeshNormals Normals(&Convex.InternalGeo);
-					bool bUseNormalsForOffset = Dimension >= 2;
-					if (bUseNormalsForOffset)
-					{
-						Normals.ComputeVertexNormals();
-					}
-					// Compute offsets to use based on the dimension of the failed convex hull
-					FVector3d OffsetBasis[3];
-					FVector3d FallbackOffset = FVector3d(1, 1, 1) * FMathd::InvSqrt3;
-					int32 NumOffsets = 0;
-					if (Dimension == 0)
-					{
-						OffsetBasis[0] = FVector3d(1, 0, 0);
-						OffsetBasis[1] = FVector3d(0, 1, 0);
-						OffsetBasis[2] = FVector3d(0, 0, 1);
-						NumOffsets = 3;
-					}
-					else if (Dimension == 1)
-					{
-						FLine3d Line = HullCompute.GetLine();
-						VectorUtil::MakePerpVectors(Line.Direction, OffsetBasis[0], OffsetBasis[1]);
-						NumOffsets = 2;
-					}
-					else if (Dimension == 2)
-					{
-						FPlane3d HullPlane = HullCompute.GetPlane();
-						OffsetBasis[0] = HullPlane.Normal;
-						FallbackOffset = HullPlane.Normal;
-						NumOffsets = 1;
-					}
-					const double OffsetFactor = ThickenHullAfterFailure;
-					for (int32 VID = 0; VID < OrigMaxID; ++VID)
-					{
-						if (Convex.InternalGeo.IsVertex(VID) && ForHull[Side][VID])
-						{
-							
-							if (bUseNormalsForOffset)
-							{
-								FVector3d Offset = Normals[VID];
-								if (Offset == FVector::ZeroVector)
-								{
-									Offset = FallbackOffset;
-								}
-								OffsetVertices[Side].Add(Convex.InternalGeo.GetVertex(VID) - Offset * OffsetFactor);
-							}
-							else
-							{
-								for (int32 OffsetIdx = 0; OffsetIdx < NumOffsets; ++OffsetIdx)
-								{
-									OffsetVertices[Side].Add(Convex.InternalGeo.GetVertex(VID) - OffsetBasis[OffsetIdx] * OffsetFactor);
-								}
-							}
-							
-						}
-					}
-					bOK = HullCompute.Solve(Convex.InternalGeo.MaxVertexID() + CutVertices.Num() + OffsetVertices[Side].Num(),
-						[this, &Convex, &ForHull, Side](int32 Index)
-						{
-							const int32 MaxVID = Convex.InternalGeo.MaxVertexID();
-							const int32 MaxCutVID = MaxVID + CutVertices.Num();
-							return Index < MaxVID ? Convex.InternalGeo.GetVertex(Index) : Index < MaxCutVID ? CutVertices[Index - MaxVID] : OffsetVertices[Side][Index - MaxCutVID];
-						},
-						[this, &Convex, &ForHull, Side](int32 Index) { return Index < Convex.InternalGeo.MaxVertexID() ? ForHull[Side][Index] : true; });
+				bOK = ComputeThickenedHull(ThickenHullAfterFailureInLocalSpace, HullCompute, Convex.InternalGeo, CutVertices, OffsetVertices[Side], [&ForHull, Side](int32 VID) { return ForHull[Side][VID]; });
 
-				}
 				if (!bOK)
 				{
 					// don't cut the hull on a plane that would lead to a degenerate hull on either side
@@ -1660,8 +1667,10 @@ void FConvexDecomposition3::FConvexPart::InitializeFromInternalGeo(const FConvex
 	TransformOut = FTransformSRT3d(InitialBounds.Center());
 	TransformOut.SetScale(FVector3d(InvScaleFactor, InvScaleFactor, InvScaleFactor));
 
+	double ThickenInLocalSpace = Preprocess.ThickenInputAfterHullFailure * ScaleFactor;
+
 	// Compute hull and standard measurements (volume, center, bounds, etc)
-	ComputeHull();
+	ComputeHull(true /*bComputePlanes*/, ThickenInLocalSpace);
 	ComputeStats();
 }
 
@@ -1772,17 +1781,29 @@ void FConvexDecomposition3::Compute(int32 NumOutputHulls, int32 NumAdditionalSpl
 }
 
 
-bool FConvexDecomposition3::FConvexPart::ComputeHull(bool bComputePlanes)
+
+bool FConvexDecomposition3::FConvexPart::ComputeHull(bool bComputePlanes, double ThickenAfterHullFailureInLocalSpace)
 {
 	FConvexHull3d HullCompute;
 	bool bOK = HullCompute.Solve(InternalGeo.MaxVertexID(),
 		[this](int32 Index) { return InternalGeo.GetVertex(Index); },
 		[this](int32 Index) { return InternalGeo.IsVertex(Index); });
-	if (!bOK)
+	if (!bOK && ThickenAfterHullFailureInLocalSpace > 0)
 	{
-		int32 Dimension = HullCompute.GetDimension();
-		bFailed = true;
-		return false;
+		TArray<FVector3d> CutVertices_Unused, ThickenVertices;
+		if (ComputeThickenedHull(ThickenAfterHullFailureInLocalSpace, HullCompute, InternalGeo, CutVertices_Unused, ThickenVertices))
+		{
+			bOK = true;
+			for (FVector3d V : ThickenVertices)
+			{
+				InternalGeo.AppendVertex(V);
+			}
+		}
+		else
+		{
+			bFailed = true;
+			return false;
+		}
 	}
 	HullTriangles = HullCompute.MoveTriangles();
 	if (bComputePlanes)
@@ -2124,7 +2145,7 @@ bool FConvexDecomposition3::SplitWorstHelper(bool bCanSkipUnreliableGeoVolumes,
 
 	int32 NewPartsStartIdx = Decomposition.Num();
 	int32 OtherSideStartIdx = -1;
-	BestCutResult.ApplyToGeo(Decomposition, WorstIdx, CandidatePlanes[BestPlaneIdx], OtherSideStartIdx, OnPlaneTolerance, ConnectedComponentTolerance, NegativeSpace, bCutAsSolid, bSplitDisconnectedComponents);
+	BestCutResult.ApplyToGeo(Decomposition, WorstIdx, CandidatePlanes[BestPlaneIdx], OtherSideStartIdx, OnPlaneTolerance, ConnectedComponentTolerance, NegativeSpace, bCutAsSolid, bSplitDisconnectedComponents, ThickenAfterHullFailureLocalSpace);
 
 	UpdateProximitiesAfterSplit(WorstIdx, NewPartsStartIdx, CandidatePlanes[BestPlaneIdx], OtherSideStartIdx, OrigHullVolume);
 
@@ -2487,7 +2508,7 @@ int32 FConvexDecomposition3::MergeBest(const FMergeSettings& Settings)
 						double OrigHullVolume = PartToSplit.HullVolume;
 						int32 NewPartsStartIdx = Decomposition.Num();
 						int32 OtherSideStartIdx = -1;
-						PlaneResult.ApplyToGeo(Decomposition, PartIdx, Plane, OtherSideStartIdx, OnPlaneTolerance, ConnectedComponentTolerance, NegativeSpace, bTreatAsSolid, bSplitDisconnectedComponents);
+						PlaneResult.ApplyToGeo(Decomposition, PartIdx, Plane, OtherSideStartIdx, OnPlaneTolerance, ConnectedComponentTolerance, NegativeSpace, bTreatAsSolid, bSplitDisconnectedComponents, ThickenAfterHullFailureLocalSpace);
 						UpdateProximitiesAfterSplit(PartIdx, NewPartsStartIdx, Plane, OtherSideStartIdx, OrigHullVolume);
 						PartToSplit.bMustMerge = true;
 						for (int32 NewIdx = NewPartsStartIdx; NewIdx < Decomposition.Num(); NewIdx++)
diff --git a/Engine/Source/Runtime/GeometryCore/Public/CompGeom/ConvexDecomposition3.h b/Engine/Source/Runtime/GeometryCore/Public/CompGeom/ConvexDecomposition3.h
--- a/Engine/Source/Runtime/GeometryCore/Public/CompGeom/ConvexDecomposition3.h
+++ b/Engine/Source/Runtime/GeometryCore/Public/CompGeom/ConvexDecomposition3.h
@@ -290,6 +290,11 @@ public:
 		// Whether to attempt to merge boundary edges before processing
 		bool bMergeEdges = true;
 
+		// If greater than zero, we can 'inflate' the input shape by this amount along any degenerate axes in cases where the hull failed to construct due to the input being coplanar/colinear
+		// Note: If passed to the FConvexDecomposition3 constructor, will also set the similar ThickenAfterHullFailure 
+		// member, which applies to generated parts of the decomposition.
+		double ThickenInputAfterHullFailure = 0.0;
+
 		// CustomPreprocess takes the mesh to process, and that mesh's bounding box, and updates the mesh in place
 		// Applied after discarding attributes and applying Merge Edges (if enabled)
 		// Can be used e.g. to simplify/repair/solidify the input
@@ -298,6 +303,7 @@ public:
 
 	FConvexDecomposition3(const FDynamicMesh3& SourceMesh, const FPreprocessMeshOptions& Options)
 	{
+		ThickenAfterHullFailure = Options.ThickenInputAfterHullFailure;
 		InitializeFromMesh(SourceMesh, Options);
 	}
 
@@ -664,7 +670,7 @@ public:
 		}
 
 		// Helper to create hull after InternalGeo is set
-		bool ComputeHull(bool bComputePlanes = true);
+		bool ComputeHull(bool bComputePlanes = true, double ThickenAfterHullFailureInLocalSpace = 0.0);
 
 		// Helper to compute volumes, centroid, bounds, etc after the convex part is initialized
 		void ComputeStats();
