GeometryCore: Make TDynamicVector BlockSize a template parameter plus internal refactoring

diff --git a/Engine/Source/Runtime/GeometryCore/Public/Util/DynamicVector.h b/Engine/Source/Runtime/GeometryCore/Public/Util/DynamicVector.h
--- a/Engine/Source/Runtime/GeometryCore/Public/Util/DynamicVector.h
+++ b/Engine/Source/Runtime/GeometryCore/Public/Util/DynamicVector.h
@@ -1,7 +1,5 @@
 // Copyright Epic Games, Inc. All Rights Reserved.
 
-// Port of geometry3cpp dvector
-
 #pragma once
 
 #include <CoreMinimal.h>
@@ -24,56 +22,118 @@ using namespace UE::Math;
  *
  * Iterator functions suitable for use with range-based for are provided
  */
-template <class Type>
+template <typename Type, int32 BlockSize = 512>
 class TDynamicVector
 {
+	static_assert(BlockSize > 0, "TDynamicVector: BlockSize must be larger than zero.");
+	static_assert(((BlockSize & (BlockSize - 1)) == 0), "TDynamicVector: BlockSize must be a power of two.");
+
+	static constexpr uint32 NumBitsNeeded(const uint32 N) {
+		return N <= 1 ? 0 : 1 + NumBitsNeeded((N + 1) / 2);
+	}
+
+	static constexpr uint32 GetBlockIndex(const uint32 Index)
+	{
+		constexpr int BlockBitsShift = NumBitsNeeded(BlockSize);
+		return Index >> BlockBitsShift;
+	}
+
+	static constexpr uint32 GetIndexInBlock(const uint32 Index)
+	{
+		constexpr int BlockBitMask = BlockSize - 1;
+		return Index & BlockBitMask;
+	}
+
 public:
-	static constexpr uint32 MaxSize = MAX_uint32;
+	using ElementType = Type;
 
 	TDynamicVector()
 	{
-		Blocks.Add(new BlockType());
+		AddAllocatedBlock();
+	}
+
+	TDynamicVector(const TDynamicVector& Copy)
+		: CurBlock(Copy.CurBlock)
+		, CurBlockUsed(Copy.CurBlockUsed)
+	{
+		const int32 N = Copy.Blocks.Num();
+		Blocks.Reserve(N);
+		for (int32 k = 0; k < N; ++k)
+		{
+			Blocks.Add(new TBlock(*Copy.Blocks[k]));
+		}
 	}
 
-	TDynamicVector(const TDynamicVector& Copy) = default;
 	TDynamicVector(TDynamicVector&& Moved)
+		: CurBlock(Moved.CurBlock)
+		, CurBlockUsed(Moved.CurBlockUsed)
+		, Blocks(MoveTemp(Moved.Blocks))
 	{
-		Blocks = MoveTemp(Moved.Blocks);
-		CurBlock = MoveTemp(Moved.CurBlock);
-		CurBlockUsed = MoveTemp(Moved.CurBlockUsed);
-		Moved.Blocks.Add(new BlockType());
 		Moved.CurBlock = 0;
 		Moved.CurBlockUsed = 0;
+		Moved.AddAllocatedBlock();
+	}
+
+	TDynamicVector& operator=(const TDynamicVector& Copy)
+	{
+		if (this != &Copy)
+		{
+			const int32 N = Copy.Blocks.Num();
+			Empty(N);
+			CurBlock = Copy.CurBlock;
+			CurBlockUsed = Copy.CurBlockUsed;
+			for (int32 k = 0; k < N; ++k)
+			{
+				Blocks.Add(new TBlock(*Copy.Blocks[k]));
+			}
+		}
+		return *this;
 	}
-	TDynamicVector& operator=(const TDynamicVector& Copy) = default;
+
 	TDynamicVector& operator=(TDynamicVector&& Moved)
 	{
-		Blocks = MoveTemp(Moved.Blocks);
-		CurBlock = MoveTemp(Moved.CurBlock);
-		CurBlockUsed = MoveTemp(Moved.CurBlockUsed);
-		Moved.Blocks.Add(new BlockType());
-		Moved.CurBlock = 0;
-		Moved.CurBlockUsed = 0;
+		if (this != &Moved)
+		{
+			Empty();
+
+			CurBlock = Moved.CurBlock;
+			CurBlockUsed = Moved.CurBlockUsed;
+			Blocks = MoveTemp(Moved.Blocks);
+
+			Moved.CurBlock = 0;
+			Moved.CurBlockUsed = 0;
+			Moved.AddAllocatedBlock();
+		}
 		return *this;
 	}
 
 	TDynamicVector(const TArray<Type>& Array)
 	{
-		SetNum((unsigned int) Array.Num());
-		for (int32 Idx = 0; Idx < Array.Num(); Idx++)
+		const uint32 N = static_cast<uint32>(Array.Num());
+		SetNum(N);
+		const Type* ArrayPtr = Array.GetData();
+		for (uint32 Idx = 0; Idx < N; ++Idx)
 		{
-			(*this)[(unsigned int) Idx] = Array[Idx];
+			(*this)[Idx] = ArrayPtr[Idx];
 		}
 	}
+
 	TDynamicVector(TArrayView<const Type> Array)
 	{
-		SetNum((unsigned int) Array.Num());
-		for (int32 Idx = 0; Idx < Array.Num(); Idx++)
+		const uint32 N = static_cast<uint32>(Array.Num());
+		SetNum(N);
+		const Type* ArrayPtr = Array.GetData();
+		for (uint32 Idx = 0; Idx < N; ++Idx)
 		{
-			(*this)[(unsigned int) Idx] = Array[Idx];
+			(*this)[Idx] = ArrayPtr[Idx];
 		}
 	}
 
+	~TDynamicVector()
+	{
+		Empty();
+	}
+
 	inline void Clear();
 	inline void Fill(const Type& Value);
 	inline void Resize(unsigned int Count);
@@ -85,11 +145,13 @@ public:
 	inline bool IsEmpty() const { return CurBlock == 0 && CurBlockUsed == 0; }
 	inline size_t GetLength() const { return CurBlock * BlockSize + CurBlockUsed; }
 	inline size_t Num() const { return GetLength(); }
-	static int GetBlockSize() { return BlockSize; }
+	static constexpr int32 GetBlockSize() { return BlockSize; }
 	inline size_t GetByteCount() const { return Blocks.Num() * BlockSize * sizeof(Type); }
 
 	inline void Add(const Type& Data);
-	inline void Add(const TDynamicVector& Data);
+	template <int32 BlockSizeData> void Add(const TDynamicVector<Type, BlockSizeData>& Data);
+	void Add(const TArray<Type>& Data);
+	void Add(TArrayView<const Type> Data);
 	inline void PopBack();
 
 	inline void InsertAt(const Type& Data, unsigned int Index);
@@ -99,47 +161,35 @@ public:
 	inline const Type& Front() const
 	{
 		checkSlow(CurBlockUsed > 0);
-		return Blocks[0][0];
+		return GetElement(0, 0);
 	}
 
 	inline const Type& Back() const
 	{
 		checkSlow(CurBlockUsed > 0);
-		return Blocks[CurBlock][CurBlockUsed - 1];
+		return GetElement(CurBlock, CurBlockUsed - 1);
 	}
 
-
 #if USING_ADDRESS_SANITISER
-	FORCENOINLINE Type& operator[](unsigned int Index)
+	FORCENOINLINE
+#endif
+	const Type& operator[](uint32 Index) const
 	{
 		checkSlow(Index < Num());
-
-		return Blocks[Index >> nShiftBits][Index & BlockIndexBitmask];
+		return GetElement(GetBlockIndex(Index), GetIndexInBlock(Index));
 	}
-	FORCENOINLINE const Type& operator[](unsigned int Index) const
-	{
-		checkSlow(Index < Num());
 
-		return Blocks[Index >> nShiftBits][Index & BlockIndexBitmask];
-	}
-#else
-	inline Type& operator[](unsigned int Index)
-	{
-		checkSlow(Index < Num());
-
-		return Blocks[Index >> nShiftBits][Index & BlockIndexBitmask];
-	}
-	inline const Type& operator[](unsigned int Index) const
+#if USING_ADDRESS_SANITISER
+	FORCENOINLINE
+#endif
+	Type& operator[](uint32 Index)
 	{
-		checkSlow(Index < Num());
-
-		return Blocks[Index >> nShiftBits][Index & BlockIndexBitmask];
+		return const_cast<Type&>(const_cast<const TDynamicVector&>(*this)[Index]);
 	}
-#endif
 
-	// apply f() to each member sequentially
+	// apply ApplyFunc() to each member sequentially
 	template <typename Func>
-	void Apply(const Func& f);
+	void Apply(const Func& ApplyFunc);
 
 	/**
 	 * Serialization operator for TDynamicVector.
@@ -161,51 +211,8 @@ public:
 	 * @param Ar Archive to serialize with
 	 */
 	template <bool bForceBulkSerialization = false, bool bUseCompression = false>
-	void Serialize(FArchive& Ar)
-	{
-		Ar.UsingCustomVersion(FUE5MainStreamObjectVersion::GUID);
-		if (Ar.IsLoading() && Ar.CustomVer(FUE5MainStreamObjectVersion::GUID) < FUE5MainStreamObjectVersion::DynamicMeshCompactedSerialization)
-		{
-			const SIZE_T CountBytes = 3 * sizeof(uint32) + Blocks.Num() * BlockSize * sizeof(Type);
-			Ar.CountBytes(CountBytes, CountBytes);
-			Ar << CurBlock;
-			Ar << CurBlockUsed;
-			Blocks.Serialize_LegacyLoad(Ar);
-
-			// Account for a previously existing issue where one additional empty block was allocated and the resulting values for the current block can cause
-			// crashes, for example when using Back().
-			if (CurBlock > 0 && CurBlockUsed == 0)
-			{
-				const SIZE_T Count = Num();
-				checkSlow(Count % BlockSize == 0);
-				checkSlow(Count / BlockSize < Blocks.Num());
-
-				// Fix CurBlock and CurBlockUsed.
-				SetCurBlock(Count);
+	void Serialize(FArchive& Ar);
 
-				// Remove empty block.
-				Blocks.Truncate(int32(Count / BlockSize), EAllowShrinking::No);
-			}
-		}
-		else
-		{
-			uint32 SerializeNum = Num();
-			const SIZE_T CountBytes =  sizeof(uint32) + SerializeNum * sizeof(Type);
-			Ar.CountBytes(CountBytes, CountBytes);
-			Ar << SerializeNum;
-			if (SerializeNum == 0 && Ar.IsLoading())
-			{
-				Clear();
-			}
-			else if (SerializeNum > 0)
-			{
-				SetCurBlock(SerializeNum);
-				Blocks.template Serialize<bForceBulkSerialization, bUseCompression>(Ar, SerializeNum);
-			}
-		}
-	}
-
-public:
 	/*
 	 * FIterator class iterates over values of vector
 	 */
@@ -251,7 +258,7 @@ public:
 	/** @return iterator at beginning of vector */
 	FIterator begin()
 	{
-		return IsEmpty() ? end() : FIterator{this, 0};
+		return FIterator{this, 0};
 	}
 	/** @return iterator at end of vector */
 	FIterator end()
@@ -300,7 +307,7 @@ public:
 	/** @return iterator at beginning of vector */
 	FConstIterator begin() const
 	{
-		return IsEmpty() ? end() : FConstIterator{this, 0};
+		return FConstIterator{this, 0};
 	}
 	/** @return iterator at end of vector */
 	FConstIterator end() const
@@ -309,264 +316,92 @@ public:
 	}
 
 private:
-	static constexpr int nShiftBits = 9;
-	static constexpr int BlockSize = 1 << nShiftBits;
-	static constexpr int BlockIndexBitmask = BlockSize - 1; // low 9 bits
-	static_assert( BlockSize && ((BlockSize & (BlockSize - 1)) == 0), "DynamicVector: BlockSize must be a power of two");
-	static constexpr unsigned int MaxBlockCount = MaxSize / BlockSize + 1;
+	struct TBlock
+	{
+		Type Elements[BlockSize];
+	};
 
-	unsigned int CurBlock{0};  //< Current block index; always points to the block with the last item in the vector or it is set to zero if the vector is empty. 
+	unsigned int CurBlock{0};  //< Current block index; always points to the block with the last item in the vector, or is set to zero if the vector is empty. 
 	unsigned int CurBlockUsed{0};  //< Number of used items in the current block.
 
-	using BlockType = TStaticArray<Type, BlockSize>;
-	
+	TArray<TBlock*> Blocks;
 
-	/**
-	 * This is used to store/manage an array of BlockType*. Basically a stripped
-	 * down TIndirectArray that does not do range-checking on operator[] in non-debug builds
-	 */
-	template<typename ArrayType>
-	class TBlockVector
+	void AddAllocatedBlock()
 	{
-	protected:
-		TArray<ArrayType*> Elements;
-
-	public: 
-		TBlockVector() = default;
-		TBlockVector(TBlockVector&& Moved) = default;
+		Blocks.Add(new TBlock);
+	}
 
-		TBlockVector(const TBlockVector& Copy)
-		{
-			int32 N = Copy.Num();
-			Elements.Reserve(N);
-			for (int32 k = 0; k < N; ++k)
-			{
-				Elements.Add(new ArrayType(*Copy.Elements[k]));
-			}
-		}
-		TBlockVector& operator=(const TBlockVector& Copy)
-		{
-			if (&Copy != this)
-			{
-				int32 N = Copy.Num();
-				Empty(N);
-				Elements.Reserve(N);
-				for (int32 k = 0; k < N; ++k)
-				{
-					Elements.Add(new ArrayType(*Copy.Elements[k]));
-				}
-			}
-			return *this;
-		}
-		TBlockVector& operator=(TBlockVector&& Moved)
-		{
-			if (&Moved != this)
-			{
-				Empty();
-				Elements = MoveTemp(Moved.Elements);
-			}
-			return *this;
-		}
-		~TBlockVector()
+	void Empty(int32 NewReservedBlockCount = 0)
+	{
+		for (int32 k = 0, N = Blocks.Num(); k < N; ++k)
 		{
-			Empty();
+			delete Blocks[k];
 		}
+		Blocks.Empty(NewReservedBlockCount);
+	}
 
-		int32 Num() const { return Elements.Num(); }
+	const Type& GetElement(int32 BlockIndex, int32 IndexInBlock) const
+	{
+		checkSlow(0 <= BlockIndex && BlockIndex < Blocks.Num() && 0 <= IndexInBlock && IndexInBlock < BlockSize);
+		return Blocks.GetData()[BlockIndex]->Elements[IndexInBlock];
+	}
 
-#if UE_BUILD_DEBUG
-		ArrayType& operator[](int32 Index) { return *Elements[Index]; }
-		const ArrayType& operator[](int32 Index) const { return *Elements[Index]; }
-#else
-		ArrayType& operator[](int32 Index) { return *Elements.GetData()[Index]; }
-		const ArrayType& operator[](int32 Index) const { return *Elements.GetData()[Index]; }
-#endif
+	Type& GetElement(int32 BlockIndex, int32 IndexInBlock)
+	{
+		return const_cast<Type&>(const_cast<const TDynamicVector&>(*this).GetElement(BlockIndex, IndexInBlock));
+	}
 
-		void Add(ArrayType* NewElement)
+	void TruncateBlocks(int32 NewBlockCount, EAllowShrinking AllowShrinking)
+	{
+		for (int32 k = NewBlockCount; k < Blocks.Num(); ++k)
 		{
-			Elements.Add(NewElement);
+			delete Blocks[k];
+			Blocks[k] = nullptr;
 		}
+		Blocks.RemoveAt(NewBlockCount, Blocks.Num() - NewBlockCount, AllowShrinking);
+	}
 
-		void Truncate(int32 NewElementCount, EAllowShrinking AllowShrinking)
+	template <int32 BlockSizeRhs>
+	friend bool operator==(const TDynamicVector& Lhs, const TDynamicVector<Type, BlockSizeRhs>& Rhs)
+	{
+		if (Lhs.Num() != Rhs.Num())
 		{
-			for (int32 k = NewElementCount; k < Elements.Num(); ++k)
-			{
-				delete Elements[k];
-				Elements[k] = nullptr;
-			}
-			Elements.RemoveAt(NewElementCount, Elements.Num() - NewElementCount, AllowShrinking);
+			return false;
 		}
 
-		void Empty(int32 NewReservedSize = 0)
+		if (Lhs.IsEmpty())
 		{
-			for (int32 k = 0, N = Elements.Num(); k < N; ++k)
-			{
-				delete Elements[k];
-			}
-			Elements.Empty(NewReservedSize);
+			return true;
 		}
 
-		/** Serialize TBlockVector to archive. */
-		void Serialize_LegacyLoad(FArchive& Ar)
+		if constexpr (BlockSize == BlockSizeRhs)
 		{
-			// Only loading of legacy archives is allowed.
-			checkSlow(Ar.IsLoading());
-			if (Ar.IsLoading())
+			const uint32 LhsCurBlock = Lhs.CurBlock;
+			for (uint32 BlockIndex = 0; BlockIndex < LhsCurBlock; ++BlockIndex)
 			{
-				// Bulk serialization for a number of double types was enabled as part of the transition to Large World Coordinates.
-				// If the currently stored type is one of these types, and the archive is from before bulk serialization for these types was enabled,
-				// we need to still use per element serialization for legacy data.
-				constexpr bool bIsLWCBulkSerializedDoubleType =
-					std::is_same_v<Type, FVector2d> ||
-					std::is_same_v<Type, FVector3d> ||
-					std::is_same_v<Type, FVector4d> ||
-					std::is_same_v<Type, FQuat4d> ||
-					std::is_same_v<Type, FTransform3d>;
-				const bool bUseBulkSerialization = TCanBulkSerialize<Type>::Value && !(bIsLWCBulkSerializedDoubleType && Ar.UEVer() <
-					EUnrealEngineObjectUE5Version::LARGE_WORLD_COORDINATES);
-
-				// Lambda for serializing a block element either via bulk serializing the contained data or via serializing the element container itself.
-				// Note that the static_cast<> was necessary to resolve compiler errors when using MSVC.
-				const auto SerializeElement = bUseBulkSerialization
-					                              ? static_cast<void(*)(FArchive&, ArrayType*)>([](FArchive& Archive, ArrayType* Element)
-					                              {
-						                              Archive.Serialize(Element->GetData(), Element->Num() * sizeof(Type));
-					                              })
-					                              : static_cast<void(*)(FArchive&, ArrayType*)>([](FArchive& Archive, ArrayType* Element)
-					                              {
-						                              Archive << *Element;
-					                              });
-
-				int32 BlockNum = Num();
-				Ar << BlockNum;
-				Empty(BlockNum);
-
-				for (int32 Index = 0; Index < BlockNum; Index++)
+				if (!CompareItems(&Lhs.Blocks[BlockIndex]->Elements[0], &Rhs.Blocks[BlockIndex]->Elements[0], BlockSize))
 				{
-					ArrayType* NewElement = new ArrayType;
-					SerializeElement(Ar, NewElement);
-					Add(NewElement);
+					return false;
 				}
 			}
+			return CompareItems(&Lhs.Blocks[LhsCurBlock]->Elements[0], &Rhs.Blocks[LhsCurBlock]->Elements[0], Lhs.CurBlockUsed);
 		}
-
-		template <bool bForceBulkSerialization, bool bUseCompression>
-		void Serialize(FArchive& Ar, uint32 Num)
+		else
 		{
-			constexpr bool bUseBulkSerialization = bForceBulkSerialization || TCanBulkSerialize<Type>::Value || sizeof(Type) == 1;
-			static_assert(!bUseCompression || bUseBulkSerialization, "Compression only available when using bulk serialization");
-
-			checkSlow(Num > 0);
-
-			// Serialize compression flag, which adds flexibility when de-serializing existing data even if some implementation details change.  
-			bool bUseCompressionForBulkSerialization = bUseBulkSerialization && bUseCompression;
-			Ar << bUseCompressionForBulkSerialization;
-
-			// Determine number of blocks.
-			const bool bNumIsNotMultipleOfBlockSize = Num % BlockSize != 0;
-			const uint32 NumBlocks = Num / BlockSize + bNumIsNotMultipleOfBlockSize;
-			
-			if (bUseCompressionForBulkSerialization)
+			for (int32 Index = 0, Num = Lhs.Num(); Index < Num; ++Index)
 			{
-				// When using compression, copy everything to/from a big single buffer and serialize the big buffer.
-				// This results in better compression ratios while at the same time accelerating compression. 
-
-				TArray<Type> Buffer;
-				Buffer.SetNumUninitialized(Num);
-				Type* BufferPtr = Buffer.GetData();
-				SIZE_T NumCopyRemaining = Num;
-
-				if (!Ar.IsLoading())
-				{					
-					for (uint32 Index = 0; Index < NumBlocks; ++Index)
-					{
-						const SIZE_T NumCopy = FMath::template Min<SIZE_T>(NumCopyRemaining, (SIZE_T)BlockSize);
-						FMemory::Memcpy(BufferPtr, Elements[Index]->GetData(), NumCopy * sizeof(Type));
-						BufferPtr += NumCopy;
-						NumCopyRemaining -= BlockSize;
-					}
-				}
-
-				Ar.SerializeCompressedNew(Buffer.GetData(), Num * sizeof(Type), NAME_Oodle, NAME_Oodle, COMPRESS_NoFlags, false, nullptr);
-
-				if (Ar.IsLoading())
+				if (!(Lhs[Index] == Rhs[Index]))
 				{
-					Empty(NumBlocks);
-					for (uint32 Index = 0; Index < NumBlocks; ++Index)
-					{
-						ArrayType *const NewElement = new ArrayType;
-						const SIZE_T NumCopy = FMath::template Min<SIZE_T>(NumCopyRemaining, (SIZE_T)BlockSize);
-						FMemory::Memcpy(NewElement->GetData(), BufferPtr, NumCopy * sizeof(Type));
-						Add(NewElement);
-						BufferPtr += NumCopy;
-						NumCopyRemaining -= BlockSize;
-					}
+					return false;
 				}
 			}
-			else
-			{
-				const auto SerializeElement = [&Ar, bUseBulkSerialization](ArrayType* Element, uint32 ElementNum)
-				{
-					if (bUseBulkSerialization)
-					{
-						Ar.Serialize(Element->GetData(), ElementNum * sizeof(Type));
-					}
-					else
-					{
-						for (uint32 Index = 0; Index < ElementNum; ++Index)
-						{
-							Ar << (*Element)[Index];
-						}
-					}
-				};
 
-				if (Ar.IsLoading())
-				{
-					Empty(NumBlocks);
-					for (uint32 Index = 0; Index < NumBlocks; ++Index)
-					{
-						ArrayType *const NewElement = new ArrayType;
-						SerializeElement(NewElement, FMath::template Min<uint32>(Num, BlockSize));
-						Add(NewElement);
-						Num -= BlockSize;
-					}
-				}
-				else
-				{
-					for (uint32 Index = 0; Index < NumBlocks; ++Index)
-					{
-						SerializeElement(&(*this)[Index], FMath::template Min<uint32>(Num, BlockSize));
-						Num -= BlockSize;
-					}
-				}
-			}
-		}
-	};
-
-	friend bool operator==(const TDynamicVector& Lhs, const TDynamicVector& Rhs)
-	{
-		if (Lhs.Num() != Rhs.Num())
-		{
-			return false;
-		}
-
-		if (Lhs.IsEmpty())
-		{
 			return true;
 		}
-		
-		const uint32 LhsCurBlock = Lhs.CurBlock;
-		for (uint32 BlockIndex = 0; BlockIndex < LhsCurBlock; ++BlockIndex)
-		{
-			if (!CompareItems(Lhs.Blocks[BlockIndex].GetData(), Rhs.Blocks[BlockIndex].GetData(), BlockSize))
-			{
-				return false;
-			}
-		}
-		return CompareItems(Lhs.Blocks[LhsCurBlock].GetData(), Rhs.Blocks[LhsCurBlock].GetData(), Lhs.CurBlockUsed);
 	}
 
-	friend bool operator!=(const TDynamicVector& Lhs, const TDynamicVector& Rhs)
+	template <int32 BlockSizeRhs>
+	friend bool operator!=(const TDynamicVector& Lhs, const TDynamicVector<Type, BlockSizeRhs>& Rhs)
 	{
 		return !(Lhs == Rhs);
 	}
@@ -577,12 +412,9 @@ private:
 		// This is similar to what happens when computing the indices in operator[], but we additionally account for (1) the vector being empty and (2) that the
 		// used item count within the last block needs to be one more than the index of the last item. 
 		const int32 LastItemIndex = int32(Count - 1);
-		CurBlock = Count != 0 ? LastItemIndex >> nShiftBits : 0;
-		CurBlockUsed = Count != 0 ? (LastItemIndex & BlockIndexBitmask) + 1 : 0;
+		CurBlock = Count != 0 ? GetBlockIndex(LastItemIndex) : 0;
+		CurBlockUsed = Count != 0 ? GetIndexInBlock(LastItemIndex) + 1 : 0;
 	}
-
-	// memory chunks for dynamic vector
-	TBlockVector<BlockType> Blocks;
 };
 
 template <class Type, int N>
@@ -737,33 +569,33 @@ using  TDynamicVector2d = TDynamicVectorN<double, 2>;
 using  TDynamicVector3i = TDynamicVectorN<int, 3>;
 using  TDynamicVector2i = TDynamicVectorN<int, 2>;
 
-template <class Type>
-void TDynamicVector<Type>::Clear()
+template <typename Type, int32 BlockSize>
+void TDynamicVector<Type, BlockSize>::Clear()
 {
-	Blocks.Truncate(1, EAllowShrinking::No);
+	TruncateBlocks(1, EAllowShrinking::No);
 	CurBlock = 0;
 	CurBlockUsed = 0;
 	if (Blocks.Num() == 0)
 	{
-		Blocks.Add(new BlockType());
+		AddAllocatedBlock();
 	}
 }
 
-template <class Type>
-void TDynamicVector<Type>::Fill(const Type& Value)
+template <typename Type, int32 BlockSize>
+void TDynamicVector<Type, BlockSize>::Fill(const Type& Value)
 {
-	size_t Count = Blocks.Num();
-	for (unsigned int i = 0; i < Count; ++i)
+	for (uint32 BlockIndex = 0, NumBlocks = Blocks.Num(); BlockIndex < NumBlocks; ++BlockIndex)
 	{
-		for (uint32 ElementIndex = 0; ElementIndex < BlockSize; ++ElementIndex)
+		const uint32 NumElementsInBlock = BlockIndex < NumBlocks - 1 ? BlockSize : GetLength() - BlockSize * (NumBlocks - 1);
+		for (uint32 IndexInBlock = 0; IndexInBlock < NumElementsInBlock; ++IndexInBlock)
 		{
-			Blocks[i][ElementIndex] = Value;
+			GetElement(BlockIndex, IndexInBlock) = Value;
 		}
 	}
 }
 
-template <class Type>
-void TDynamicVector<Type>::Resize(unsigned int Count)
+template <typename Type, int32 BlockSize>
+void TDynamicVector<Type, BlockSize>::Resize(unsigned int Count)
 {
 	if (GetLength() == Count)
 	{
@@ -780,33 +612,33 @@ void TDynamicVector<Type>::Resize(unsigned int Count)
 	// Allocate needed additional blocks.
 	while (NumBlocksCurrent < NumBlocksNeeded)
 	{
-		Blocks.Add(new BlockType());
+		AddAllocatedBlock();
 		++NumBlocksCurrent;
 	}
 
 	// Remove unneeded blocks.
 	if (NumBlocksCurrent > NumBlocksNeeded)
 	{
-		Blocks.Truncate(NumBlocksNeeded, EAllowShrinking::No);
+		TruncateBlocks(NumBlocksNeeded, EAllowShrinking::No);
 	}
 
 	// Set current block.
 	SetCurBlock(Count);
 }
 
-template <class Type>
-void TDynamicVector<Type>::Resize(unsigned int Count, const Type& InitValue)
+template <typename Type, int32 BlockSize>
+void TDynamicVector<Type, BlockSize>::Resize(unsigned int Count, const Type& InitValue)
 {
 	size_t nCurSize = GetLength();
 	Resize(Count);
 	for (unsigned int Index = (unsigned int)nCurSize; Index < Count; ++Index)
 	{
-		this->operator[](Index) = InitValue;
+		(*this)[Index] = InitValue;
 	}
 }
 
-template <class Type>
-bool TDynamicVector<Type>::SetMinimumSize(unsigned int Count, const Type& InitValue)
+template <typename Type, int32 BlockSize>
+bool TDynamicVector<Type, BlockSize>::SetMinimumSize(unsigned int Count, const Type& InitValue)
 {
 	size_t nCurSize = GetLength();
 	if (Count <= nCurSize)
@@ -816,42 +648,68 @@ bool TDynamicVector<Type>::SetMinimumSize(unsigned int Count, const Type& InitVa
 	Resize(Count);
 	for (unsigned int Index = (unsigned int)nCurSize; Index < Count; ++Index)
 	{
-		this->operator[](Index) = InitValue;
+		(*this)[Index] = InitValue;
 	}
 	return true;
 }
 
-template <class Type>
-void TDynamicVector<Type>::Add(const Type& Value)
+template <typename Type, int32 BlockSize>
+void TDynamicVector<Type, BlockSize>::Add(const Type& Data)
 {
-	checkSlow(size_t(MaxSize) >= GetLength() + 1)
+	checkSlow(size_t(MAX_uint32) >= GetLength() + 1)
 	if (CurBlockUsed == BlockSize)
 	{
 		if (CurBlock == Blocks.Num() - 1)
 		{
-			Blocks.Add(new BlockType());
+			AddAllocatedBlock();
 		}
-		CurBlock++;
+		++CurBlock;
 		CurBlockUsed = 0;
 	}
-	Blocks[CurBlock][CurBlockUsed] = Value;
-	CurBlockUsed++;
+	GetElement(CurBlock, CurBlockUsed) = Data;
+	++CurBlockUsed;
 }
 
-
-template <class Type>
-void TDynamicVector<Type>::Add(const TDynamicVector<Type>& AddData)
+template <typename Type, int32 BlockSize>
+template <int32 BlockSizeData>
+void TDynamicVector<Type, BlockSize>::Add(const TDynamicVector<Type, BlockSizeData>& Data)
 {
 	// @todo it could be more efficient to use memcopies here...
-	size_t nSize = AddData.Num();
-	for (unsigned int k = 0; k < (unsigned int) nSize; ++k)
+	const uint32 Offset = Num();
+	const uint32 DataNum = static_cast<uint32>(Data.Num());
+	SetNum(Offset + DataNum);
+	for (uint32 DataIndex = 0; DataIndex < DataNum; ++DataIndex)
 	{
-		Add(AddData[k]);
+		(*this)[Offset + DataIndex] = Data[DataIndex];
 	}
 }
 
-template <class Type>
-void TDynamicVector<Type>::PopBack()
+template <typename Type, int32 BlockSize>
+void TDynamicVector<Type, BlockSize>::Add(const TArray<Type>& Data)
+{
+	const uint32 Offset = Num();
+	const uint32 DataNum = static_cast<uint32>(Data.Num());
+	SetNum(Offset + DataNum);
+	for (uint32 DataIndex = 0; DataIndex < DataNum; ++DataIndex)
+	{
+		(*this)[Offset + DataIndex] = Data[DataIndex];
+	}
+}
+
+template <typename Type, int32 BlockSize>
+void TDynamicVector<Type, BlockSize>::Add(TArrayView<const Type> Data)
+{
+	const uint32 Offset = Num();
+	const uint32 DataNum = static_cast<uint32>(Data.Num());
+	SetNum(Offset + DataNum);
+	for (uint32 DataIndex = 0; DataIndex < DataNum; ++DataIndex)
+	{
+		(*this)[Offset + DataIndex] = Data[DataIndex];
+	}
+}
+
+template <typename Type, int32 BlockSize>
+void TDynamicVector<Type, BlockSize>::PopBack()
 {
 	if (CurBlockUsed > 0)
 	{
@@ -861,12 +719,11 @@ void TDynamicVector<Type>::PopBack()
 	{
 		CurBlock--;
 		CurBlockUsed = BlockSize;
-		// remove block ??
 	}
 }
 
-template <class Type>
-Type& TDynamicVector<Type>::ElementAt(unsigned int Index, Type InitialValue)
+template <typename Type, int32 BlockSize>
+Type& TDynamicVector<Type, BlockSize>::ElementAt(unsigned int Index, Type InitialValue)
 {
 	size_t s = GetLength();
 	if (Index == s)
@@ -878,59 +735,226 @@ Type& TDynamicVector<Type>::ElementAt(unsigned int Index, Type InitialValue)
 		Resize(Index);
 		Add(InitialValue);
 	}
-	return this->operator[](Index);
+	return (*this)[Index];
 }
 
-template <class Type>
-void TDynamicVector<Type>::InsertAt(const Type& AddData, unsigned int Index)
+template <typename Type, int32 BlockSize>
+void TDynamicVector<Type, BlockSize>::InsertAt(const Type& Data, unsigned int Index)
 {
 	size_t s = GetLength();
 	if (Index == s)
 	{
-		Add(AddData);
+		Add(Data);
 	}
 	else if (Index > s)
 	{
 		Resize(Index);
-		Add(AddData);
+		Add(Data);
 	}
 	else
 	{
-		(*this)[Index] = AddData;
+		(*this)[Index] = Data;
 	}
 }
 
-template <class Type>
-void TDynamicVector<Type>::InsertAt(const Type& AddData, unsigned int Index, const Type& InitValue)
+template <typename Type, int32 BlockSize>
+void TDynamicVector<Type, BlockSize>::InsertAt(const Type& AddData, unsigned int Index, const Type& InitValue)
 {
 	size_t nCurSize = GetLength();
 	InsertAt(AddData, Index);
 	// initialize all new values up to (but not including) the inserted index
 	for (unsigned int i = (unsigned int)nCurSize; i < Index; ++i)
 	{
-		this->operator[](i) = InitValue;
+		(*this)[i] = InitValue;
 	}
 }
 
-template <typename Type>
+template <typename Type, int32 BlockSize>
 template <typename Func>
-void TDynamicVector<Type>::Apply(const Func& applyF)
+void TDynamicVector<Type, BlockSize>::Apply(const Func& ApplyFunc)
 {
-	for (int bi = 0; bi < CurBlock; ++bi)
+	for (int32 BlockIndex = 0; BlockIndex <= CurBlock; ++BlockIndex)
 	{
-		auto block = Blocks[bi];
-		for (int k = 0; k < BlockSize; ++k)
+		TBlock* Block = Blocks[BlockIndex];
+		const int32 NumElements = BlockIndex < CurBlock ? BlockSize : CurBlockUsed;
+		for (int32 ElementIndex = 0; ElementIndex < NumElements; ++ElementIndex)
+		{
+			ApplyFunc(Block[ElementIndex]);
+		}
+	}
+}
+
+template <typename Type, int32 BlockSize>
+template <bool bForceBulkSerialization, bool bUseCompression>
+void TDynamicVector<Type, BlockSize>::Serialize(FArchive& Ar)
+{
+	Ar.UsingCustomVersion(FUE5MainStreamObjectVersion::GUID);
+	if (Ar.IsLoading() && Ar.CustomVer(FUE5MainStreamObjectVersion::GUID) < FUE5MainStreamObjectVersion::DynamicMeshCompactedSerialization)
+	{
+		// In this version the serialization was done with a fixed block size of 512, and blocks were serialized in their entirety even if they were not
+		// fully occupied, i.e. the last block might have had garbage in it.
+		// To load this data, we first serialize all legacy blocks into a temporary buffer, and then copy out all valid elements one by one. While this
+		// solution is making an additional copies, it is simple and robust.
+
+		constexpr int32 LegacyBlockSize = 512;
+
+		uint32 LegacyCurBlock = 0;
+		uint32 LegacyCurBlockUsed = 0;
+		int32 BlockNum = 0;
+		Ar << LegacyCurBlock;
+		Ar << LegacyCurBlockUsed;
+		Ar << BlockNum;
+
+		// Bulk serialization for a number of double types was enabled as part of the transition to Large World Coordinates.
+		// If the currently stored type is one of these types, and the archive is from before bulk serialization for these types was enabled,
+		// we need to still use per element serialization for legacy data.
+		constexpr bool bIsLWCBulkSerializedDoubleType =
+			std::is_same_v<Type, FVector2d> ||
+			std::is_same_v<Type, FVector3d> ||
+			std::is_same_v<Type, FVector4d> ||
+			std::is_same_v<Type, FQuat4d> ||
+			std::is_same_v<Type, FTransform3d>;
+		const bool bUseBulkSerialization = TCanBulkSerialize<Type>::Value && !(bIsLWCBulkSerializedDoubleType && Ar.UEVer() <
+			EUnrealEngineObjectUE5Version::LARGE_WORLD_COORDINATES);
+
+		// Lambda for serializing a block either via bulk serializing the contained data or via serializing the block container itself.
+		// Note that the static_cast<> was necessary to resolve compiler errors when using MSVC.
+		const auto SerializeElements = bUseBulkSerialization
+			                               ? static_cast<void(*)(FArchive&, Type*)>([](FArchive& Archive, Type* BlockElements)
+			                               {
+				                               Archive.Serialize(BlockElements, static_cast<int64>(LegacyBlockSize) * sizeof(Type));
+			                               })
+			                               : static_cast<void(*)(FArchive&, Type*)>([](FArchive& Archive, Type* BlockElements)
+			                               {
+				                               for (int32 Index = 0; Index < LegacyBlockSize; ++Index)
+				                               {
+					                               Archive << BlockElements[Index];
+				                               }
+			                               });
+
+		// Serialize all blocks into a temporary buffer.
+		TArray<Type> TempElementBuffer;
+		TempElementBuffer.SetNum(BlockNum * LegacyBlockSize);
+		Type* TempElementBufferPtr = TempElementBuffer.GetData();
+		for (int32 BlockIndex = 0; BlockIndex < BlockNum; ++BlockIndex)
+		{
+			SerializeElements(Ar, TempElementBufferPtr);
+			TempElementBufferPtr += LegacyBlockSize;
+		}
+
+		// Add all valid elements from the temporary buffer into the vector.
+		const uint32 ElementsNum = LegacyCurBlock * LegacyBlockSize + LegacyCurBlockUsed;
+		Empty(ElementsNum / LegacyBlockSize + (ElementsNum % LegacyBlockSize != 0));
+		CurBlock = 0;
+		CurBlockUsed = 0;
+		AddAllocatedBlock();
+		for (uint32 ElementIndex = 0; ElementIndex < ElementsNum; ElementIndex++)
 		{
-			applyF(block[k], k);
+			Add(TempElementBuffer[ElementIndex]);
 		}
 	}
-	auto lastblock = Blocks[CurBlock];
-	for (int k = 0; k < CurBlockUsed; ++k)
+	else
 	{
-		applyF(lastblock[k], k);
+		uint32 SerializeNum = Num();
+		const SIZE_T CountBytes =  sizeof(uint32) + SerializeNum * sizeof(Type);
+		Ar.CountBytes(CountBytes, CountBytes);
+		Ar << SerializeNum;
+		if (SerializeNum == 0 && Ar.IsLoading())
+		{
+			Clear();
+		}
+		else if (SerializeNum > 0)
+		{
+			SetCurBlock(SerializeNum);
+
+			constexpr bool bUseBulkSerialization = bForceBulkSerialization || TCanBulkSerialize<Type>::Value || sizeof(Type) == 1;
+			static_assert(!bUseCompression || bUseBulkSerialization, "Compression only available when using bulk serialization");
+
+			// Serialize compression flag, which adds flexibility when de-serializing existing data even if some implementation details change.  
+			bool bUseCompressionForBulkSerialization = bUseBulkSerialization && bUseCompression;
+			Ar << bUseCompressionForBulkSerialization;
+
+			// Determine number of blocks.
+			const bool bNumIsNotMultipleOfBlockSize = SerializeNum % BlockSize != 0;
+			const uint32 NumBlocks = SerializeNum / BlockSize + bNumIsNotMultipleOfBlockSize;
+
+			if (bUseCompressionForBulkSerialization)
+			{
+				// When using compression, copy everything to/from a big single buffer and serialize the big buffer.
+				// This results in better compression ratios while at the same time accelerating compression. 
+
+				TArray<Type> Buffer;
+				Buffer.SetNumUninitialized(SerializeNum);
+				Type* BufferPtr = Buffer.GetData();
+				SIZE_T NumCopyRemaining = SerializeNum;
+
+				if (!Ar.IsLoading())
+				{
+					for (uint32 Index = 0; Index < NumBlocks; ++Index)
+					{
+						const SIZE_T NumCopy = FMath::Min<SIZE_T>(NumCopyRemaining, static_cast<SIZE_T>(BlockSize));
+						FMemory::Memcpy(BufferPtr, &Blocks[Index]->Elements[0], NumCopy * sizeof(Type));
+						BufferPtr += BlockSize;
+						NumCopyRemaining -= BlockSize;
+					}
+				}
+
+				Ar.SerializeCompressedNew(Buffer.GetData(), SerializeNum * sizeof(Type), NAME_Oodle, NAME_Oodle, COMPRESS_NoFlags, false, nullptr);
+
+				if (Ar.IsLoading())
+				{
+					Empty(NumBlocks);
+					for (uint32 Index = 0; Index < NumBlocks; ++Index)
+					{
+						TBlock *const NewBlock = new TBlock;
+						const SIZE_T NumCopy = FMath::Min<SIZE_T>(NumCopyRemaining, static_cast<SIZE_T>(BlockSize));
+						FMemory::Memcpy(&NewBlock->Elements[0], BufferPtr, NumCopy * sizeof(Type));
+						Blocks.Add(NewBlock);
+						BufferPtr += BlockSize;
+						NumCopyRemaining -= BlockSize;
+					}
+				}
+			}
+			else
+			{
+				const auto SerializeBlock = [&Ar, bUseBulkSerialization](TBlock* Block, uint32 NumElements)
+				{
+					if (bUseBulkSerialization)
+					{
+						Ar.Serialize(&Block->Elements[0], NumElements * sizeof(Type));
+					}
+					else
+					{
+						for (uint32 Index = 0; Index < NumElements; ++Index)
+						{
+							Ar << Block->Elements[Index];
+						}
+					}
+				};
+
+				if (Ar.IsLoading())
+				{
+					Empty(NumBlocks);
+					for (uint32 BlockIndex = 0; BlockIndex < NumBlocks; ++BlockIndex)
+					{
+						TBlock *const NewBlock = new TBlock;
+						SerializeBlock(NewBlock, FMath::Min<uint32>(SerializeNum, BlockSize));
+						Blocks.Add(NewBlock);
+						SerializeNum -= BlockSize;
+					}
+				}
+				else
+				{
+					for (uint32 BlockIndex = 0; BlockIndex < NumBlocks; ++BlockIndex)
+					{
+						SerializeBlock(Blocks[BlockIndex], FMath::Min<uint32>(SerializeNum, BlockSize));
+						SerializeNum -= BlockSize;
+					}
+				}
+			}
+		}
 	}
 }
 
-
 } // end namespace UE::Geometry
 } // end namespace UE
