Add option to run KMeans' distance-to-cluster loop in parallel.

diff --git a/Engine/Source/Runtime/GeometryCore/Public/Clustering/KMeans.h b/Engine/Source/Runtime/GeometryCore/Public/Clustering/KMeans.h
--- a/Engine/Source/Runtime/GeometryCore/Public/Clustering/KMeans.h
+++ b/Engine/Source/Runtime/GeometryCore/Public/Clustering/KMeans.h
@@ -7,6 +7,8 @@
 #include "VectorTypes.h"
 #include "MatrixTypes.h"
 #include "Spatial/PriorityOrderPoints.h"
+#include "Containers/StaticArray.h"
+#include "Async/ParallelFor.h"
 
 namespace UE
 {
@@ -15,6 +17,108 @@ namespace Geometry
 
 using namespace UE::Math;
 
+// Not all platforms support c++20 or have their standard library containing floating atomic operations, 
+// which is required for using std::atomic::+= opertor on floating types
+#define PLATFORM_SUPPORT_FLOATING_INTERLOCKED_ADD (PLATFORM_WINDOWS && (__cplusplus >= 202002L))
+
+// Abstract array of positions for running the KMeans algorithm in serial or parallel
+template<typename TVectorType, bool bUseAtomicType>
+struct TClusterCenterArray
+{
+private:
+	// Atomic values for parallel algorithm
+	TStaticArray<TArray<std::atomic<typename TVectorType::FReal>>, TVectorType::NumComponents> AtomicValues;
+
+	// Regular values for serial algorithm
+	TArray<TVectorType> Values;
+public:
+	void SetNumZeroed(uint32 Num)
+	{ 
+		#if PLATFORM_SUPPORT_FLOATING_INTERLOCKED_ADD
+		if (bUseAtomicType)
+		{
+			for (uint32 ComponentIndex = 0; ComponentIndex < TVectorType::NumComponents; ++ComponentIndex)
+			{
+				AtomicValues[ComponentIndex].SetNumZeroed(Num);
+			}
+		}
+		else
+		#endif
+		{
+			Values.SetNumZeroed(Num);
+		}
+	}
+
+	void Assign(int32 Index, const TVectorType& In)
+	{ 
+		#if PLATFORM_SUPPORT_FLOATING_INTERLOCKED_ADD
+		if (bUseAtomicType)
+		{
+			for (uint32 ComponentIndex = 0; ComponentIndex < TVectorType::NumComponents; ++ComponentIndex)
+			{
+				AtomicValues[ComponentIndex][Index] = In[ComponentIndex];
+			}
+		}
+		else
+		#endif
+		{
+			Values[Index] = In;
+		}
+	}
+
+	TVectorType DivideBy(int32 Index, int32 InDivider)
+	{ 
+		#if PLATFORM_SUPPORT_FLOATING_INTERLOCKED_ADD
+		if (bUseAtomicType)
+		{
+			TVectorType Out;
+			for (uint32 ComponentIndex = 0; ComponentIndex < TVectorType::NumComponents; ++ComponentIndex)
+			{
+				Out[ComponentIndex] = AtomicValues[ComponentIndex][Index] / InDivider;
+			}
+			return Out;
+		}
+		else
+		#endif
+		{
+			return Values[Index] / InDivider;
+		}
+	}
+
+	void Accumulate(int32 Index, const TVectorType& In)
+	{ 
+		#if PLATFORM_SUPPORT_FLOATING_INTERLOCKED_ADD
+		if (bUseAtomicType)
+		{
+			for (uint32 ComponentIndex = 0; ComponentIndex < TVectorType::NumComponents; ++ComponentIndex)
+			{
+				AtomicValues[ComponentIndex][Index].fetch_add(In[ComponentIndex]);
+			}
+		}
+		else
+		#endif
+		{
+			Values[Index] += In;
+		}
+	}
+
+	void RemoveAtSwap(size_t Index, EAllowShrinking AllowShrinking)
+	{
+		#if PLATFORM_SUPPORT_FLOATING_INTERLOCKED_ADD
+		if (bUseAtomicType)
+		{
+			for (uint32 ComponentIndex = 0; ComponentIndex < TVectorType::NumComponents; ++ComponentIndex)
+			{
+				AtomicValues[ComponentIndex].RemoveAtSwap(Index, AllowShrinking);
+			}
+		}
+		else
+		#endif
+		{
+			Values.RemoveAtSwap(Index, AllowShrinking);
+		}
+	}
+};
 
 struct FClusterKMeans
 {
@@ -32,7 +136,7 @@ struct FClusterKMeans
 	TArray<int32> ClusterIDs;
 
 	// Number of points in each cluster
-	TArray<int32> ClusterSizes;
+	TArray<std::atomic<int32>> ClusterSizes;
 
 	/**
 	 * Compute the K-Means clustering of FVector points
@@ -41,14 +145,17 @@ struct FClusterKMeans
 	 * @param NumClusters		Target number of clusters to create, if InitialCenters is not provided
 	 * @param InitialCenters	If non-empty, these positions will be used to initialize the cluster locations
 	 * @param OutClusterCenters	If non-null, will be filled with the cluster centers
+	 * @param bRunParallel		Runs a parallel algorithm (faster) if possible. The output might be non-deterministic when running in parallel
 	 * @return number of clusters found
 	 */
-	template<typename TVectorType = FVector>
+	template<typename TVectorType = FVector, bool bRunParallel=false>
 	int32 ComputeClusters(
 		TArrayView<const TVectorType> PointsToCluster, int32 NumClusters,
 		TArrayView<const TVectorType> InitialCenters = TArrayView<const TVectorType>(), 
 		TArray<TVectorType>* OutClusterCenters = nullptr)
 	{
+		constexpr bool bRunParallelIfCompatible = PLATFORM_SUPPORT_FLOATING_INTERLOCKED_ADD && bRunParallel && std::atomic<typename TVectorType::FReal>::is_always_lock_free;
+
 		if (PointsToCluster.IsEmpty())
 		{
 			// nothing to cluster
@@ -108,7 +215,7 @@ struct FClusterKMeans
 		}
 
 		
-		TArray<TVectorType> NextCenters;
+		TClusterCenterArray<TVectorType, bRunParallelIfCompatible> NextCenters;
 		NextCenters.SetNumZeroed(NumClusters);
 
 		TArray<int32> ClusterIDRemap; // array to use if we need to remap cluster IDs due to an empty cluster
@@ -117,6 +224,48 @@ struct FClusterKMeans
 		for (int32 Iterations = 0; Iterations < UseMaxIterations; ++Iterations)
 		{
 			bool bClustersChanged = Iterations == 0; // clusters always change on first iteration; otherwise we must check
+			if (bRunParallelIfCompatible)
+			{
+				ParallelFor(PointsToCluster.Num(), 
+				[
+					this,
+					Iterations,
+					&PointsToCluster,
+					&NextCenters,
+					&UseCenters,
+					&bClustersChanged
+				] (uint32 PointIdx)
+				{
+					const TVectorType& Point = PointsToCluster[PointIdx];
+					double ClosestDistSq = (double)UE::Geometry::DistanceSquared(Point, (*UseCenters)[0]);
+					int32 ClosestCenter = 0;
+					for (int32 CenterIdx = 1; CenterIdx < UseCenters->Num(); ++CenterIdx)
+					{
+						double DistSq = UE::Geometry::DistanceSquared(PointsToCluster[PointIdx], (*UseCenters)[CenterIdx]);
+						if (DistSq < ClosestDistSq)
+						{
+							ClosestDistSq = DistSq;
+							ClosestCenter = CenterIdx;
+						}
+					}
+
+					ClusterSizes[ClosestCenter]++;
+					if (Iterations > 0)
+					{
+						int32 OldClusterID = ClusterIDs[PointIdx];
+						ClusterSizes[OldClusterID]--;
+						if (OldClusterID != ClosestCenter)
+						{
+							bClustersChanged = true;
+						}
+					}
+				
+					ClusterIDs[PointIdx] = ClosestCenter;
+					NextCenters.Accumulate(ClosestCenter, Point);
+				});
+			}
+			else
+			{
 			for (int32 PointIdx = 0; PointIdx < PointsToCluster.Num(); ++PointIdx)
 			{
 				const TVectorType& Point = PointsToCluster[PointIdx];
@@ -144,8 +293,9 @@ struct FClusterKMeans
 				}
 				
 				ClusterIDs[PointIdx] = ClosestCenter;
-				NextCenters[ClosestCenter] += Point;
+				NextCenters.Accumulate(ClosestCenter, Point);
 			}
+			} // bParallel
 
 			// Stop iterating if clusters are unchanged
 			if (!bClustersChanged)
@@ -159,8 +309,8 @@ struct FClusterKMeans
 			{
 				if (ClusterSizes[ClusterIdx] > 0)
 				{
-					(*UseCenters)[ClusterIdx] = NextCenters[ClusterIdx] / ClusterSizes[ClusterIdx];
-					NextCenters[ClusterIdx] = TVectorType::ZeroVector;
+					(*UseCenters)[ClusterIdx] = NextCenters.DivideBy(ClusterIdx, int32(ClusterSizes[ClusterIdx]));
+					NextCenters.Assign(ClusterIdx, TVectorType::ZeroVector);
 				}
 				else
 				{
