Add GetBytesCount to DynamicMesh (and attributes/overlays) to track combined memory usage of a dynamic mesh.

diff --git a/Engine/Source/Runtime/GeometryCore/Private/DynamicMesh/DynamicMesh3.cpp b/Engine/Source/Runtime/GeometryCore/Private/DynamicMesh/DynamicMesh3.cpp
--- a/Engine/Source/Runtime/GeometryCore/Private/DynamicMesh/DynamicMesh3.cpp
+++ b/Engine/Source/Runtime/GeometryCore/Private/DynamicMesh/DynamicMesh3.cpp
@@ -693,6 +693,16 @@ FString FDynamicMesh3::MeshInfoString() const
 	return VtxString + "\n" + TriString + "\n" + EdgeString + "\n" + AttribString + "\n" + InfoString;
 }
 
+SIZE_T FDynamicMesh3::GetByteCount() const
+{
+	SIZE_T ByteCount = VertexRefCounts.GetByteCount() + VertexEdgeLists.GetByteCount() + TriangleRefCounts.GetByteCount() + EdgeRefCounts.GetByteCount();
+	if (HasAttributes())
+	{
+		ByteCount += Attributes()->GetByteCount();
+	}
+	return ByteCount;
+}
+
 bool FDynamicMesh3::IsSameAs(const FDynamicMesh3& m2, const FSameAsOptions& Options) const
 {
 	return IsSameAs_Helper(m2, Options, nullptr);
diff --git a/Engine/Source/Runtime/GeometryCore/Private/DynamicMesh/DynamicMeshAttributeSet.cpp b/Engine/Source/Runtime/GeometryCore/Private/DynamicMesh/DynamicMeshAttributeSet.cpp
--- a/Engine/Source/Runtime/GeometryCore/Private/DynamicMesh/DynamicMeshAttributeSet.cpp
+++ b/Engine/Source/Runtime/GeometryCore/Private/DynamicMesh/DynamicMeshAttributeSet.cpp
@@ -1482,6 +1482,77 @@ void FDynamicMeshAttributeSet::Serialize(FArchive& Ar, const FCompactMaps* Compa
 	//Ar << GenericAttributes; // TODO
 }
 
+SIZE_T FDynamicMeshAttributeSet::GetByteCount() const
+{
+	SIZE_T ByteCount = 0;
+
+	for (const FDynamicMeshUVOverlay& UVLayer : UVLayers)
+	{
+		ByteCount += UVLayer.GetByteCount();
+	}
+
+	for (const FDynamicMeshNormalOverlay& NormalLayer : NormalLayers)
+	{
+		ByteCount += NormalLayer.GetByteCount();
+	}
+
+	if (ColorLayer)
+	{
+		ByteCount += ColorLayer->GetByteCount();
+	}
+
+	if (MaterialIDAttrib)
+	{
+		ByteCount += MaterialIDAttrib->GetByteCount();
+	}
+
+	for (const FDynamicMeshPolygroupAttribute& PolygroupLayer : PolygroupLayers)
+	{
+		ByteCount += PolygroupLayer.GetByteCount();
+	}
+
+	for (const FDynamicMeshWeightAttribute& WeightLayer : WeightLayers)
+	{
+		ByteCount += WeightLayer.GetByteCount();
+	}
+
+	if (BoneNameAttrib)
+	{
+		ByteCount += BoneNameAttrib->GetByteCount();
+	}
+	if (BoneParentIndexAttrib)
+	{
+		ByteCount += BoneParentIndexAttrib->GetByteCount();
+	}
+
+	if (BoneColorAttrib)
+	{
+		ByteCount += BoneColorAttrib->GetByteCount();
+	}
+	
+	if (BonePoseAttrib)
+	{
+		ByteCount += BonePoseAttrib->GetByteCount();
+	}
+	
+	if (BoneColorAttrib)
+	{
+		ByteCount += BoneColorAttrib->GetByteCount();
+	}
+
+	for (int SculptLayerIndex = 0; SculptLayerIndex < NumSculptLayers(); ++SculptLayerIndex)
+	{
+		ByteCount += SculptLayers.Layers[SculptLayerIndex].GetByteCount();
+	}
+	
+	for (const TPair<FName, TUniquePtr<FDynamicMeshAttributeBase>>& GenericAttribute : GenericAttributes)
+	{
+		ByteCount += GenericAttribute.Value->GetByteCount();
+	}
+
+	return ByteCount;
+}
+
 
 void FDynamicMeshAttributeSet::EnableSculptLayers(int32 MinLayers)
 {
diff --git a/Engine/Source/Runtime/GeometryCore/Private/DynamicMesh/NonManifoldMappingSupport.cpp b/Engine/Source/Runtime/GeometryCore/Private/DynamicMesh/NonManifoldMappingSupport.cpp
--- a/Engine/Source/Runtime/GeometryCore/Private/DynamicMesh/NonManifoldMappingSupport.cpp
+++ b/Engine/Source/Runtime/GeometryCore/Private/DynamicMesh/NonManifoldMappingSupport.cpp
@@ -218,6 +218,11 @@ public:
 	{
 		bValid = false;
 	}
+
+	virtual SIZE_T GetByteCount() const override
+	{
+		return SourceVertexIds.GetByteCount();
+	}
 };
 
 void FNonManifoldSourceVertexIdAttributeChange::SaveInitialVertex(const FDynamicMeshAttributeBase* Attribute, int VertexID)
diff --git a/Engine/Source/Runtime/GeometryCore/Public/DynamicMesh/DynamicAttribute.h b/Engine/Source/Runtime/GeometryCore/Public/DynamicMesh/DynamicAttribute.h
--- a/Engine/Source/Runtime/GeometryCore/Public/DynamicMesh/DynamicAttribute.h
+++ b/Engine/Source/Runtime/GeometryCore/Public/DynamicMesh/DynamicAttribute.h
@@ -227,6 +227,8 @@ public:
 		}
 	}
 
+	virtual SIZE_T GetByteCount() const { return 0; }
+
 protected:
 
 	/** 
diff --git a/Engine/Source/Runtime/GeometryCore/Public/DynamicMesh/DynamicBoneAttribute.h b/Engine/Source/Runtime/GeometryCore/Public/DynamicMesh/DynamicBoneAttribute.h
--- a/Engine/Source/Runtime/GeometryCore/Public/DynamicMesh/DynamicBoneAttribute.h
+++ b/Engine/Source/Runtime/GeometryCore/Public/DynamicMesh/DynamicBoneAttribute.h
@@ -91,6 +91,10 @@ public:
 	/** Add value at the end of the value array. */
 	void Append(const AttribValueType& InValue);
 
+	SIZE_T GetByteCount() const
+	{
+		return AttribValues.GetAllocatedSize();
+	}
 public:
 
     //
diff --git a/Engine/Source/Runtime/GeometryCore/Public/DynamicMesh/DynamicMesh3.h b/Engine/Source/Runtime/GeometryCore/Public/DynamicMesh/DynamicMesh3.h
--- a/Engine/Source/Runtime/GeometryCore/Public/DynamicMesh/DynamicMesh3.h
+++ b/Engine/Source/Runtime/GeometryCore/Public/DynamicMesh/DynamicMesh3.h
@@ -1477,6 +1477,8 @@ public:
 	 */
 	GEOMETRYCORE_API virtual FString MeshInfoString() const;
 
+	GEOMETRYCORE_API SIZE_T GetByteCount() const;
+
 	/**
 	 * Options for the IsSameAs check
 	 */
diff --git a/Engine/Source/Runtime/GeometryCore/Public/DynamicMesh/DynamicMeshAttributeSet.h b/Engine/Source/Runtime/GeometryCore/Public/DynamicMesh/DynamicMeshAttributeSet.h
--- a/Engine/Source/Runtime/GeometryCore/Public/DynamicMesh/DynamicMeshAttributeSet.h
+++ b/Engine/Source/Runtime/GeometryCore/Public/DynamicMesh/DynamicMeshAttributeSet.h
@@ -583,6 +583,8 @@ public:
 	*/
 	GEOMETRYCORE_API void Serialize(FArchive& Ar, const FCompactMaps* CompactMaps, bool bUseCompression);
 
+	GEOMETRYCORE_API SIZE_T GetByteCount() const;
+
 protected:
 	/** Parent mesh of this attribute set */
 	FDynamicMesh3* ParentMesh;
diff --git a/Engine/Source/Runtime/GeometryCore/Public/DynamicMesh/DynamicMeshOverlay.h b/Engine/Source/Runtime/GeometryCore/Public/DynamicMesh/DynamicMeshOverlay.h
--- a/Engine/Source/Runtime/GeometryCore/Public/DynamicMesh/DynamicMeshOverlay.h
+++ b/Engine/Source/Runtime/GeometryCore/Public/DynamicMesh/DynamicMeshOverlay.h
@@ -611,6 +611,11 @@ public:
 	 */
 	void Serialize(FArchive& Ar, const FCompactMaps* CompactMaps, bool bUseCompression);
 
+	SIZE_T GetByteCount() const
+	{
+		return ElementsRefCounts.GetByteCount() + Elements.GetByteCount() + ParentVertices.GetByteCount() + ElementTriangles.GetByteCount();
+	}
+
 public:
 	/** Set a triangle's element indices to InvalidID */
 	void InitializeNewTriangle(int TriangleID);
diff --git a/Engine/Source/Runtime/GeometryCore/Public/DynamicMesh/DynamicMeshTriangleAttribute.h b/Engine/Source/Runtime/GeometryCore/Public/DynamicMesh/DynamicMeshTriangleAttribute.h
--- a/Engine/Source/Runtime/GeometryCore/Public/DynamicMesh/DynamicMeshTriangleAttribute.h
+++ b/Engine/Source/Runtime/GeometryCore/Public/DynamicMesh/DynamicMeshTriangleAttribute.h
@@ -517,6 +517,11 @@ public:
 			}
 		}
 	}
+
+	virtual SIZE_T GetByteCount() const
+	{
+		return AttribValues.GetByteCount();
+	}
 };
 
 
diff --git a/Engine/Source/Runtime/GeometryCore/Public/DynamicMesh/DynamicVertexAttribute.h b/Engine/Source/Runtime/GeometryCore/Public/DynamicMesh/DynamicVertexAttribute.h
--- a/Engine/Source/Runtime/GeometryCore/Public/DynamicMesh/DynamicVertexAttribute.h
+++ b/Engine/Source/Runtime/GeometryCore/Public/DynamicMesh/DynamicVertexAttribute.h
@@ -518,6 +518,11 @@ public:
 		}
 	}
 
+	virtual SIZE_T GetByteCount() const override
+	{
+		return AttribValues.GetByteCount();
+	}
+
 
 protected:
 
diff --git a/Engine/Source/Runtime/GeometryCore/Public/DynamicMesh/DynamicVertexSkinWeightsAttribute.h b/Engine/Source/Runtime/GeometryCore/Public/DynamicMesh/DynamicVertexSkinWeightsAttribute.h
--- a/Engine/Source/Runtime/GeometryCore/Public/DynamicMesh/DynamicVertexSkinWeightsAttribute.h
+++ b/Engine/Source/Runtime/GeometryCore/Public/DynamicMesh/DynamicVertexSkinWeightsAttribute.h
@@ -636,6 +636,11 @@ public:
 		}
 	}
 
+	virtual SIZE_T GetByteCount() const override
+	{
+		return VertexBoneWeights.GetByteCount();
+	}
+
 protected:
 
 	// interpolation functions; default implementation assumes your attributes can be interpolated as reals
diff --git a/Engine/Source/Runtime/GeometryCore/Public/Util/RefCountVector.h b/Engine/Source/Runtime/GeometryCore/Public/Util/RefCountVector.h
--- a/Engine/Source/Runtime/GeometryCore/Public/Util/RefCountVector.h
+++ b/Engine/Source/Runtime/GeometryCore/Public/Util/RefCountVector.h
@@ -501,6 +501,11 @@ public:
 			RefCounts.GetLength(), FreeIndices.GetLength(), (FreeIndices.GetByteCount() / 1024));
 	}
 
+	SIZE_T GetByteCount() const
+	{
+		return RefCounts.GetByteCount() + FreeIndices.GetByteCount();
+	}
+
 	/**
 	 * Serialization operator for FRefCountVector.
 	 *
diff --git a/Engine/Source/Runtime/GeometryCore/Public/Util/SmallListSet.h b/Engine/Source/Runtime/GeometryCore/Public/Util/SmallListSet.h
--- a/Engine/Source/Runtime/GeometryCore/Public/Util/SmallListSet.h
+++ b/Engine/Source/Runtime/GeometryCore/Public/Util/SmallListSet.h
@@ -680,6 +680,11 @@ public:
 			ListBlocks.GetLength(), (LinkedListElements.GetLength() * sizeof(int32) / 1024));
 	}
 
+	SIZE_T GetByteCount() const
+	{
+		return ListHeads.GetByteCount() + FreeBlocks.GetByteCount() + ListBlocks.GetByteCount() + LinkedListElements.GetByteCount();
+	}
+
 };
 
 
