- Move Splines in to Geometry Core
- Support for DeepCopy/Clone of splines and their attributes
- Expand parameter range when adding values with out of range parameters
- Clear attribute channels by name
- BSplines now use template parameter for Degree

diff --git a/Engine/Source/Runtime/GeometryCore/Public/Splines/AttributeChannel.h b/Engine/Source/Runtime/GeometryCore/Public/Splines/AttributeChannel.h
new file mode 100644
--- /dev/null
+++ b/Engine/Source/Runtime/GeometryCore/Public/Splines/AttributeChannel.h
@@ -0,0 +1,133 @@
+// Copyright Epic Games, Inc. All Rights Reserved.
+
+#pragma once
+
+#include "ParameterizedStorage.h"
+#include "BaseEvaluator.h"
+
+/** Base interface for attribute channels **/
+class IAttributeChannelBase
+{
+public:
+	virtual ~IAttributeChannelBase() = default;
+
+	/** Clears all values from this channel */
+    virtual void Clear() = 0;
+};
+
+/**
+ * Represents a single attribute channel with its own evaluation scheme
+ */
+template<typename AttrType>
+class FAttributeChannel : public IAttributeChannelBase
+{
+    FAttributeStorage<AttrType> Storage;
+    TUniquePtr<TBaseEvaluator<AttrType>> Evaluator;
+
+    // Optional overrides:
+    TFunction<TArray<const AttrType*>(TArrayView<const FStorageEntry<AttrType>>, float)> FindWindowOverride;
+    TFunction<AttrType(TArrayView<const AttrType* const>, float)> InterpolateOverride;
+
+public:
+    explicit FAttributeChannel(TUniquePtr<TBaseEvaluator<AttrType>> InEvaluator)
+        : Evaluator(MoveTemp(InEvaluator))
+    {
+    }
+
+	// Add copy constructor to properly clone the evaluator
+	FAttributeChannel(const FAttributeChannel& Other)
+		: Storage(Other.Storage)
+		, FindWindowOverride(Other.FindWindowOverride)
+		, InterpolateOverride(Other.InterpolateOverride)
+    {
+    	// Clone the evaluator - we need to implement Clone() in TBaseEvaluator
+    	if (Other.Evaluator)
+    	{
+    		Evaluator = Other.Evaluator->Clone();
+    	}
+    }
+
+	FAttributeChannel(FAttributeChannel&& Other) noexcept = default;
+
+	/**
+	 * Clears all values from this channel
+	 */
+	virtual void Clear() override
+	{
+		Storage.Clear();
+	}
+	
+    /**
+     * Adds a new value to the channel at the given parameter
+     * @param Value 
+     * @param Parameter
+     * @return Index of the added value
+     */
+    int32 AddValue(const AttrType& Value, float Parameter)
+    {
+        return Storage.AddValue(Value, Parameter);
+    }
+
+    /**
+     * Evaluates the channel at the given parameter in channel space
+     * @param Parameter - Parameter value in channel's local space
+     */
+    AttrType Evaluate(float Parameter) const
+    {
+        if (Storage.GetEntries().Num() == 0) return AttrType();
+        if (Storage.GetEntries().Num() == 1) return Storage.GetEntries()[0].Value;
+
+        // Get window using override or evaluator
+        TArray<const AttrType*> Window;
+        if (FindWindowOverride)
+        {
+            Window = FindWindowOverride(Storage.GetEntries(), Parameter);
+        }
+        else
+        {
+            Window = Evaluator->FindWindow(Storage.GetEntries(), Parameter);
+        }
+        if (Window.Num() == 0) return AttrType();
+        
+        if (InterpolateOverride)
+        {
+            return InterpolateOverride(Window, Parameter);
+        }
+        return Evaluator->InterpolateWindow(Window, Parameter);
+    }
+
+    TBaseEvaluator<AttrType>& GetEvaluator()
+    {
+        return *Evaluator;
+    }
+    
+    const TBaseEvaluator<AttrType>& GetEvaluator() const
+    {
+        return *Evaluator;
+    }
+    
+    // Access to underlying storage
+    const FAttributeStorage<AttrType>& GetStorage() const
+    {
+        return Storage;
+    }
+    
+
+    // Set overrides
+    void SetFindWindowOverride(TFunction<TArray<const AttrType*>(TArrayView<const FStorageEntry<AttrType>>, float)> Override)
+    {
+        FindWindowOverride = MoveTemp(Override);
+    }
+
+    void SetInterpolateOverride(TFunction<AttrType(TArrayView<const AttrType* const>, float)> Override)
+    {
+        InterpolateOverride = MoveTemp(Override);
+    }
+    
+    // Clear overrides
+    void ClearOverrides()
+    {
+        FindWindowOverride = nullptr;
+        InterpolateOverride = nullptr;
+    }
+};
diff --git a/Engine/Source/Runtime/GeometryCore/Public/Splines/BSpline.h b/Engine/Source/Runtime/GeometryCore/Public/Splines/BSpline.h
new file mode 100644
--- /dev/null
+++ b/Engine/Source/Runtime/GeometryCore/Public/Splines/BSpline.h
@@ -0,0 +1,640 @@
+// Copyright Epic Games, Inc. All Rights Reserved.
+
+#pragma once
+
+#include "Spline.h"
+#include "BaseEvaluator.h"
+
+
+enum class EBSplineKnotMode 
+{
+	/** Equal spacing between knots regardless of control point positions */
+	Uniform,        
+
+	/** Spacing based on actual distances between control points - good for uneven point distribution */
+	ChordLength,    
+
+	/** Uses square root of distances - helps prevent loops and cusps in the curve */
+	Centripetal,    
+
+	/** User-provided custom knot vector */
+	Custom          
+};
+
+// Forward declaration
+template<typename ScalarType, int32 Dimensions, int32 Degree>
+class TBSpline;
+
+template <typename ScalarType, int32 Dimensions, int32 Degree>
+class TBSplineEvaluator final : public TBaseEvaluator<typename TChooseVectorType<ScalarType, Dimensions>::Type>
+{
+	static_assert(Degree > 0, "BSpline degree must be positive");
+	
+public:
+	using FValue = typename TChooseVectorType<ScalarType, Dimensions>::Type;
+	using Base = TBaseEvaluator<FValue>;
+	
+private:
+	
+	/** Vector containing knot values that determine spline parameterization */
+	TArray<float> KnotVector;
+
+	/** Number of control points in the spline */
+	int32 NumValues = 0;
+
+	/** Method used to generate the knot vector */
+	EBSplineKnotMode KnotMode = EBSplineKnotMode::Uniform;
+
+	/** Whether to clamp the endpoints for interpolation of first/last control points */
+	bool bClampEnds = true;
+
+public:
+	TBSplineEvaluator()
+	{
+	}
+
+	virtual TArray<const FValue*> FindWindow(
+		TArrayView<const FStorageEntry<FValue>> Entries,
+		float Parameter) const override
+	{
+		TArray<const FValue*> Window;
+		if (Entries.Num() < Degree + 1)
+		{
+			return Window;
+		}
+
+		const float MappedParam = Base::MapToParentSpace(Parameter);
+		const int32 Span = FindKnotSpan(MappedParam);
+
+		// Get control points based on knot vector
+		for (int32 i = 0; i <= Degree; ++i)
+		{
+			const int32 Index = Span - Degree + i;
+			if (Index >= 0 && Index < Entries.Num())
+			{
+				Window.Add(&Entries[Index].Value);
+			}
+		}
+
+		return Window;
+	}
+
+	virtual FValue InterpolateWindow(TArrayView<const FValue* const> Window, float Parameter) const override
+	{
+		TArray<float> Basis;
+		const float MappedParam = Base::MapToParentSpace(Parameter);
+		ComputeBasisFunctions(FindKnotSpan(MappedParam), MappedParam, Basis);
+		return TSplineInterpolationPolicy<FValue>::InterpolateWithBasis(Window, Basis);
+	}
+
+	virtual int32 GetDegree() const override
+	{
+		return Degree;
+	}
+
+	/**
+	 * Sets a custom knot vector and switches to custom knot mode
+	 * @param NewKnots - Array of knot values. Must be non-decreasing and have correct length (n + p + 1)
+	 * @param InNumValues - Number of values (control points) in the spline
+	 * @return True if the knot vector was valid and set successfully
+	 */
+	bool SetCustomKnots(const TArray<float>& NewKnots, int32 InNumValues)
+	{
+		NumValues = InNumValues;
+		const int32 RequiredKnots = NumValues + Degree + 1;
+		if (NewKnots.Num() < RequiredKnots)
+		{
+			UE_LOG(LogSpline, Warning, TEXT("Invalid knot vector size. Need %d knots, got %d."), 
+				RequiredKnots, NewKnots.Num());
+			return false;
+		}
+
+		// Verify knot vector is non-decreasing
+		for (int32 i = 1; i < NewKnots.Num(); ++i)
+		{
+			if (NewKnots[i] < NewKnots[i - 1])
+			{
+				UE_LOG(LogSpline, Warning, TEXT("Invalid knot vector: knots must be non-decreasing."));
+				return false;
+			}
+		}
+
+		KnotVector = NewKnots;
+		KnotMode = EBSplineKnotMode::Custom;
+		
+		UE_LOG(LogSpline, Verbose, TEXT("Set Custom knot vector - Points: %d, Degree: %d, Knots: %d"),
+			   NumValues, Degree, KnotVector.Num());
+		for (int32 i = 0; i < KnotVector.Num(); ++i)
+		{
+			UE_LOG(LogSpline, Verbose, TEXT("  Knot[%d] = %f"), i, KnotVector[i]);
+		}
+		
+		return true;
+	}
+
+	TArray<float> GetKnotVector() const
+	{
+		return KnotVector;
+	}
+
+
+	/**
+	 * Updates the knot vector based on current control points and settings
+	 * 
+	 * @param Points - Array of control points and their parameters
+	 */
+	void UpdateKnotVector(const TArrayView<const FStorageEntry<FValue>>& Points)
+	{
+		NumValues = Points.Num();
+		if (NumValues < Degree + 1)
+		{
+			UE_LOG(LogSpline, Error, TEXT("Not enough points for B-spline: %d (need >= %d)"),
+				   NumValues, Degree + 1);
+			return;
+		}
+
+		const int32 KnotCount = NumValues + Degree + 1;
+		KnotVector.SetNum(KnotCount);
+
+		switch (KnotMode)
+		{
+		case EBSplineKnotMode::Uniform:
+			GenerateUniformKnots(KnotCount);
+			break;
+		case EBSplineKnotMode::ChordLength:
+			GenerateChordLengthKnots(KnotCount, Points);
+			break;
+		case EBSplineKnotMode::Centripetal:
+			GenerateCentripetalKnots(KnotCount, Points);
+			break;
+		case EBSplineKnotMode::Custom:
+			// Do nothing - use existing knot vector
+				break;
+		}
+	}
+	
+	/** Changes the method used to generate knot vectors */
+	void SetKnotGenerationMode(EBSplineKnotMode InMode)
+	{
+		KnotMode = InMode;
+	}
+
+	/** 
+	 * Sets whether the spline should be clamped at endpoints
+	 * When true, the curve will interpolate the first and last control points
+	 */
+	void SetClampedEnds(bool bInClampEnds)
+	{
+		bClampEnds = bInClampEnds;
+	}
+
+	bool IsClampedEnds() const
+	{
+		return bClampEnds;
+	}
+	
+	EBSplineKnotMode GetKnotMode() const
+	{
+		return KnotMode;
+	}
+
+	// Get the valid parameter range in knot space
+	void GetKnotRange(float& OutStart, float& OutEnd) const
+	{
+		if (KnotVector.Num() < 2)
+		{
+			OutStart = 0.0f;  
+			OutEnd = 1.0f;
+			return;
+		}
+
+		if (bClampEnds)
+		{
+			// For clamped B-splines, the valid range is [Degree, NumKnots - Degree - 1]
+			OutStart = KnotVector[Degree];
+			OutEnd = KnotVector[KnotVector.Num() - Degree - 1];
+		}
+		else
+		{
+			// For unclamped B-splines, the valid range is [FirstKnot, LastKnot]
+			OutStart = KnotVector[0];
+			OutEnd = KnotVector.Last();
+		}
+	}
+
+	// Convert from normalized [0,1] to knot space
+	float ToKnotSpace(float NormalizedParam) const
+	{
+		float Start, End;
+		GetKnotRange(Start, End);
+		return FMath::Lerp(Start, End, NormalizedParam);
+	}
+	
+protected:
+	virtual float GetParentSpaceMapping(float NormalizedParam) const override
+	{
+		if (KnotVector.Num() < 2) return NormalizedParam;
+
+		float Start, End;
+		GetKnotRange(Start, End);
+
+		// Map normalized parameter to parent space
+		return FMath::Lerp(Start, End, NormalizedParam);
+	}
+
+	virtual float GetNormalizedSpaceMapping(float KnotParam) const override
+	{
+		if (KnotVector.Num() < 2) return KnotParam;
+
+		float Start, End;
+		GetKnotRange(Start, End);
+
+		// Normalize the parameter
+		const float Range = End - Start;
+		if (FMath::IsNearlyZero(Range))
+		{
+			return 0.0f;
+		}
+		return (KnotParam - Start) / Range;
+	}
+	
+private:
+	/** Find knot span using binary search */
+	int32 FindKnotSpan(float Parameter) const
+	{
+		// Validate we have enough knots
+		const int32 MinKnots = NumValues + Degree + 1;
+		if (KnotVector.Num() < MinKnots)
+		{
+			UE_LOG(LogSpline, Error, TEXT("Not enough knots for FindKnotSpan"));
+			return 0;
+		}
+
+		// Handle numerical precision at endpoints
+		const float Epsilon = 1e-6f;
+
+		// Handle the last knot explicitly for clamped endpoints
+		if (Parameter >= KnotVector.Last() - Epsilon)
+		{
+			return NumValues - 1; // Last span index
+		}
+
+		if (Parameter <= KnotVector[Degree] + Epsilon)
+		{
+			return Degree; // First valid span
+		}
+
+		// Binary search for the span
+		int32 Low = Degree;
+		int32 High = KnotVector.Num() - Degree - 1;
+
+		while (Low <= High)
+		{
+			const int32 Mid = (Low + High) / 2;
+
+			if (Parameter >= KnotVector[Mid] && Parameter < KnotVector[Mid + 1])
+			{
+				return Mid;
+			}
+
+			if (Parameter < KnotVector[Mid])
+			{
+				High = Mid - 1;
+			}
+			else
+			{
+				Low = Mid + 1;
+			}
+		}
+
+		// Shouldn't reach here; return fallback
+		UE_LOG(LogSpline, Warning, TEXT("Fallback in FindKnotSpan for t=%f"), Parameter);
+		return Degree;
+	}
+
+	/** Compute B-spline basis functions and their derivatives */
+	void ComputeBasisFunctions(int32 Span, float Parameter, TArray<float>& OutBasis) const
+	{
+		const int32 MinKnots = NumValues + Degree + 1;
+		if (KnotVector.Num() < MinKnots)
+		{
+			UE_LOG(LogSpline, Warning, TEXT("Not enough knots in ComputeBasisFunctions. Have %d, need %d"),
+			       KnotVector.Num(), MinKnots);
+			return;
+		}
+
+		OutBasis.SetNum(Degree + 1);
+		TArray<float> Left;
+		TArray<float> Right;
+		Left.SetNum(Degree + 1);
+		Right.SetNum(Degree + 1);
+
+		OutBasis[0] = 1.0f;
+
+		for (int32 j = 1; j <= Degree; ++j)
+		{
+			// Validate array indices before access
+			const int32 LeftIndex = Span + 1 - j;
+			const int32 RightIndex = Span + j;
+
+			if (LeftIndex < 0 || LeftIndex >= KnotVector.Num() ||
+				RightIndex < 0 || RightIndex >= KnotVector.Num())
+			{
+				UE_LOG(LogSpline, Warning, TEXT("Invalid knot indices in ComputeBasisFunctions: LeftIndex=%d, RightIndex=%d"),
+				       LeftIndex, RightIndex);
+				continue;
+			}
+
+			Left[j] = Parameter - KnotVector[LeftIndex];
+			Right[j] = KnotVector[RightIndex] - Parameter;
+
+			float Saved = 0.0f;
+
+			for (int32 r = 0; r < j; ++r)
+			{
+				float Temp;
+				const float Denominator = Right[r + 1] + Left[j - r];
+				if (FMath::Abs(Denominator) < UE_SMALL_NUMBER)
+				{
+					Temp = 0.0f;
+				}
+				else
+				{
+					Temp = OutBasis[r] / Denominator;
+				}
+
+				OutBasis[r] = Saved + Right[r + 1] * Temp;
+				Saved = Left[j - r] * Temp;
+			}
+
+			OutBasis[j] = Saved;
+		}
+
+		UE_LOG(LogSpline, VeryVerbose, TEXT("BSpline basis at t=%f, span=%d:"), Parameter, Span);
+		for (int32 i = 0; i <= Degree; ++i)
+		{
+			UE_LOG(LogSpline, VeryVerbose, TEXT("  N_%d = %f"), i, OutBasis[i]);
+		}
+	}
+
+	/** 
+     * Generates a uniform knot vector with equal spacing
+     * For clamped ends, multiplicity of degree+1 is used at endpoints
+     */
+	void GenerateUniformKnots(int32 KnotCount)
+	{
+		for (int32 i = 0; i < KnotCount; ++i)
+		{
+			if (bClampEnds)
+			{
+				if (i <= Degree)
+				{
+					// Multiple knots at start for endpoint interpolation
+					KnotVector[i] = 0.0f;
+				}
+				else if (i >= KnotCount - Degree - 1)
+				{
+					// Multiple knots at end for endpoint interpolation
+					KnotVector[i] = KnotCount - 2 * Degree;
+				}
+				else
+				{
+					// Uniform spacing for internal knots
+					KnotVector[i] = i - Degree;
+				}
+			}
+			else
+			{
+				// Simple uniform spacing when not clamped
+				KnotVector[i] = static_cast<float>(i);
+			}
+		}
+
+		UE_LOG(LogSpline, Verbose, TEXT("Updated knot vector - Points: %d, Degree: %d, Knots: %d"),
+			   NumValues, Degree, KnotCount);
+		for (int32 i = 0; i < KnotCount; ++i)
+		{
+			UE_LOG(LogSpline, Verbose, TEXT("  Knot[%d] = %f"), i, KnotVector[i]);
+		}
+	}
+
+	/**
+     * Generates a knot vector based on chord lengths between control points
+     * This gives better parameterization when control points are unevenly spaced
+     */
+    void GenerateChordLengthKnots(int32 KnotCount, const TArrayView<const FStorageEntry<FValue>>& Points)
+    {
+        if (Points.Num() < 2)
+        {
+            GenerateUniformKnots(KnotCount);
+            return;
+        }
+
+        // Calculate chord lengths between consecutive points
+        TArray<float> Chords;
+        Chords.Reserve(Points.Num() - 1);
+        float TotalLength = 0.0f;
+        
+        for (int32 i = 1; i < Points.Num(); ++i)
+        {
+           const float Length = static_cast<float>((Points[i].Value - Points[i-1].Value).Size());
+            Chords.Add(Length);
+            TotalLength += Length;
+        }
+
+        // Generate knot vector based on accumulated chord lengths
+        for (int32 i = 0; i < KnotCount; ++i)
+        {
+            if (bClampEnds)
+            {
+                if (i <= Degree)
+                {
+                    KnotVector[i] = 0.0f;
+                }
+                else if (i >= KnotCount - Degree - 1)
+                {
+                    KnotVector[i] = 1.0f;
+                }
+                else
+                {
+                    // Accumulate normalized chord lengths
+                    const int32 Segment = (i - Degree);
+                    float Param = 0.0f;
+                    for (int32 j = 0; j < Segment; ++j)
+                    {
+                        Param += (TotalLength > UE_SMALL_NUMBER) ? 
+                            (Chords[j] / TotalLength) : 
+                            (1.0f / (float)(Points.Num() - 1));
+                    }
+                    KnotVector[i] = Param;
+                }
+            }
+            else
+            {
+                const float Param = static_cast<float>(i) / static_cast<float>(KnotCount - 1);
+                KnotVector[i] = Param * TotalLength;
+            }
+        }
+    }
+
+    /**
+     * Generates a knot vector using centripetal parameterization
+     * Uses square root of chord lengths which helps prevent cusps and unwanted loops
+     * Often provides the most visually pleasing results for interactive curve editing
+     */
+    void GenerateCentripetalKnots(int32 KnotCount, const TArrayView<const FStorageEntry<FValue>>& Points)
+	{
+		if (Points.Num() < 2)
+		{
+			GenerateUniformKnots(KnotCount);
+			return;
+		}
+
+		TArray<float> Lengths;
+		Lengths.Reserve(Points.Num() - 1);
+		float TotalLength = 0.0f;
+        
+		for (int32 i = 1; i < Points.Num(); ++i)
+		{
+			const float Length = FMath::Sqrt(static_cast<float>((Points[i].Value - Points[i-1].Value).Size()));
+			Lengths.Add(Length);
+			TotalLength += Length;
+		}
+
+		for (int32 i = 0; i < KnotCount; ++i)
+		{
+			if (bClampEnds)
+			{
+				if (i <= Degree)
+				{
+					KnotVector[i] = 0.0f;
+				}
+				else if (i >= KnotCount - Degree - 1)
+				{
+					KnotVector[i] = 1.0f;
+				}
+				else
+				{
+					const int32 Segment = (i - Degree);
+					float Param = 0.0f;
+					for (int32 j = 0; j < Segment; ++j)
+					{
+						Param += (TotalLength > UE_SMALL_NUMBER) ? 
+							(Lengths[j] / TotalLength) : 
+							(1.0f / (float)(Points.Num() - 1));
+					}
+					KnotVector[i] = Param;
+				}
+			}
+			else
+			{
+				const float Param = static_cast<float>(i) / static_cast<float>(KnotCount - 1);
+				KnotVector[i] = Param * TotalLength;
+			}
+		}
+	}
+
+};
+
+template <typename ScalarType, int32 Dimensions, int32 Degree = 3>
+class TBSpline : public TSplineBase<ScalarType, Dimensions>
+{
+	static_assert(Degree > 0, "BSpline degree must be positive");
+    
+	using Base = TSplineBase<ScalarType, Dimensions>;
+	
+protected:
+	virtual TSharedPtr<Base> CreateNewInstance() const override
+	{
+		return MakeShared<TBSpline>();
+	}
+    
+	virtual void OnClone(TSharedPtr<Base> NewSpline, ESplineCloneFlags CloneFlags) const override
+	{
+		TSharedPtr<TBSpline> TypedSpline = StaticCastSharedPtr<TBSpline>(NewSpline);
+		TypedSpline->GetEvaluator().SetKnotGenerationMode(this->GetEvaluator().GetKnotMode());
+		TypedSpline->GetEvaluator().SetClampedEnds(this->GetEvaluator().IsClampedEnds());
+        
+		if (this->GetEvaluator().GetKnotMode() == EBSplineKnotMode::Custom)
+		{
+			TypedSpline->SetCustomKnots(this->GetEvaluator().GetKnotVector());
+		}
+	}
+	
+public:
+	using typename Base::FValue;
+
+	TBSpline()
+		: Base(MakeUnique<TBSplineEvaluator<ScalarType, Dimensions, Degree>>())
+	{
+	}
+
+	TBSplineEvaluator<ScalarType, Dimensions, Degree>& GetEvaluator()
+	{
+		return this->template GetEvaluatorAs<TBSplineEvaluator<ScalarType, Dimensions, Degree>>();
+	}
+
+	const TBSplineEvaluator<ScalarType, Dimensions, Degree>& GetEvaluator() const
+	{
+		return this->template GetEvaluatorAs<TBSplineEvaluator<ScalarType, Dimensions, Degree>>();
+	}
+
+	/**
+	 * Sets a custom knot vector for the spline
+	 * @param NewKnots - Array of knot values. Must be non-decreasing and have length numPoints + degree + 1
+	 * @return True if the knot vector was valid and set successfully
+	 */
+	void SetCustomKnots(const TArray<float>& NewKnots)
+	{
+		GetEvaluator().SetCustomKnots(NewKnots, Base::GetNumKeys());
+	}
+
+	/**
+	 * Adds a new value to the spline and updates knot vector accordingly
+	 * @param Value - The control point value to add
+	 * @param Parameter - Parameter value for the point
+	 * @return Index of the added value
+	 */
+	virtual int32 AddValue(const FValue& Value, float Parameter) override
+	{
+		const int32 Index = Base::AddValue(Value, Parameter);
+		// Update knot vector unless using custom knots
+		if (GetEvaluator().GetKnotMode() != EBSplineKnotMode::Custom)
+		{
+			GetEvaluator().UpdateKnotVector(Base::Values->GetStorage().GetEntries());
+		}
+		return Index;
+	}
+	
+	/**
+	 * Changes how the knot vector is generated
+	 * @param Mode - The new knot generation mode
+	 */
+	void SetKnotGenerationMode(EBSplineKnotMode Mode)
+	{
+		GetEvaluator().SetKnotGenerationMode(Mode);
+		GetEvaluator().UpdateKnotVector(Base::Values->GetStorage().GetEntries());
+	}
+
+	/**
+	 * Sets whether the spline should interpolate its endpoints
+	 * @param bInClampEnds - True to interpolate endpoints, false for unclamped curve
+	 */
+	void SetClampedEnds(bool bInClampEnds)
+	{
+		GetEvaluator().SetClampedEnds(bInClampEnds);
+		GetEvaluator().UpdateKnotVector(Base::Values->GetStorage().GetEntries());
+	}
+};
+
+// Common type definitions
+using FBSpline2f = TBSpline<float, 2>;
+using FBSpline3f = TBSpline<float, 3>;
+using FBSpline2d = TBSpline<double, 2>;
+using FBSpline3d = TBSpline<double, 3>;
+
+// Explicit degree type definitions if needed
+using FQuadraticBSpline2f = TBSpline<float, 2, 2>;
+using FCubicBSpline3f = TBSpline<float, 3, 3>;
+using FQuarticBSpline3f = TBSpline<float, 3, 4>;
+using FQuinticBSpline2d = TBSpline<double, 2, 5>;
\ No newline at end of file
diff --git a/Engine/Source/Runtime/GeometryCore/Public/Splines/BaseEvaluator.h b/Engine/Source/Runtime/GeometryCore/Public/Splines/BaseEvaluator.h
new file mode 100644
--- /dev/null
+++ b/Engine/Source/Runtime/GeometryCore/Public/Splines/BaseEvaluator.h
@@ -0,0 +1,221 @@
+// Copyright Epic Games, Inc. All Rights Reserved.
+
+#pragma once
+
+#include "CoreMinimal.h"
+#include "InterpolationPolicies.h"
+
+DEFINE_LOG_CATEGORY_STATIC(LogSpline, Log, All);
+/** 
+ * Enum for handling out-of-bounds parameter values
+ */
+enum class EOutOfBoundsHandlingMode
+{
+	Zero,       // Null out the values outside of the bounds
+	Constant,   // Choose the nearest of the first/last value
+	Cycle       // Repeat the curve indefinitely (mod of the parameter space)
+};
+
+/**
+ * Base evaluator class providing parameter mapping and bounds handling
+ */
+template<typename ValueType>
+class TBaseEvaluator
+{
+protected:
+	/** Start point in parent space */
+	float ParentStart = 0.0f;
+    
+	/** End point in parent space */
+	float ParentEnd = 1.0f;
+
+	/** How to handle parameters outside the mapped range */
+	EOutOfBoundsHandlingMode PreInfinityMode = EOutOfBoundsHandlingMode::Zero;
+	EOutOfBoundsHandlingMode PostInfinityMode = EOutOfBoundsHandlingMode::Zero;
+
+public:
+	virtual ~TBaseEvaluator() = default;
+	
+	/** Get the polynomial degree of the evaluator */
+	virtual int32 GetDegree() const{ return 1; }
+
+	/** 
+	 * Finds interpolation window at the given parameter in channel space
+	 * This class implements a default window finding implementation for two-point interpolation
+	 * @param Parameter - Parameter value in channel's local space
+	 */
+	/**  */
+	virtual TArray<const ValueType*> FindWindow(
+		TArrayView<const FStorageEntry<ValueType>> Entries,
+		float Parameter) const
+	{
+		TArray<const ValueType*> Window;
+		if (Entries.Num() < 2) return Window;
+
+		// Map parameter to parent space
+		const float ParentParam = MapToParentSpace(Parameter);
+		
+		// Handle bounds
+		if (ParentParam <= Entries[0].Parameter)
+		{
+			// Before first entry
+			Window.Add(&Entries[0].Value);
+			Window.Add(&Entries[1].Value);
+			return Window;
+		}
+		else if (ParentParam >= Entries.Last().Parameter)
+		{
+			// At or beyond last entry
+			Window.Add(&Entries[Entries.Num() - 2].Value);
+			Window.Add(&Entries.Last().Value);
+			return Window;
+		}
+        
+		// Find bracketing points
+		int32 Index = 0;
+		for (; Index < Entries.Num() - 1; ++Index)
+		{
+			if (Entries[Index + 1].Parameter > ParentParam)
+			{
+				break;
+			}
+		}
+
+		Window.Add(&Entries[Index].Value);
+		Window.Add(&Entries[Index + 1].Value);
+		return Window;
+	}
+
+	/** Default interpolation using the type's interpolation policy */
+	virtual ValueType InterpolateWindow(
+		TArrayView<const ValueType* const> Window,
+		float Parameter) const
+	{
+		const float MappedParam = MapToParentSpace(Parameter);
+		return TSplineInterpolationPolicy<ValueType>::Interpolate(Window, MappedParam);
+	}
+
+	/** Clone the evaluator */
+	virtual TUniquePtr<TBaseEvaluator<ValueType>> Clone() const
+	{
+		return MakeUnique<TBaseEvaluator<ValueType>>(*this);
+	}
+	
+	/** Convert from normalized [0,1] to parent space with bounds handling */
+	virtual float MapToParentSpace(float NormalizedParam) const
+	{
+		const float BoundedParam = HandleOutOfBounds(NormalizedParam);
+		const float Result = GetParentSpaceMapping(BoundedParam);
+
+		UE_LOG(LogSpline, VeryVerbose, TEXT("MapToParentSpace: %f -> Bounded: %f -> Result: %f"), 
+			NormalizedParam, BoundedParam, Result);
+		return Result;
+	}
+
+	/** Convert from parent space back to normalized [0,1] with bounds handling */
+	virtual float MapFromParentSpace(float ParentParam) const
+	{
+		const float LocalParam = GetNormalizedSpaceMapping(ParentParam);
+		const float Result = HandleOutOfBounds(LocalParam);
+    
+		UE_LOG(LogSpline, VeryVerbose, TEXT("MapFromParentSpace: %f -> Local: %f -> Result: %f"), 
+			ParentParam, LocalParam, Result);
+		return Result;
+	}
+
+	/** Set the mapping range in parent space */
+	void SetParentSpaceRange(float InStart, float InEnd)
+	{
+		ParentStart = InStart;
+		ParentEnd = InEnd;
+	}
+
+	// Auto-expands parameter range for parameters outside current bounds
+	void ExpandRangeToInclude(float Parameter)
+	{
+		// Expand range if parameter lies outside current bounds
+		ParentStart = FMath::Min(ParentStart, Parameter);
+		ParentEnd = FMath::Max(ParentEnd, Parameter);
+	}
+
+	/** Set Pre-Infinity handling mode when out of bounds*/ 
+	void SetPreInfinityMode(EOutOfBoundsHandlingMode InMode)
+	{
+		PreInfinityMode = InMode;
+	}
+
+	/** Set Post-Infinity handling mode when out of bounds */
+	void SetPostInfinityMode(EOutOfBoundsHandlingMode InMode)
+	{
+		PostInfinityMode = InMode;
+	}
+
+	/** Get the current post-infinity mode */
+	EOutOfBoundsHandlingMode GetPreInfinityMode() const
+	{
+		return PreInfinityMode;
+	}
+
+	/** Get the current post-infinity mode */
+	EOutOfBoundsHandlingMode GetPostInfinityMode() const
+	{
+		return PostInfinityMode;
+	}
+
+	/** Get the current parent space range */
+	void GetParentSpaceRange(float& OutStart, float& OutEnd) const
+	{
+		OutStart = ParentStart;
+		OutEnd = ParentEnd;
+	}
+	
+protected:
+	/** Override to implement specific parent space mapping */
+	virtual float GetParentSpaceMapping(float NormalizedParam) const
+	{
+		return NormalizedParam;
+	}
+
+	/** Override to implement specific normalized space mapping */
+	virtual float GetNormalizedSpaceMapping(float ParentParam) const
+	{
+		return ParentParam;
+	}
+private:
+	/** Apply out-of-bounds handling to parameter */
+	float HandleOutOfBounds(float Parameter) const
+	{
+		if (Parameter < ParentStart)
+		{
+			switch (PreInfinityMode)
+			{
+			case EOutOfBoundsHandlingMode::Zero:
+				return 0.0f;
+			case EOutOfBoundsHandlingMode::Constant:
+				return ParentStart;
+			case EOutOfBoundsHandlingMode::Cycle:
+				{
+					const float Range = ParentEnd - ParentStart;
+					return FMath::Fmod(Parameter - ParentStart, Range) + ParentStart;
+				}
+			}
+		}
+		else if (Parameter > ParentEnd)
+		{
+			switch (PostInfinityMode)
+			{
+			case EOutOfBoundsHandlingMode::Zero:
+				return 0.0f;
+			case EOutOfBoundsHandlingMode::Constant:
+				return ParentEnd;
+			case EOutOfBoundsHandlingMode::Cycle:
+				{
+					const float Range = ParentEnd - ParentStart;
+					return FMath::Fmod(Parameter - ParentStart, Range) + ParentStart;
+				}
+			}
+		}
+		
+		return Parameter;
+	}
+};
\ No newline at end of file
diff --git a/Engine/Source/Runtime/GeometryCore/Public/Splines/CubicBezierSpline.h b/Engine/Source/Runtime/GeometryCore/Public/Splines/CubicBezierSpline.h
new file mode 100644
--- /dev/null
+++ b/Engine/Source/Runtime/GeometryCore/Public/Splines/CubicBezierSpline.h
@@ -0,0 +1,37 @@
+// Copyright Epic Games, Inc. All Rights Reserved.
+
+#pragma once
+
+#include "BSpline.h"
+
+/**
+ * Cubic Bezier spline implementation as a specialization of B-spline
+ * with degree 3 and appropriate knot multiplicity
+ */
+template<typename ScalarType, int32 Dimensions>
+class TCubicBezierSpline : public TBSpline<ScalarType, Dimensions, 3>  // Explicitly specify degree 3
+{
+	using Base = TBSpline<ScalarType, Dimensions, 3>;
+
+public:
+    using typename Base::FValue;
+
+    TCubicBezierSpline()
+        : Base()
+    {
+		// Configure base B-spline for Bezier behavior
+		this->SetKnotGenerationMode(EBSplineKnotMode::Uniform);
+		this->SetClampedEnds(true);  // This gives us the Bezier end behavior
+    }
+	
+	virtual TSharedPtr<TSplineBase<ScalarType, Dimensions>> CreateNewInstance() const override
+    {
+    	return MakeShared<TCubicBezierSpline>();
+    }
+};
+
+// Common type definitions
+using FCubicBezierSpline2f = TCubicBezierSpline<float, 2>;
+using FCubicBezierSpline3f = TCubicBezierSpline<float, 3>;
+using FCubicBezierSpline2d = TCubicBezierSpline<double, 2>;
+using FCubicBezierSpline3d = TCubicBezierSpline<double, 3>;
\ No newline at end of file
diff --git a/Engine/Source/Runtime/GeometryCore/Public/Splines/HermiteSpline.h b/Engine/Source/Runtime/GeometryCore/Public/Splines/HermiteSpline.h
new file mode 100644
--- /dev/null
+++ b/Engine/Source/Runtime/GeometryCore/Public/Splines/HermiteSpline.h
@@ -0,0 +1,239 @@
+// Copyright Epic Games, Inc. All Rights Reserved.
+
+#pragma once
+
+#include "Spline.h"
+
+/**
+ * Hermite spline implementation with tangent control
+ */
+template<typename T, int32 D>
+class THermiteSpline : public TSplineBase<T, D>
+{
+    using Base = TSplineBase<T, D>;
+    using FValue = typename Base::FValue;
+	
+	static const FName TangentInAttribute;
+	static const FName TangentOutAttribute;
+	
+    class FEvaluator : public TBaseEvaluator<FValue>
+    {
+        const THermiteSpline<T, D>& Owner;
+
+    public:
+        explicit FEvaluator(const THermiteSpline<T, D>& InOwner)
+            : Owner(InOwner)
+        {
+        }
+
+        virtual int32 GetDegree() const override { return 3; }
+
+    	virtual TArray<const FValue*> FindWindow(
+			TArrayView<const FStorageEntry<FValue>> Entries,
+			float Parameter) const override
+        {
+        	TArray<const FValue*> Window;
+        	if (Entries.Num() < 2) return Window;
+
+        	// Map to parent space
+        	const float ParentParam = this->MapToParentSpace(Parameter);
+
+        	// Find segment based on clamped parameter
+        	const int32 NumSegments = Entries.Num() - 1;
+        	
+        	// Add explicit clamping for the end points
+        	const float ClampedParam = FMath::Clamp(ParentParam, 0.0f, float(NumSegments));
+        	const int32 SegmentIndex = FMath::Min(FMath::FloorToInt(ClampedParam), NumSegments - 1);
+
+
+        	UE_LOG(LogSpline, Verbose, TEXT("FindWindow: t=%f, segment=%d, NumSegments=%d"), 
+				ParentParam, SegmentIndex, NumSegments);
+
+        	// Add points that bound this segment
+        	Window.Add(&Entries[SegmentIndex].Value);
+        	Window.Add(&Entries[SegmentIndex + 1].Value);
+
+        	// Get corresponding tangents
+        	static FValue DefaultTanOut, DefaultTanIn;
+
+        	if (SegmentIndex == 0)  // First segment
+        	{
+        		DefaultTanOut = Entries[1].Value - Entries[0].Value;
+        	}
+        	else if (SegmentIndex == NumSegments - 1)  // Last segment
+        	{
+        		DefaultTanOut = (Entries[SegmentIndex + 1].Value - Entries[SegmentIndex].Value) * 0.5f;
+        	}
+        	else  // Middle segments
+        	{
+        		DefaultTanOut = (Entries[SegmentIndex + 1].Value - Entries[SegmentIndex - 1].Value) * 0.5f;
+        	}
+
+        	if (SegmentIndex == NumSegments - 1)  // Last segment
+        	{
+        		DefaultTanIn = Entries[NumSegments].Value - Entries[NumSegments - 1].Value;
+        	}
+        	else if (SegmentIndex == 0)  // First segment
+        	{
+        		DefaultTanIn = (Entries[SegmentIndex + 2].Value - Entries[SegmentIndex].Value) * 0.5f;
+        	}
+        	else  // Middle segments
+        	{
+        		DefaultTanIn = (Entries[SegmentIndex + 2].Value - Entries[SegmentIndex].Value) * 0.5f;
+        	}
+
+        	const FValue* TanOut = Owner.GetTangentOut(SegmentIndex);
+        	Window.Add(TanOut ? TanOut : &DefaultTanOut);
+
+        	const FValue* TanIn = Owner.GetTangentIn(SegmentIndex + 1);
+        	Window.Add(TanIn ? TanIn : &DefaultTanIn);
+
+        	return Window;
+        }
+
+        virtual FValue InterpolateWindow(TArrayView<const FValue* const> Window, float Parameter) const override
+        {
+            if (Window.Num() < 4) return FValue::ZeroVector;
+
+        	// Map to parent space
+        	const float ParentParam = this->MapToParentSpace(Parameter);
+        	
+        	// Parameter is in [0,NumSegments], get local t for segment
+        	// Get the segment index and local parameter
+        	const int32 NumSegments = Owner.GetNumKeys() - 1;
+        	const float SegmentIndex = FMath::Floor(ParentParam);
+    
+        	// Calculate local T, ensuring we get 1.0 at the end of the last segment
+        	float LocalT;
+        	if (SegmentIndex >= NumSegments)
+        	{
+        		LocalT = 1.0f;
+        	}
+        	else 
+        	{
+        		LocalT = ParentParam - SegmentIndex;
+        	}
+        	
+        	const float t2 = LocalT * LocalT;
+        	const float t3 = t2 * LocalT;
+
+        	// Compute Hermite basis functions with LocalT
+        	const float h00 = 2.0f*t3 - 3.0f*t2 + 1.0f;
+        	const float h10 = t3 - 2.0f*t2 + LocalT;
+        	const float h01 = -2.0f*t3 + 3.0f*t2;
+        	const float h11 = t3 - t2;
+
+        	const float Weights[4] = { h00, h01, h10, h11 };
+        	return TSplineInterpolationPolicy<FValue>::InterpolateWithBasis(Window, Weights);
+        }
+
+    protected:
+    	virtual float GetParentSpaceMapping(float NormalizedParam) const override
+    	{
+    		const int32 NumSegments = Owner.GetNumKeys() - 1;
+    		if (NumSegments <= 0) return NormalizedParam;
+
+    		// Map [0,1] to [0, NumSegments]
+    		return NormalizedParam * NumSegments;
+    	}
+
+    	virtual float GetNormalizedSpaceMapping(float ParentParam) const override
+    	{
+    		const int32 NumSegments = Owner.GetNumKeys() - 1;
+    		if (NumSegments <= 0) return ParentParam;
+
+    		// Map [0, NumSegments] back to [0,1]
+    		return ParentParam / NumSegments;
+    	}
+    };
+
+
+    FEvaluator* GetEvaluator()
+    {
+        return static_cast<FEvaluator*>(&this->Values->GetEvaluator());
+    }
+
+    const FEvaluator* GetEvaluator() const
+    {
+        return static_cast<const FEvaluator*>(&this->Values->GetEvaluator());
+    }
+
+public:
+	THermiteSpline()
+		: Base(MakeUnique<FEvaluator>(*this))
+	{
+		Base::template CreateAttributeChannel<FValue>(TangentInAttribute); 
+		Base::template CreateAttributeChannel<FValue>(TangentOutAttribute);
+	}
+	
+    void SetTangentIn(int32 Index, const FValue& Tangent)
+    {
+		Base::template AddAttributeValue(TangentInAttribute, Tangent, Index);
+    }
+
+    void SetTangentOut(int32 Index, const FValue& Tangent)
+    {
+		Base::template AddAttributeValue(TangentOutAttribute, Tangent, Index);
+    }
+
+    const FValue* GetTangentIn(int32 Index) const
+    {
+		if (const FAttributeChannel<FValue>* Channel = this->template GetAttributeChannel<FValue>(TangentInAttribute))
+		{
+			return &Channel->GetStorage().GetEntries()[Index].Value;
+		}
+		return nullptr;
+    }
+
+    const FValue* GetTangentOut(int32 Index) const
+    {
+		if (const FAttributeChannel<FValue>* Channel = this->template GetAttributeChannel<FValue>(TangentOutAttribute))
+		{
+			return &Channel->GetStorage().GetEntries()[Index].Value;
+		}
+		return nullptr;
+    }
+
+    virtual void Clear() override
+    {
+        Base::Clear();
+    }
+
+protected:
+	virtual TSharedPtr<TSplineBase<T, D>> CreateNewInstance() const override
+	{
+		return MakeShared<THermiteSpline>();
+	}
+    
+	virtual void OnClone(TSharedPtr<TSplineBase<T, D>> NewSpline, ESplineCloneFlags CloneFlags) const override
+	{
+		TSharedPtr<THermiteSpline> HermiteSpline = StaticCastSharedPtr<THermiteSpline>(NewSpline);
+		// if attributes are being cloned, we don't need to explicitly clone the tangents here
+		// as they are part of the attributes (already cloned with CopyBaseData call above)
+		if (!EnumHasAnyFlags(CloneFlags, ESplineCloneFlags::Attributes))
+		{
+			// Copy Hermite-specific data
+			if (const FAttributeChannel<FValue>* Channel = this->template GetAttributeChannel<FValue>(TangentInAttribute))
+			{
+				HermiteSpline->AttributeChannels.Add(TangentInAttribute, MakeUnique<FAttributeChannel<FValue>>(*Channel));
+			}
+
+			if (const FAttributeChannel<FValue>* Channel = this->template GetAttributeChannel<FValue>(TangentOutAttribute))
+			{
+				HermiteSpline->AttributeChannels.Add(TangentOutAttribute, MakeUnique<FAttributeChannel<FValue>>(*Channel));
+			}
+		}
+	}
+};
+
+template<typename T, int32 D>
+const FName THermiteSpline<T, D>::TangentInAttribute = TEXT("TangentIn");
+
+template<typename T, int32 D>
+const FName THermiteSpline<T, D>::TangentOutAttribute = TEXT("TangentOut");
+
+// Common type definitions
+using FHermiteSpline2f = THermiteSpline<float, 2>;
+using FHermiteSpline3f = THermiteSpline<float, 3>;
+using FHermiteSpline2d = THermiteSpline<double, 2>;
+using FHermiteSpline3d = THermiteSpline<double, 3>;
diff --git a/Engine/Source/Runtime/GeometryCore/Public/Splines/InterpolationPolicies.h b/Engine/Source/Runtime/GeometryCore/Public/Splines/InterpolationPolicies.h
new file mode 100644
--- /dev/null
+++ b/Engine/Source/Runtime/GeometryCore/Public/Splines/InterpolationPolicies.h
@@ -0,0 +1,330 @@
+// Copyright Epic Games, Inc. All Rights Reserved.
+
+#pragma once
+
+#include "ParameterizedTypes.h"
+
+/**
+ * Common interpolation policies for built-in types
+ */
+/**
+ * Base spline interpolation policy that works across all spline types
+ * Extends the existing parameter-based interpolation with window-based methods
+ */
+template<typename T>
+class TSplineInterpolationPolicy
+{
+public:
+    // Keep existing parameter-based interpolation (used by attribute channels)
+    static T Interpolate(TArrayView<const T* const> Window, float Parameter)
+    {
+        if (Window.Num() < 2)
+        {
+            return Window.Num() == 1 ? *Window[0] : T();
+        }
+        return FMath::Lerp(*Window[0], *Window[1], Parameter);
+    }
+
+    // Add window-based interpolation for spline implementations
+    static T InterpolateWithBasis(TArrayView<const T* const> Window, TArrayView<const float> Basis)
+    {
+        T Result = T();
+        for (int32 i = 0; i < FMath::Min(Window.Num(), Basis.Num()); ++i)
+        {
+            Result += *Window[i] * Basis[i];
+        }
+        return Result;
+    }
+
+    // Helper for auto-tangent computation
+    static T BlendTangents(const T& A, const T& B, float Alpha = 0.5f)
+    {
+        return Interpolate({ &A, &B }, Alpha);
+    }
+};
+
+// FColor interpolation policy
+template<>
+class TSplineInterpolationPolicy<FColor>
+{
+public:
+    static FColor Interpolate(TArrayView<const FColor* const> Window, float Parameter)
+    {
+        if (Window.Num() < 2)
+        {
+            return Window.Num() == 1 ? *Window[0] : FColor();
+        }
+        
+        // Convert to linear color for better interpolation
+        return FLinearColor::LerpUsingHSV(
+            FLinearColor(*Window[0]),
+            FLinearColor(*Window[1]),
+            Parameter
+        ).ToFColor(true);
+    }
+
+    static FColor InterpolateWithBasis(TArrayView<const FColor* const> Window, TArrayView<const float> Basis)
+    {
+        // Convert all colors to linear space first
+        TArray<FLinearColor> LinearColors;
+        LinearColors.Reserve(Window.Num());
+        for (const FColor* Color : Window)
+        {
+            LinearColors.Add(FLinearColor(*Color));
+        }
+
+        // Perform weighted sum in linear space
+        FLinearColor Result = FLinearColor::Black;
+        for (int32 i = 0; i < FMath::Min(Window.Num(), Basis.Num()); ++i)
+        {
+            Result += LinearColors[i] * Basis[i];
+        }
+
+        // Convert back to FColor
+        return Result.ToFColor(true);
+    }
+
+    static FColor BlendTangents(const FColor& A, const FColor& B, float Alpha = 0.5f)
+    {
+        return FLinearColor::LerpUsingHSV(FLinearColor(A), FLinearColor(B), Alpha).ToFColor(true);
+    }
+};
+
+// FLinearColor interpolation policy
+template<>
+class TSplineInterpolationPolicy<FLinearColor>
+{
+public:
+    static FLinearColor Interpolate(TArrayView<const FLinearColor* const> Window, float Parameter)
+    {
+        if (Window.Num() < 2)
+        {
+            return Window.Num() == 1 ? *Window[0] : FLinearColor();
+        }
+        
+        // Convert to linear color for better interpolation
+        return FLinearColor::LerpUsingHSV(
+            *Window[0],
+            *Window[1],
+            Parameter
+        ).ToFColor(true);
+    }
+
+    static FLinearColor InterpolateWithBasis(TArrayView<const FLinearColor* const> Window, TArrayView<const float> Basis)
+    {
+        // Convert all colors to linear space first
+        TArray<FLinearColor> LinearColors;
+        LinearColors.Reserve(Window.Num());
+        for (const FLinearColor* Color : Window)
+        {
+            LinearColors.Add(*Color);
+        }
+
+        // Perform weighted sum in linear space
+        FLinearColor Result = FLinearColor::Black;
+        for (int32 i = 0; i < FMath::Min(Window.Num(), Basis.Num()); ++i)
+        {
+            Result += LinearColors[i] * Basis[i];
+        }
+
+        // Convert back to FColor
+        return Result;
+    }
+
+    static FLinearColor BlendTangents(const FLinearColor& A, const FLinearColor& B, float Alpha = 0.5f)
+    {
+        return FLinearColor::LerpUsingHSV(A, B, Alpha);
+    }
+};
+
+// FQuat interpolation policy
+template<>
+class TSplineInterpolationPolicy<FQuat>
+{
+public:
+    // Existing parameter-based interpolation
+    static FQuat Interpolate(TArrayView<const FQuat* const> Window, float Parameter)
+    {
+        if (Window.Num() < 2)
+        {
+            return Window.Num() == 1 ? *Window[0] : FQuat::Identity;
+        }
+        return FQuat::Slerp(*Window[0], *Window[1], Parameter);
+    }
+
+    // Window-based interpolation for spline implementations
+    static FQuat InterpolateWithBasis(TArrayView<const FQuat* const> Window, TArrayView<const float> Basis)
+    {
+        FQuat Result = FQuat::Identity;
+        float TotalWeight = 0.0f;
+        
+        for (int32 i = 0; i < FMath::Min(Window.Num(), Basis.Num()); ++i)
+        {
+            const float Weight = Basis[i];
+            if (Weight > UE_SMALL_NUMBER)
+            {
+                const FQuat& Q = *Window[i];
+                const double Dot = Result.W * Q.W + Result.X * Q.X + Result.Y * Q.Y + Result.Z * Q.Z;
+                Result += (Dot >= 0.0f ? Q : -Q) * Weight;
+                TotalWeight += Weight;
+            }
+        }
+
+        if (TotalWeight > UE_SMALL_NUMBER)
+        {
+            Result.Normalize();
+        }
+        return Result;
+    }
+
+    static FQuat BlendTangents(const FQuat& A, const FQuat& B, float Alpha = 0.5f)
+    {
+        return FQuat::Slerp(A, B, Alpha);
+    }
+};
+
+// FRotator interpolation policy
+template<>
+class TSplineInterpolationPolicy<FRotator>
+{
+public:
+    static FRotator Interpolate(TArrayView<const FRotator* const> Window, float Parameter)
+    {
+        if (Window.Num() < 2)
+        {
+            return Window.Num() == 1 ? *Window[0] : FRotator::ZeroRotator;
+        }
+        
+        return FMath::LerpRange(*Window[0], *Window[1], Parameter);
+    }
+
+    static FRotator InterpolateWithBasis(TArrayView<const FRotator* const> Window, TArrayView<const float> Basis)
+    {
+        // Convert to quaternions for better interpolation
+        TArray<FQuat> Quats;
+        Quats.Reserve(Window.Num());
+        for (const FRotator* Rot : Window)
+        {
+            Quats.Add(Rot->Quaternion());
+        }
+
+        // Use quaternion interpolation
+        FQuat Result = TSplineInterpolationPolicy<FQuat>::InterpolateWithBasis(
+            MakeArrayView((const FQuat**)Quats.GetData(), Quats.Num()),
+            Basis
+        );
+
+        return FRotator(Result);
+    }
+
+    static FRotator BlendTangents(const FRotator& A, const FRotator& B, float Alpha = 0.5f)
+    {
+        return FMath::LerpRange(A, B, Alpha);
+    }
+};
+
+// FTransform interpolation policy
+template<>
+class TSplineInterpolationPolicy<FTransform>
+{
+public:
+    static FTransform Interpolate(TArrayView<const FTransform* const> Window, float Parameter)
+    {
+        if (Window.Num() < 2)
+        {
+            return Window.Num() == 1 ? *Window[0] : FTransform::Identity;
+        }
+        
+        const FTransform& A = *Window[0];
+        const FTransform& B = *Window[1];
+        
+        FTransform Result;
+        Result.SetLocation(FMath::Lerp(A.GetLocation(), B.GetLocation(), Parameter));
+        Result.SetRotation(FQuat::Slerp(A.GetRotation(), B.GetRotation(), Parameter));
+        Result.SetScale3D(FMath::Lerp(A.GetScale3D(), B.GetScale3D(), Parameter));
+        
+        return Result;
+    }
+
+    static FTransform InterpolateWithBasis(TArrayView<const FTransform* const> Window, TArrayView<const float> Basis)
+    {
+        // Handle empty case
+        if (Window.Num() == 0 || Basis.Num() == 0)
+        {
+            return FTransform::Identity;
+        }
+
+        // Interpolate components separately
+        FVector Location = FVector::ZeroVector;
+        FQuat Rotation = FQuat::Identity;
+        FVector Scale = FVector(1.0f);
+
+        float TotalRotWeight = 0.0f;
+        
+        // Weighted sum of locations
+        for (int32 i = 0; i < FMath::Min(Window.Num(), Basis.Num()); ++i)
+        {
+            const float Weight = Basis[i];
+            if (Weight > UE_SMALL_NUMBER)
+            {
+                Location += Window[i]->GetLocation() * Weight;
+                Scale += Window[i]->GetScale3D() * Weight;
+
+                // Handle rotation separately with proper quaternion blending
+                const FQuat& Q = Window[i]->GetRotation();
+                const double Dot = Rotation.W * Q.W + Rotation.X * Q.X + Rotation.Y * Q.Y + Rotation.Z * Q.Z;
+                Rotation += (Dot >= 0.0f ? Q : -Q) * Weight;
+                TotalRotWeight += Weight;
+            }
+        }
+
+        // Normalize the quaternion if we accumulated any weights
+        if (TotalRotWeight > UE_SMALL_NUMBER)
+        {
+            Rotation.Normalize();
+        }
+
+        return FTransform(Rotation, Location, Scale);
+    }
+
+    static FTransform BlendTangents(const FTransform& A, const FTransform& B, float Alpha = 0.5f)
+    {
+        FTransform Result;
+        Result.SetLocation(FMath::Lerp(A.GetLocation(), B.GetLocation(), Alpha));
+        Result.SetRotation(FQuat::Slerp(A.GetRotation(), B.GetRotation(), Alpha));
+        Result.SetScale3D(FMath::Lerp(A.GetScale3D(), B.GetScale3D(), Alpha));
+        return Result;
+    }
+};
+
+// Specialization for vectors
+template<>
+class TSplineInterpolationPolicy<FVector3f>
+{
+public:
+    // Classic two-point interpolation
+    static FVector3f Interpolate(TArrayView<const FVector3f* const> Window, float Parameter)
+    {
+        if (Window.Num() < 2)
+        {
+            return Window.Num() == 1 ? *Window[0] : FVector3f::ZeroVector;
+        }
+        return FMath::Lerp(*Window[0], *Window[1], Parameter);
+    }
+
+    // Window-based interpolation for splines
+    static FVector3f InterpolateWithBasis(TArrayView<const FVector3f* const> Window, TArrayView<const float> Basis)
+    {
+        FVector3f Result = FVector3f::ZeroVector;
+        for (int32 i = 0; i < FMath::Min(Window.Num(), Basis.Num()); ++i)
+        {
+            Result += *Window[i] * Basis[i];
+        }
+        return Result;
+    }
+
+    static FVector3f BlendTangents(const FVector3f& A, const FVector3f& B, float Alpha = 0.5f)
+    {
+        return FMath::Lerp(A, B, Alpha);
+    }
+};
\ No newline at end of file
diff --git a/Engine/Source/Runtime/GeometryCore/Public/Splines/ParameterizedStorage.h b/Engine/Source/Runtime/GeometryCore/Public/Splines/ParameterizedStorage.h
new file mode 100644
--- /dev/null
+++ b/Engine/Source/Runtime/GeometryCore/Public/Splines/ParameterizedStorage.h
@@ -0,0 +1,36 @@
+// Copyright Epic Games, Inc. All Rights Reserved.
+
+#pragma once
+
+#include "ParameterizedTypes.h"
+/**
+ * Pure storage class - no evaluation logic
+ */
+template<typename AttrType>
+class FAttributeStorage
+{
+    TArray<FStorageEntry<AttrType>> Entries;
+
+public:
+    int32 AddValue(const AttrType& Value, float Parameter)
+    {
+        const int32 Index = Algo::LowerBound(
+            Entries, Parameter,
+            [](const FStorageEntry<AttrType>& Entry, float Param) { 
+                return Entry.Parameter < Param; 
+            }
+        );
+        Entries.Insert({Parameter, Value}, Index);
+    	return Index;
+    }
+
+    TArrayView<const FStorageEntry<AttrType>> GetEntries() const 
+    { 
+        return Entries; 
+    }
+
+    void Clear()
+    {
+        Entries.Reset();
+    }
+};
\ No newline at end of file
diff --git a/Engine/Source/Runtime/GeometryCore/Public/Splines/ParameterizedTypes.h b/Engine/Source/Runtime/GeometryCore/Public/Splines/ParameterizedTypes.h
new file mode 100644
--- /dev/null
+++ b/Engine/Source/Runtime/GeometryCore/Public/Splines/ParameterizedTypes.h
@@ -0,0 +1,66 @@
+// Copyright Epic Games, Inc. All Rights Reserved.
+
+#pragma once
+
+#include "CoreMinimal.h"
+
+/**
+ * Vector type policies based on dimension and numeric type
+ */
+template<typename T, int32 D> 
+struct TVectorPolicy
+{
+	// Default case - invalid
+	using Type = void;
+};
+
+// Specializations for 2D vectors
+template<> struct TVectorPolicy<float, 2> { using Type = FVector2f; };
+template<> struct TVectorPolicy<double, 2> { using Type = FVector2D; };
+template<> struct TVectorPolicy<int32, 2> { using Type = FIntPoint; };
+
+// Specializations for 3D vectors  
+template<> struct TVectorPolicy<float, 3> { using Type = FVector3f; };
+template<> struct TVectorPolicy<double, 3> { using Type = FVector; };
+template<> struct TVectorPolicy<int32, 3> { using Type = FIntVector; };
+
+// Specializations for 4D vectors
+template<> struct TVectorPolicy<float, 4> { using Type = FVector4f; };
+template<> struct TVectorPolicy<double, 4> { using Type = FVector4d; };
+template<> struct TVectorPolicy<int32, 4> { using Type = FIntVector4; };
+
+/**
+ * Traits to determine if a type is a valid numeric type for vectors
+ */
+template<typename T>
+struct TIsValidVectorType
+{
+	static constexpr bool Value = 
+		std::is_same<T, float>::value ||
+		std::is_same<T, double>::value ||
+		std::is_same<T, int32>::value;
+};
+
+template<typename T, int32 D>
+struct TChooseVectorType
+{
+	static_assert(TIsValidVectorType<T>::Value, 
+		 "Vector component type must be float, double, or int32");
+	static_assert(D >= 2 && D <= 4,
+		"Vector dimension must be 2, 3, or 4");
+
+	using Type = typename TVectorPolicy<T, D>::Type;
+	static_assert(!std::is_same<Type, void>::value,
+		"Invalid vector type combination");
+};
+
+// Forward declarations
+template<typename AttrType>
+class FAttributeChannel;
+
+template<typename AttrType>
+struct FStorageEntry
+{
+	float Parameter;
+	AttrType Value;
+};
diff --git a/Engine/Source/Runtime/GeometryCore/Public/Splines/Spline.h b/Engine/Source/Runtime/GeometryCore/Public/Splines/Spline.h
new file mode 100644
--- /dev/null
+++ b/Engine/Source/Runtime/GeometryCore/Public/Splines/Spline.h
@@ -0,0 +1,366 @@
+// Copyright Epic Games, Inc. All Rights Reserved.
+
+#pragma once
+
+#include "AttributeChannel.h"
+#include "BaseEvaluator.h"
+
+/** Flags controlling what data gets copied during spline cloning */
+enum class ESplineCloneFlags : uint32
+{
+	None               = 0,
+	CoreData           = 1 << 0,   // Core spline data (points, tangents etc.)
+	Attributes		   = 1 << 1,   // Attribute channels
+    
+	// Common combinations
+	All = CoreData | Attributes,
+	CoreOnly = CoreData
+};
+ENUM_CLASS_FLAGS(ESplineCloneFlags);
+
+// Default evaluator for metadata/attributes that follows parent spline's parameter space
+template <typename ValueType, typename SplineT, int32 SplineD>
+class TAttributeChannelEvaluator final : public TBaseEvaluator<ValueType>
+{
+public:
+	// Make parent interface type accessible
+	using InterfaceType = TBaseEvaluator<ValueType>;
+
+	using SplineValueType = typename TChooseVectorType<SplineT, SplineD>::Type;
+	const TBaseEvaluator<SplineValueType>& ParentEvaluator;
+
+public:
+	explicit TAttributeChannelEvaluator(const TBaseEvaluator<SplineValueType>& InParentEvaluator)
+		: ParentEvaluator(InParentEvaluator)
+	{
+	}
+
+	virtual TUniquePtr<TBaseEvaluator<ValueType>> Clone() const override
+	{
+		return MakeUnique<TAttributeChannelEvaluator<ValueType, SplineT, SplineD>>(ParentEvaluator);
+	}
+	
+	virtual int32 GetDegree() const override
+	{
+		return ParentEvaluator.GetDegree();
+	}
+
+protected:
+	virtual float GetParentSpaceMapping(float NormalizedParam) const override
+	{
+		return ParentEvaluator.MapToParentSpace(NormalizedParam);
+	}
+
+	virtual float GetNormalizedSpaceMapping(float ParentParam) const override
+	{
+		return ParentEvaluator.MapFromParentSpace(ParentParam);
+	}
+
+};
+
+/**
+ * Base spline class with support for varied attribute channels
+ */
+template <typename ScalarType, int32 Dimensions>
+class TSplineBase
+{
+public:
+	using FValue = typename TChooseVectorType<ScalarType, Dimensions>::Type;
+
+protected:
+	TUniquePtr<FAttributeChannel<FValue>> Values;
+	TMap<FName, TUniquePtr<IAttributeChannelBase>> AttributeChannels;
+
+public:
+	explicit TSplineBase(
+		TUniquePtr<TBaseEvaluator<FValue>> Evaluator)
+		: Values(MakeUnique<FAttributeChannel<FValue>>(MoveTemp(Evaluator)))
+	{
+	}
+
+	virtual ~TSplineBase() = default;
+
+	
+	/**
+	  * Creates an attribute channel with a custom evaluator
+	  */
+	template <typename AttrType>
+	void CreateAttributeChannel(
+		const FName& Name,
+		TUniquePtr<TBaseEvaluator<AttrType>> Evaluator)
+	{
+		// Ensure attribute interpolation degree doesn't exceed spline degree
+		const int32 SplineDegree = Values->GetEvaluator().GetDegree();
+		const int32 AttrDegree = Evaluator->GetDegree();
+
+		if (AttrDegree > SplineDegree)
+		{
+			// Log an explicit error
+			UE_LOG(LogSpline, Error, TEXT("Attribute channel '%s' has a degree (%d) greater than the spline degree (%d). This is not allowed."),
+				*Name.ToString(), AttrDegree, SplineDegree);
+			
+			return;
+		}
+
+		TUniquePtr<FAttributeChannel<AttrType>> Channel = MakeUnique<FAttributeChannel<AttrType>>(MoveTemp(Evaluator));
+		AttributeChannels.Add(Name, MoveTemp(Channel));
+	}
+
+	/**
+	 * Creates an attribute channel that follows the spline's evaluation pattern
+	 */
+	template <typename AttrType>
+	void CreateAttributeChannel(const FName& Name)
+	{
+		using AttributeEvaluator = TAttributeChannelEvaluator<AttrType, ScalarType, Dimensions>;
+		const TBaseEvaluator<FValue>& SplineEvaluator = Values->GetEvaluator();
+    
+		// Create the evaluator
+		TUniquePtr<AttributeEvaluator> Evaluator = MakeUnique<AttributeEvaluator>(SplineEvaluator);
+    
+		// Release ownership and create new unique ptr of base type
+		TUniquePtr<TBaseEvaluator<AttrType>> BaseEvaluator(Evaluator.Release());
+    
+		CreateAttributeChannel(Name, MoveTemp(BaseEvaluator));
+	}
+
+	template <typename AttrType>
+	int32 AddAttributeValue(
+		FName Name,
+		const AttrType& Value,
+		float Parameter)
+	{
+		if (FAttributeChannel<AttrType>* Channel = GetAttributeChannel<AttrType>(Name))
+		{
+			return Channel->AddValue(Value, Parameter);
+		}
+		return INDEX_NONE;
+	}
+
+	virtual int32 AddValue(const FValue& Value, float Parameter)
+	{
+		Values->GetEvaluator().ExpandRangeToInclude(Parameter);
+		return Values->AddValue(Value, Parameter);		
+	}
+
+	FValue Evaluate(float Parameter) const
+	{
+		return Values->Evaluate(Parameter);
+	}
+
+	/** 
+	 * Evaluates an attribute at the given parameter in parent space
+	 * @param Name - Name of the attribute channel
+	 * @param Parameter - Parameter value in parent spline space
+	 */
+	template <typename AttrType>
+	AttrType EvaluateAttribute(FName Name, float Parameter) const
+	{
+		if (const FAttributeChannel<AttrType>* Channel = GetAttributeChannel<AttrType>(Name))
+		{
+			// Convert from parent space to channel's local space
+			const float ChannelParam = Channel->GetEvaluator().MapFromParentSpace(Parameter);
+			// Evaluate in channel's local space
+			return Channel->Evaluate(ChannelParam);
+		}
+		return AttrType();
+	}
+
+	virtual void Clear()
+	{
+		Values->Clear();
+		for (TTuple<FName, TUniquePtr<IAttributeChannelBase>>& Pair : AttributeChannels)
+		{
+			static_cast<FAttributeChannel<FValue>*>(Pair.Value.Get())->Clear();
+		}
+	}
+
+	/**
+	 * Clears all values from a specific attribute channel
+	 * @param Name - Name of the attribute channel to clear
+	 * @return true if the channel was found and cleared, false otherwise
+	 */
+	bool ClearAttributeChannel(const FName& Name)
+	{
+		if (const TUniquePtr<IAttributeChannelBase>* Channel = AttributeChannels.Find(Name))
+		{
+			(*Channel)->Clear();
+			return true;
+		}
+		return false;
+	}
+
+	/** 
+	 * Auto-expands parameter range to include all control points.
+	 * Called automatically when adding points outside current range.
+	 */
+	void UpdateParameterRange()
+	{
+		if (GetNumKeys() < 2) return;
+
+		const auto& Entries = Values->GetStorage().GetEntries();
+		const float NewStart = Entries[0].Parameter;
+		const float NewEnd = Entries.Last().Parameter;
+        
+		SetParameterRange(NewStart, NewEnd);
+	}
+	
+	/** 
+	 * Sets the parameter range for evaluation. This remaps how parameters are interpreted
+	 * but does not move control points.
+	 * @param InStart - New start parameter
+	 * @param InEnd - New end parameter 
+	 */
+	void SetParameterRange(float InStart, float InEnd)
+	{
+		Values->GetEvaluator().SetParentSpaceRange(InStart, InEnd);
+	}
+
+	/**
+	 *  Get custom parameter range in parent space
+	 *  @param OutStart - Start of the range
+	 *	@param OutEnd - End of the range	
+	 */ 
+	void GetParameterRange(float& OutStart, float& OutEnd) const
+	{
+		Values->GetEvaluator().GetParentSpaceRange(OutStart, OutEnd);
+	}
+
+	/** Set Pre-Infinity handling mode when out of bounds */ 
+	void SetPreInfinityMode(EOutOfBoundsHandlingMode InMode)
+	{
+		Values->GetEvaluator().SetPreInfinityMode(InMode);
+	}
+
+	/** Set Post-Infinity handling mode when out of bounds */
+	void SetPostInfinityMode(EOutOfBoundsHandlingMode InMode)
+	{
+		Values->GetEvaluator().SetPostInfinityMode(InMode);
+	}
+
+	/** Get the current post-infinity mode */
+	EOutOfBoundsHandlingMode GetPreInfinityMode() const
+	{
+		return Values->GetEvaluator().GetPreInfinityMode();
+	}
+
+	/** Get the current post-infinity mode */
+	EOutOfBoundsHandlingMode GetPostInfinityMode() const
+	{
+		return Values->GetEvaluator().GetPostInfinityMode();
+	}
+	
+	/** Access to the spline's evaluator for attribute channels */
+	const TBaseEvaluator<FValue>& GetSplineEvaluator() const
+	{
+		return Values->GetEvaluator();
+	}
+
+	/** For derived classes that need concrete type access */
+	template <typename EvaluatorType>
+	EvaluatorType& GetEvaluatorAs()
+	{
+		return static_cast<EvaluatorType&>(Values->GetEvaluator());
+	}
+
+	template <typename EvaluatorType>
+	const EvaluatorType& GetEvaluatorAs() const
+	{
+		return static_cast<const EvaluatorType&>(Values->GetEvaluator());
+	}
+
+	int32 GetNumKeys() const
+	{
+		return Values->GetStorage().GetEntries().Num();
+	}
+
+	FValue GetValueAt(int32 Index) const
+	{
+		return Values->GetStorage().GetEntries()[Index].Value;
+	}
+
+	float GetParameterAt(int32 Index) const
+	{
+		return Values->GetStorage().GetEntries()[Index].Parameter;
+	}
+
+	/**
+	* Creates a clone of this spline, copying all specified data.
+	* @param CloneFlags - Flags controlling what data gets copied
+	* @return A new spline containing the copied data
+	*/
+	TSharedPtr<TSplineBase> Clone(ESplineCloneFlags CloneFlags = ESplineCloneFlags::All) const
+	{
+		TSharedPtr<TSplineBase> NewSpline = CopyBaseData(CloneFlags);
+		OnClone(NewSpline, CloneFlags);
+		return NewSpline;
+	}
+	
+protected:
+	template <typename AttrType>
+	FAttributeChannel<AttrType>* GetAttributeChannel(const FName Name)
+	{
+		TUniquePtr<IAttributeChannelBase>* Found = AttributeChannels.Find(Name);
+		return Found ? static_cast<FAttributeChannel<AttrType>*>(Found->Get()) : nullptr;
+	}
+
+	template <typename AttrType>
+	const FAttributeChannel<AttrType>* GetAttributeChannel(const FName Name) const
+	{
+		const TUniquePtr<IAttributeChannelBase>* Found = AttributeChannels.Find(Name);
+		return Found ? static_cast<const FAttributeChannel<AttrType>*>(Found->Get()) : nullptr;
+	}
+
+	/** Hook for derived classes to copy their specific data after base data is copied */
+	virtual void OnClone(TSharedPtr<TSplineBase> NewSpline, ESplineCloneFlags CloneFlags) const {}
+	
+private:
+	/**
+	 * Creates a new instance of the correct derived type for cloning.
+	 */
+	virtual TSharedPtr<TSplineBase> CreateNewInstance() const = 0;
+	
+	/** Creates a new instance and copies common base class state */
+	TSharedPtr<TSplineBase> CopyBaseData(ESplineCloneFlags CloneFlags) const
+	{
+		// Core data is always required for a valid clone
+		check(EnumHasAnyFlags(CloneFlags, ESplineCloneFlags::CoreData));
+        
+		// Create new instance of correct derived type
+        TSharedPtr<TSplineBase> NewSpline = CreateNewInstance();
+        NewSpline->Values = MakeUnique<FAttributeChannel<FValue>>(*Values);
+		
+		// Copy parameter range settings
+		float Start, End;
+		GetParameterRange(Start, End);
+		NewSpline->SetParameterRange(Start, End);
+			
+		// Copy out-of-bounds handling mode
+		NewSpline->SetPostInfinityMode(GetPostInfinityMode());
+		NewSpline->SetPreInfinityMode(GetPreInfinityMode());
+
+		if (EnumHasAnyFlags(CloneFlags, ESplineCloneFlags::Attributes))
+		{
+			// Copy all attribute channels
+			for (const auto& AttributePair : AttributeChannels)
+			{
+				const FName& Name = AttributePair.Key;
+				const TUniquePtr<IAttributeChannelBase>& Channel = AttributePair.Value;
+
+				// Use type erasure to copy the channel without knowing its concrete type
+				if (const FAttributeChannel<FValue>* TypedChannel = static_cast<const FAttributeChannel<FValue>*>(Channel.Get()))
+				{
+					NewSpline->AttributeChannels.Add(Name, MakeUnique<FAttributeChannel<FValue>>(*TypedChannel));
+				}
+			}
+		}
+		return NewSpline;
+	}
+};
+
+
+// Common type definitions for convenience and readability
+using FSpline2f = TSplineBase<float, 2>;
+using FSpline3f = TSplineBase<float, 3>;
+using FSpline2d = TSplineBase<double, 2>;
+using FSpline3d = TSplineBase<double, 3>;
