initial sculpt layers modifier + tool target support

diff --git a/Engine/Plugins/Runtime/MeshModelingToolset/Source/MeshModelingTools/Private/MeshVertexSculptTool.cpp b/Engine/Plugins/Runtime/MeshModelingToolset/Source/MeshModelingTools/Private/MeshVertexSculptTool.cpp
--- a/Engine/Plugins/Runtime/MeshModelingToolset/Source/MeshModelingTools/Private/MeshVertexSculptTool.cpp
+++ b/Engine/Plugins/Runtime/MeshModelingToolset/Source/MeshModelingTools/Private/MeshVertexSculptTool.cpp
@@ -73,7 +73,7 @@ FToolTargetTypeRequirements UMeshVertexSculptToolBuilder::VSculptTypeRequirement
 	UMaterialProvider::StaticClass(),
 	UDynamicMeshProvider::StaticClass(),
 	UDynamicMeshCommitter::StaticClass(),
-	UPrimitiveComponentBackedTarget::StaticClass()
+	USceneComponentBackedTarget::StaticClass()
 });
 
 const FToolTargetTypeRequirements& UMeshVertexSculptToolBuilder::GetTargetRequirements() const
diff --git a/Engine/Source/Runtime/GeometryCore/Public/DynamicMesh/DynamicMeshSculptLayers.h b/Engine/Source/Runtime/GeometryCore/Public/DynamicMesh/DynamicMeshSculptLayers.h
--- a/Engine/Source/Runtime/GeometryCore/Public/DynamicMesh/DynamicMeshSculptLayers.h
+++ b/Engine/Source/Runtime/GeometryCore/Public/DynamicMesh/DynamicMeshSculptLayers.h
@@ -80,6 +80,11 @@ public:
 		UpdateMeshFromLayers();
 	}
 
+	// Update the active layer's sculpt offsets so that the sum of sculpt layers w/ current weights gives the current mesh vertex positions.
+	// Will fail if the active layer has zero weight, or if there are no sculpt layers.
+	// @return true on success
+	GEOMETRYCORE_API bool UpdateLayersFromMesh();
+
 	// Copy across sculpt layer data via a vertex mapping, for all layers that exist on both this and the other layers
 	GEOMETRYCORE_API void CopyThroughMapping(const FDynamicMeshSculptLayers& Other, const FMeshIndexMappings& Mapping);
 
@@ -103,7 +108,6 @@ private:
 	GEOMETRYCORE_API void Enable(FDynamicMeshAttributeSet* AttributeSet, int32 MinLayerCount = -1);
 	GEOMETRYCORE_API void Discard(FDynamicMeshAttributeSet* AttributeSet);
 	GEOMETRYCORE_API bool CheckValidity(const FDynamicMeshAttributeSet* AttributeSet, bool bAllowNonmanifold, EValidityCheckFailMode FailMode) const;
-	GEOMETRYCORE_API bool UpdateLayersFromMesh();
 	GEOMETRYCORE_API bool UpdateMeshFromLayers();
 	GEOMETRYCORE_API bool ValidateActiveLayer();
 
