Added a PolyBezier spline that represents a sequence of connected cubic Bezier curves. Uses composition with BSpline and custom evaluator for implementation.
Added a TangentControlledPolyBezier spline that provides tangent-based control over curve shape while using a PolyBezier spline internally for evaluation. Supports both manual tangent control and automatic tangent computation.
Updated SplineVisualizer in QAGame to visualize/test these.
Removed CubicBezierSpline and HermiteSpline implementations for now.

diff --git a/Engine/Source/Runtime/GeometryCore/Public/Splines/AttributeChannel.h b/Engine/Source/Runtime/GeometryCore/Public/Splines/AttributeChannel.h
--- a/Engine/Source/Runtime/GeometryCore/Public/Splines/AttributeChannel.h
+++ b/Engine/Source/Runtime/GeometryCore/Public/Splines/AttributeChannel.h
@@ -13,6 +13,12 @@ public:
 
 	/** Clears all values from this channel */
     virtual void Clear() = 0;
+
+	/** 
+	 * Creates a deep copy of this channel
+	 * @return A new instance of the channel
+	 */
+	virtual TUniquePtr<IAttributeChannelBase> Clone() const = 0;
 };
 
 /**
@@ -34,20 +40,44 @@ public:
     {
     }
 
+	// Make copying logic live in the channel class
+	virtual TUniquePtr<IAttributeChannelBase> Clone() const override
+	{
+    	return MakeUnique<FAttributeChannel>(*this);
+    }
+	
 	// Add copy constructor to properly clone the evaluator
 	FAttributeChannel(const FAttributeChannel& Other)
 		: Storage(Other.Storage)
 		, FindWindowOverride(Other.FindWindowOverride)
 		, InterpolateOverride(Other.InterpolateOverride)
     {
-    	// Clone the evaluator - we need to implement Clone() in TBaseEvaluator
     	if (Other.Evaluator)
     	{
     		Evaluator = Other.Evaluator->Clone();
     	}
     }
 
-	FAttributeChannel(FAttributeChannel&& Other) noexcept = default;
+	// Add copy assignment
+	FAttributeChannel& operator=(const FAttributeChannel& Other)
+    {
+    	if (this != &Other)
+    	{
+    		Storage = Other.Storage;
+    		FindWindowOverride = Other.FindWindowOverride;
+    		InterpolateOverride = Other.InterpolateOverride;
+
+    		if (Other.Evaluator)
+    		{
+    			Evaluator = Other.Evaluator->Clone();
+    		}
+    		else
+    		{
+    			Evaluator.Reset();
+    		}
+    	}
+    	return *this;
+    }
 
 	/**
 	 * Clears all values from this channel
diff --git a/Engine/Source/Runtime/GeometryCore/Public/Splines/BSpline.h b/Engine/Source/Runtime/GeometryCore/Public/Splines/BSpline.h
--- a/Engine/Source/Runtime/GeometryCore/Public/Splines/BSpline.h
+++ b/Engine/Source/Runtime/GeometryCore/Public/Splines/BSpline.h
@@ -22,11 +22,11 @@ enum class EBSplineKnotMode
 };
 
 // Forward declaration
-template<typename ScalarType, int32 Dimensions, int32 Degree>
+template<typename ScalarType, int32 Dimensions, int32 Degree, typename EvaluatorType>
 class TBSpline;
 
 template <typename ScalarType, int32 Dimensions, int32 Degree>
-class TBSplineEvaluator final : public TBaseEvaluator<typename TChooseVectorType<ScalarType, Dimensions>::Type>
+class TBSplineEvaluator : public TBaseEvaluator<typename TChooseVectorType<ScalarType, Dimensions>::Type>
 {
 	static_assert(Degree > 0, "BSpline degree must be positive");
 	
@@ -49,13 +49,15 @@ private:
 	bool bClampEnds = true;
 
 public:
-	TBSplineEvaluator()
+	TBSplineEvaluator() = default;
+	virtual ~TBSplineEvaluator() override = default;
+
+	virtual TUniquePtr<TBaseEvaluator<FValue>> Clone() const override
 	{
+		return MakeUnique<TBSplineEvaluator>(*this);
 	}
 
-	virtual TArray<const FValue*> FindWindow(
-		TArrayView<const FStorageEntry<FValue>> Entries,
-		float Parameter) const override
+	virtual TArray<const FValue*> FindWindow(TArrayView<const FStorageEntry<FValue>> Entries, float Parameter) const override
 	{
 		TArray<const FValue*> Window;
 		if (Entries.Num() < Degree + 1)
@@ -87,18 +89,56 @@ public:
 		return TSplineInterpolationPolicy<FValue>::InterpolateWithBasis(Window, Basis);
 	}
 
+	virtual FValue EvaluateDerivative(TArrayView<const FStorageEntry<FValue>> Entries, float Parameter) const override 
+	{
+		// Find window using existing window logic
+		TArray<const FValue*> Window = FindWindow(Entries, Parameter);
+		if (Window.Num() < GetDegree() + 1)
+		{
+			return FValue();
+		}
+
+		// Calculate derivative basis functions
+		TArray<float> DerivBasis;
+		const float u = Base::MapToParentSpace(Parameter);
+		const int32 k = FindKnotSpan(u);
+		ComputeDerivativeBasisFunctions(k, u, DerivBasis);
+
+		// Apply basis using type-specific policy interpolation
+		return TSplineInterpolationPolicy<FValue>::InterpolateWithBasis(Window, DerivBasis);
+	}
+
+	virtual FValue EvaluateSecondDerivative(TArrayView<const FStorageEntry<FValue>> Entries, float Parameter) const override
+	{
+		TArray<const FValue*> Window = FindWindow(Entries, Parameter);
+		if (Window.Num() < GetDegree() + 1)
+		{
+			return FValue();
+		}
+
+		// Calculate second derivative basis functions 
+		TArray<float> SecondDerivBasis;
+		const float u = Base::MapToParentSpace(Parameter);
+		const int32 k = FindKnotSpan(u);
+		ComputeSecondDerivativeBasisFunctions(k, u, SecondDerivBasis);
+
+		// Apply basis using type-specific policy interpolation
+		return TSplineInterpolationPolicy<FValue>::InterpolateWithBasis(Window, SecondDerivBasis);
+	}
+
 	virtual int32 GetDegree() const override
 	{
 		return Degree;
 	}
 
+
 	/**
 	 * Sets a custom knot vector and switches to custom knot mode
 	 * @param NewKnots - Array of knot values. Must be non-decreasing and have correct length (n + p + 1)
 	 * @param InNumValues - Number of values (control points) in the spline
 	 * @return True if the knot vector was valid and set successfully
 	 */
-	bool SetCustomKnots(const TArray<float>& NewKnots, int32 InNumValues)
+	virtual bool SetCustomKnots(const TArray<float>& NewKnots, int32 InNumValues)
 	{
 		NumValues = InNumValues;
 		const int32 RequiredKnots = NumValues + Degree + 1;
@@ -235,11 +275,7 @@ protected:
 	{
 		if (KnotVector.Num() < 2) return NormalizedParam;
 
-		float Start, End;
-		GetKnotRange(Start, End);
-
-		// Map normalized parameter to parent space
-		return FMath::Lerp(Start, End, NormalizedParam);
+		return ToKnotSpace(NormalizedParam);
 	}
 
 	virtual float GetNormalizedSpaceMapping(float KnotParam) const override
@@ -377,6 +413,85 @@ private:
 		}
 	}
 
+	void ComputeDerivativeBasisFunctions(int32 Span, float u, TArray<float>& OutDerivBasis) const
+	{
+	    const int32 p = GetDegree();
+	    OutDerivBasis.SetNum(p + 1);
+
+	    // Get regular basis functions
+	    TArray<float> N;
+	    N.SetNum(p + 1);
+	    TArray<float> left, right;
+	    left.SetNum(p + 1);
+	    right.SetNum(p + 1);
+	    
+	    // Initialize zeroth basis function
+	    N[0] = 1.0f;
+
+	    // Compute basis functions
+	    for (int32 j = 1; j <= p; ++j)
+	    {
+	        left[j] = u - KnotVector[Span + 1 - j];
+	        right[j] = KnotVector[Span + j] - u;
+	        
+	        float saved = 0.0f;
+	        for (int32 r = 0; r < j; ++r)
+	        {
+	            float temp = N[r] / (right[r + 1] + left[j - r]);
+	            N[r] = saved + right[r + 1] * temp;
+	            saved = left[j - r] * temp;
+	        }
+	        N[j] = saved;
+	    }
+
+	    // Now compute derivatives using formula from NURBS book:
+	    // N'(i,p) = p * (N(i,p-1)/(u(i+p)-u(i)) - N(i+1,p-1)/(u(i+p+1)-u(i+1)))
+	    for (int32 j = 0; j <= p; ++j)
+	    {
+	        float diff1 = KnotVector[Span + p] - KnotVector[Span + j];
+	        float diff2 = KnotVector[Span + j + p + 1] - KnotVector[Span + j + 1];
+	        
+	        if (FMath::Abs(diff1) < UE_SMALL_NUMBER)
+	        {
+	            OutDerivBasis[j] = 0.0f;
+	            continue;
+	        }
+	        
+	        float der1 = N[j] / diff1;
+	        float der2 = (j + 1 <= p) ? N[j + 1] / diff2 : 0.0f;
+	        OutDerivBasis[j] = p * (der1 - der2);
+	    }
+	}
+
+	void ComputeSecondDerivativeBasisFunctions(int32 Span, float u, TArray<float>& OutSecondDerivBasis) const
+	{
+	    const int32 p = GetDegree();
+	    OutSecondDerivBasis.SetNum(p + 1);
+
+	    // First get first derivative basis
+	    TArray<float> FirstDerivBasis;
+	    ComputeDerivativeBasisFunctions(Span, u, FirstDerivBasis);
+
+	    // Apply derivative formula again to get second derivatives
+	    // N''(i,p) = p * (N'(i,p-1)/(u(i+p)-u(i)) - N'(i+1,p-1)/(u(i+p+1)-u(i+1)))
+	    const float factor = p - 1;
+	    for (int32 j = 0; j <= p; ++j)
+	    {
+	        float diff1 = KnotVector[Span + p] - KnotVector[Span + j];
+	        float diff2 = KnotVector[Span + j + p + 1] - KnotVector[Span + j + 1];
+	        
+	        if (FMath::Abs(diff1) < UE_SMALL_NUMBER)
+	        {
+	            OutSecondDerivBasis[j] = 0.0f;
+	            continue;
+	        }
+	        
+	        float der1 = FirstDerivBasis[j] / diff1;
+	        float der2 = (j + 1 <= p) ? FirstDerivBasis[j + 1] / diff2 : 0.0f;
+	        OutSecondDerivBasis[j] = factor * (der1 - der2);
+	    }
+	}
+
 	/** 
      * Generates a uniform knot vector with equal spacing
      * For clamped ends, multiplicity of degree+1 is used at endpoints
@@ -536,48 +651,33 @@ private:
 
 };
 
-template <typename ScalarType, int32 Dimensions, int32 Degree = 3>
-class TBSpline : public TSplineBase<ScalarType, Dimensions>
+template <typename ScalarType, int32 Dimensions, int32 Degree = 3,
+			typename EvaluatorType = TBSplineEvaluator<ScalarType, Dimensions, Degree>>
+class TBSpline : public TSplineBase<ScalarType, Dimensions, EvaluatorType>
 {
 	static_assert(Degree > 0, "BSpline degree must be positive");
     
-	using Base = TSplineBase<ScalarType, Dimensions>;
+	using Base = TSplineBase<ScalarType, Dimensions, EvaluatorType>;
 	
 protected:
 	virtual TSharedPtr<Base> CreateNewInstance() const override
 	{
 		return MakeShared<TBSpline>();
 	}
-    
-	virtual void OnClone(TSharedPtr<Base> NewSpline, ESplineCloneFlags CloneFlags) const override
-	{
-		TSharedPtr<TBSpline> TypedSpline = StaticCastSharedPtr<TBSpline>(NewSpline);
-		TypedSpline->GetEvaluator().SetKnotGenerationMode(this->GetEvaluator().GetKnotMode());
-		TypedSpline->GetEvaluator().SetClampedEnds(this->GetEvaluator().IsClampedEnds());
-        
-		if (this->GetEvaluator().GetKnotMode() == EBSplineKnotMode::Custom)
-		{
-			TypedSpline->SetCustomKnots(this->GetEvaluator().GetKnotVector());
-		}
-	}
-	
 public:
 	using typename Base::FValue;
-
-	TBSpline()
-		: Base(MakeUnique<TBSplineEvaluator<ScalarType, Dimensions, Degree>>())
-	{
-	}
-
-	TBSplineEvaluator<ScalarType, Dimensions, Degree>& GetEvaluator()
+	explicit TBSpline(TUniquePtr<EvaluatorType> InEvaluator = nullptr)
+			: Base(MoveTemp(InEvaluator))
 	{
-		return this->template GetEvaluatorAs<TBSplineEvaluator<ScalarType, Dimensions, Degree>>();
 	}
+	
+	TBSpline(const TBSpline&) = default;
+	TBSpline& operator=(const TBSpline&) = default;
+	TBSpline(TBSpline&&) noexcept = default;
+	TBSpline& operator=(TBSpline&&) noexcept = default;
+	
+	virtual ~TBSpline() override = default;
 
-	const TBSplineEvaluator<ScalarType, Dimensions, Degree>& GetEvaluator() const
-	{
-		return this->template GetEvaluatorAs<TBSplineEvaluator<ScalarType, Dimensions, Degree>>();
-	}
 
 	/**
 	 * Sets a custom knot vector for the spline
@@ -586,7 +686,7 @@ public:
 	 */
 	void SetCustomKnots(const TArray<float>& NewKnots)
 	{
-		GetEvaluator().SetCustomKnots(NewKnots, Base::GetNumKeys());
+		Base::GetEvaluator().SetCustomKnots(NewKnots, Base::GetNumKeys());
 	}
 
 	/**
@@ -597,11 +697,14 @@ public:
 	 */
 	virtual int32 AddValue(const FValue& Value, float Parameter) override
 	{
+		UE_LOG(LogSpline, Verbose, TEXT("TBSpline::AddValue - Evaluator: %p, KnotMode: %d"), 
+			&this->GetEvaluator(), static_cast<int32>(this->GetEvaluator().GetKnotMode()));
+		
 		const int32 Index = Base::AddValue(Value, Parameter);
 		// Update knot vector unless using custom knots
-		if (GetEvaluator().GetKnotMode() != EBSplineKnotMode::Custom)
+		if (Base::GetEvaluator().GetKnotMode() != EBSplineKnotMode::Custom)
 		{
-			GetEvaluator().UpdateKnotVector(Base::Values->GetStorage().GetEntries());
+			Base::GetEvaluator().UpdateKnotVector(Base::Values.GetStorage().GetEntries());
 		}
 		return Index;
 	}
@@ -612,8 +715,8 @@ public:
 	 */
 	void SetKnotGenerationMode(EBSplineKnotMode Mode)
 	{
-		GetEvaluator().SetKnotGenerationMode(Mode);
-		GetEvaluator().UpdateKnotVector(Base::Values->GetStorage().GetEntries());
+		Base::GetEvaluator().SetKnotGenerationMode(Mode);
+		Base::GetEvaluator().UpdateKnotVector(Base::Values.GetStorage().GetEntries());
 	}
 
 	/**
@@ -622,8 +725,8 @@ public:
 	 */
 	void SetClampedEnds(bool bInClampEnds)
 	{
-		GetEvaluator().SetClampedEnds(bInClampEnds);
-		GetEvaluator().UpdateKnotVector(Base::Values->GetStorage().GetEntries());
+		Base::GetEvaluator().SetClampedEnds(bInClampEnds);
+		Base::GetEvaluator().UpdateKnotVector(Base::Values.GetStorage().GetEntries());
 	}
 };
 
diff --git a/Engine/Source/Runtime/GeometryCore/Public/Splines/BaseEvaluator.h b/Engine/Source/Runtime/GeometryCore/Public/Splines/BaseEvaluator.h
--- a/Engine/Source/Runtime/GeometryCore/Public/Splines/BaseEvaluator.h
+++ b/Engine/Source/Runtime/GeometryCore/Public/Splines/BaseEvaluator.h
@@ -34,6 +34,9 @@ protected:
 	EOutOfBoundsHandlingMode PostInfinityMode = EOutOfBoundsHandlingMode::Zero;
 
 public:
+	TBaseEvaluator()
+	{
+	}
 	virtual ~TBaseEvaluator() = default;
 	
 	/** Get the polynomial degree of the evaluator */
@@ -95,6 +98,44 @@ public:
 		return TSplineInterpolationPolicy<ValueType>::Interpolate(Window, MappedParam);
 	}
 
+	/** Evaluates first derivative at parameter */
+	virtual ValueType EvaluateDerivative(
+		TArrayView<const FStorageEntry<ValueType>> Entries,
+		float Parameter) const
+	{
+		TArray<const ValueType*> Window = FindWindow(Entries, Parameter);
+		if (Window.Num() < 2)
+		{
+			return ValueType();
+		}
+
+		float Start, End;
+		GetParentSpaceRange(Start, End);
+		const float ParamSpaceScale = 1.0f / (End - Start);
+
+		return TSplineInterpolationPolicy<ValueType>::CalculateDerivative(
+			*Window[0], *Window[1], ParamSpaceScale);
+	}
+
+	/** Evaluates second derivative at parameter */
+	virtual ValueType EvaluateSecondDerivative(
+		TArrayView<const FStorageEntry<ValueType>> Entries,
+		float Parameter) const
+	{
+		TArray<const ValueType*> Window = FindWindow(Entries, Parameter);
+		if (Window.Num() < 3)
+		{
+			return ValueType();
+		}
+
+		float Start, End;
+		GetParentSpaceRange(Start, End);
+		const float ParamSpaceScale = 1.0f / FMath::Square(End - Start);
+
+		return TSplineInterpolationPolicy<ValueType>::CalculateSecondDerivative(
+			*Window[0], *Window[1], *Window[2], ParamSpaceScale);
+	}
+
 	/** Clone the evaluator */
 	virtual TUniquePtr<TBaseEvaluator<ValueType>> Clone() const
 	{
diff --git a/Engine/Source/Runtime/GeometryCore/Public/Splines/CubicBezierSpline.h b/Engine/Source/Runtime/GeometryCore/Public/Splines/CubicBezierSpline.h
deleted file mode 100644
diff --git a/Engine/Source/Runtime/GeometryCore/Public/Splines/HermiteSpline.h b/Engine/Source/Runtime/GeometryCore/Public/Splines/HermiteSpline.h
deleted file mode 100644
diff --git a/Engine/Source/Runtime/GeometryCore/Public/Splines/InterpolationPolicies.h b/Engine/Source/Runtime/GeometryCore/Public/Splines/InterpolationPolicies.h
--- a/Engine/Source/Runtime/GeometryCore/Public/Splines/InterpolationPolicies.h
+++ b/Engine/Source/Runtime/GeometryCore/Public/Splines/InterpolationPolicies.h
@@ -36,57 +36,18 @@ public:
         return Result;
     }
 
-    // Helper for auto-tangent computation
-    static T BlendTangents(const T& A, const T& B, float Alpha = 0.5f)
+	// Default derivative calculation
+	static T CalculateDerivative(const T& Start, const T& End, float ParamSpaceScale)
     {
-        return Interpolate({ &A, &B }, Alpha);
+    	return (End - Start) * ParamSpaceScale;
     }
-};
-
-// FColor interpolation policy
-template<>
-class TSplineInterpolationPolicy<FColor>
-{
-public:
-    static FColor Interpolate(TArrayView<const FColor* const> Window, float Parameter)
-    {
-        if (Window.Num() < 2)
-        {
-            return Window.Num() == 1 ? *Window[0] : FColor();
-        }
-        
-        // Convert to linear color for better interpolation
-        return FLinearColor::LerpUsingHSV(
-            FLinearColor(*Window[0]),
-            FLinearColor(*Window[1]),
-            Parameter
-        ).ToFColor(true);
-    }
-
-    static FColor InterpolateWithBasis(TArrayView<const FColor* const> Window, TArrayView<const float> Basis)
-    {
-        // Convert all colors to linear space first
-        TArray<FLinearColor> LinearColors;
-        LinearColors.Reserve(Window.Num());
-        for (const FColor* Color : Window)
-        {
-            LinearColors.Add(FLinearColor(*Color));
-        }
 
-        // Perform weighted sum in linear space
-        FLinearColor Result = FLinearColor::Black;
-        for (int32 i = 0; i < FMath::Min(Window.Num(), Basis.Num()); ++i)
-        {
-            Result += LinearColors[i] * Basis[i];
-        }
-
-        // Convert back to FColor
-        return Result.ToFColor(true);
-    }
-
-    static FColor BlendTangents(const FColor& A, const FColor& B, float Alpha = 0.5f)
+	// Default second derivative calculation
+	static T CalculateSecondDerivative(const T& Start, const T& Mid, const T& End, float ParamSpaceScale)
     {
-        return FLinearColor::LerpUsingHSV(FLinearColor(A), FLinearColor(B), Alpha).ToFColor(true);
+    	T FirstDeriv1 = Mid - Start;
+    	T FirstDeriv2 = End - Mid;
+    	return (FirstDeriv2 - FirstDeriv1) * ParamSpaceScale;
     }
 };
 
@@ -131,9 +92,16 @@ public:
         return Result;
     }
 
-    static FLinearColor BlendTangents(const FLinearColor& A, const FLinearColor& B, float Alpha = 0.5f)
+	static FLinearColor CalculateDerivative(const FLinearColor& Start, const FLinearColor& End, float ParamSpaceScale)
+    {
+    	return (End - Start) * ParamSpaceScale;
+    }
+
+	static FLinearColor CalculateSecondDerivative(const FLinearColor& Start, const FLinearColor& Mid, const FLinearColor& End, float ParamSpaceScale)
     {
-        return FLinearColor::LerpUsingHSV(A, B, Alpha);
+    	FLinearColor FirstDeriv1 = Mid - Start;
+    	FLinearColor FirstDeriv2 = End - Mid;
+    	return (FirstDeriv2 - FirstDeriv1) * ParamSpaceScale;
     }
 };
 
@@ -177,10 +145,40 @@ public:
         return Result;
     }
 
-    static FQuat BlendTangents(const FQuat& A, const FQuat& B, float Alpha = 0.5f)
+	static FQuat CalculateDerivative(const FQuat& Start, const FQuat& End, float ParamSpaceScale)
     {
-        return FQuat::Slerp(A, B, Alpha);
+    	const double Dot = Start | End;
+    	const FQuat TargetQuat = (Dot < 0) ? -End : End;
+    	FQuat Delta = TargetQuat * Start.Inverse();
+        
+    	FVector Axis;
+    	float Angle;
+    	Delta.ToAxisAndAngle(Axis, Angle);
+        
+    	const FVector AngularVelocity = Axis * Angle * ParamSpaceScale;
+    	return FQuat(AngularVelocity.X, AngularVelocity.Y, AngularVelocity.Z, 0.0f);
     }
+
+	static FQuat CalculateSecondDerivative(const FQuat& Start, const FQuat& Mid, const FQuat& End, float ParamSpaceScale)
+    {
+    	auto CalcAngularVelocity = [](const FQuat& q0, const FQuat& q1) -> FVector
+    	{
+    		const double Dot = q0 | q1;
+    		const FQuat TargetQuat = (Dot < 0) ? -q1 : q1;
+    		FQuat Delta = TargetQuat * q0.Inverse();
+            
+    		FVector Axis;
+    		float Angle;
+    		Delta.ToAxisAndAngle(Axis, Angle);
+    		return Axis * Angle;
+    	};
+
+    	FVector AngVel1 = CalcAngularVelocity(Start, Mid);
+    	FVector AngVel2 = CalcAngularVelocity(Mid, End);
+    	const FVector AngularVelocity = (AngVel2 - AngVel1) * ParamSpaceScale;
+    	return FQuat(AngularVelocity.X, AngularVelocity.Y, AngularVelocity.Z, 0.0f);
+    }
+	
 };
 
 // FRotator interpolation policy
@@ -216,10 +214,27 @@ public:
 
         return FRotator(Result);
     }
+	static FRotator CalculateDerivative(const FRotator& Start, const FRotator& End, float ParamSpaceScale)
+    {
+    	// Convert to quaternions for better interpolation
+    	FQuat StartQuat = Start.Quaternion();
+    	FQuat EndQuat = End.Quaternion();
+        
+    	// Use quaternion derivative calculation and convert back to rotator
+    	FQuat QuatDeriv = TSplineInterpolationPolicy<FQuat>::CalculateDerivative(StartQuat, EndQuat, ParamSpaceScale);
+    	return FRotator(QuatDeriv);
+    }
 
-    static FRotator BlendTangents(const FRotator& A, const FRotator& B, float Alpha = 0.5f)
+	static FRotator CalculateSecondDerivative(const FRotator& Start, const FRotator& Mid, const FRotator& End, float ParamSpaceScale)
     {
-        return FMath::LerpRange(A, B, Alpha);
+    	// Convert to quaternions for better interpolation
+    	FQuat StartQuat = Start.Quaternion();
+    	FQuat MidQuat = Mid.Quaternion();
+    	FQuat EndQuat = End.Quaternion();
+        
+    	// Use quaternion second derivative calculation and convert back to rotator
+    	FQuat QuatSecondDeriv = TSplineInterpolationPolicy<FQuat>::CalculateSecondDerivative(StartQuat, MidQuat, EndQuat, ParamSpaceScale);
+    	return FRotator(QuatSecondDeriv);
     }
 };
 
@@ -287,44 +302,39 @@ public:
         return FTransform(Rotation, Location, Scale);
     }
 
-    static FTransform BlendTangents(const FTransform& A, const FTransform& B, float Alpha = 0.5f)
-    {
-        FTransform Result;
-        Result.SetLocation(FMath::Lerp(A.GetLocation(), B.GetLocation(), Alpha));
-        Result.SetRotation(FQuat::Slerp(A.GetRotation(), B.GetRotation(), Alpha));
-        Result.SetScale3D(FMath::Lerp(A.GetScale3D(), B.GetScale3D(), Alpha));
-        return Result;
-    }
-};
-
-// Specialization for vectors
-template<>
-class TSplineInterpolationPolicy<FVector3f>
-{
-public:
-    // Classic two-point interpolation
-    static FVector3f Interpolate(TArrayView<const FVector3f* const> Window, float Parameter)
-    {
-        if (Window.Num() < 2)
-        {
-            return Window.Num() == 1 ? *Window[0] : FVector3f::ZeroVector;
-        }
-        return FMath::Lerp(*Window[0], *Window[1], Parameter);
-    }
-
-    // Window-based interpolation for splines
-    static FVector3f InterpolateWithBasis(TArrayView<const FVector3f* const> Window, TArrayView<const float> Basis)
+	static FTransform CalculateDerivative(const FTransform& Start, const FTransform& End, float ParamSpaceScale)
     {
-        FVector3f Result = FVector3f::ZeroVector;
-        for (int32 i = 0; i < FMath::Min(Window.Num(), Basis.Num()); ++i)
-        {
-            Result += *Window[i] * Basis[i];
-        }
-        return Result;
+    	FTransform Result;
+    	Result.SetLocation((End.GetLocation() - Start.GetLocation()) * ParamSpaceScale);
+        
+    	// Use quaternion derivative for rotation
+    	const FQuat& StartRot = Start.GetRotation();
+    	const FQuat& EndRot = End.GetRotation();
+    	Result.SetRotation(TSplineInterpolationPolicy<FQuat>::CalculateDerivative(StartRot, EndRot, ParamSpaceScale));
+        
+    	Result.SetScale3D((End.GetScale3D() - Start.GetScale3D()) * ParamSpaceScale);
+    	return Result;
     }
 
-    static FVector3f BlendTangents(const FVector3f& A, const FVector3f& B, float Alpha = 0.5f)
+	static FTransform CalculateSecondDerivative(const FTransform& Start, const FTransform& Mid, const FTransform& End, float ParamSpaceScale)
     {
-        return FMath::Lerp(A, B, Alpha);
+    	FTransform Result;
+        
+    	// Location and scale use standard second derivative
+    	FVector LocationDeriv1 = Mid.GetLocation() - Start.GetLocation();
+    	FVector LocationDeriv2 = End.GetLocation() - Mid.GetLocation();
+    	Result.SetLocation((LocationDeriv2 - LocationDeriv1) * ParamSpaceScale);
+
+    	// Use quaternion second derivative for rotation
+    	const FQuat& StartRot = Start.GetRotation();
+    	const FQuat& MidRot = Mid.GetRotation();
+    	const FQuat& EndRot = End.GetRotation();
+    	Result.SetRotation(TSplineInterpolationPolicy<FQuat>::CalculateSecondDerivative(StartRot, MidRot, EndRot, ParamSpaceScale));
+
+    	FVector ScaleDeriv1 = Mid.GetScale3D() - Start.GetScale3D();
+    	FVector ScaleDeriv2 = End.GetScale3D() - Mid.GetScale3D();
+    	Result.SetScale3D((ScaleDeriv2 - ScaleDeriv1) * ParamSpaceScale);
+
+    	return Result;
     }
 };
\ No newline at end of file
diff --git a/Engine/Source/Runtime/GeometryCore/Public/Splines/ParameterizedStorage.h b/Engine/Source/Runtime/GeometryCore/Public/Splines/ParameterizedStorage.h
--- a/Engine/Source/Runtime/GeometryCore/Public/Splines/ParameterizedStorage.h
+++ b/Engine/Source/Runtime/GeometryCore/Public/Splines/ParameterizedStorage.h
@@ -17,7 +17,7 @@ public:
         const int32 Index = Algo::LowerBound(
             Entries, Parameter,
             [](const FStorageEntry<AttrType>& Entry, float Param) { 
-                return Entry.Parameter < Param; 
+                return Entry.Parameter <= Param; 
             }
         );
         Entries.Insert({Parameter, Value}, Index);
diff --git a/Engine/Source/Runtime/GeometryCore/Public/Splines/PolyBezierSpline.h b/Engine/Source/Runtime/GeometryCore/Public/Splines/PolyBezierSpline.h
new file mode 100644
--- /dev/null
+++ b/Engine/Source/Runtime/GeometryCore/Public/Splines/PolyBezierSpline.h
@@ -0,0 +1,309 @@
+// Copyright Epic Games, Inc. All Rights Reserved.
+
+#pragma once
+
+#include "BSpline.h"
+
+/**
+ * Custom evaluator for PolyBezier splines.
+ * Handles parameter mapping and control point selection for proper Bezier evaluation.
+ */
+template<typename T, int32 D>
+class TPolyBezierEvaluator : public TBSplineEvaluator<T, D, 3>
+{
+	using FValue = typename TChooseVectorType<T, D>::Type;
+	using Base = TBSplineEvaluator<T, D, 3>;
+
+private:
+    float NumSegments = 0.0f;
+    
+public:
+    TPolyBezierEvaluator():
+        Base()
+    {
+        this->SetKnotGenerationMode(EBSplineKnotMode::Custom);
+        this->SetClampedEnds(true);
+    }
+    
+    virtual TUniquePtr<TBaseEvaluator<FValue>> Clone() const override
+    {
+        return MakeUnique<TPolyBezierEvaluator<T, D>>(*this);
+    }
+    
+    virtual int32 GetDegree() const override { return 3; }
+
+    virtual TArray<const FValue*> FindWindow(TArrayView<const FStorageEntry<FValue>> Entries, float Parameter) const override
+    {
+        TArray<const FValue*> Window;
+        if (Entries.Num() < 4)  // Need at least one complete Bezier segment
+        {
+            return Window;
+        }
+
+        // Map parameter to parent space and handle bounds
+        const float ParentParam = this->MapToParentSpace(Parameter);
+        
+        // Find which segment we're in
+        const float ClampedParam = FMath::Clamp(ParentParam, 0.0f, NumSegments);
+        const int32 SegmentIndex = FMath::Min(FMath::FloorToInt(ClampedParam), static_cast<int32>(NumSegments) - 1);
+        
+        // Get the four control points for this segment
+        const int32 BaseIndex = SegmentIndex * 4;
+        if (BaseIndex + 3 >= Entries.Num())
+        {
+            return Window;
+        }
+
+        // Add the four Bezier control points to the window
+        Window.Add(&Entries[BaseIndex].Value);     // P0 need 
+        Window.Add(&Entries[BaseIndex + 1].Value); // P1
+        Window.Add(&Entries[BaseIndex + 2].Value); // P2
+        Window.Add(&Entries[BaseIndex + 3].Value); // P3
+
+        return Window;
+    }
+
+    virtual FValue InterpolateWindow(TArrayView<const FValue* const> Window, float Parameter) const override
+    {
+        if (Window.Num() < 4)
+        {
+            return FValue();
+        }
+
+        // Map to parent space (segment space)
+        const float ParentParam = this->MapToParentSpace(Parameter);
+        
+        // If we're at the last point exactly, just return it
+        if (ParentParam >= NumSegments)
+        {
+            return *Window[3];  // Last control point (P3) of final Bezier segment
+        }
+        
+        // Get local parameter within segment [0,1]
+        const float LocalT = FMath::Frac(ParentParam);
+
+        // Calculate Bernstein basis functions
+        const float t = LocalT;
+        const float OneMinusT = 1.0f - t;
+        const float t2 = t * t;
+        const float OneMinusT2 = OneMinusT * OneMinusT;
+
+        TArray<float> Basis;
+        Basis.SetNum(4);
+        
+        // Cubic Bernstein polynomials
+        Basis[0] = OneMinusT * OneMinusT2;     // (1-t)³
+        Basis[1] = 3.0f * t * OneMinusT2;      // 3t(1-t)²
+        Basis[2] = 3.0f * t2 * OneMinusT;      // 3t²(1-t)
+        Basis[3] = t * t2;                     // t³
+
+        return TSplineInterpolationPolicy<FValue>::InterpolateWithBasis(Window, Basis);
+    }
+
+    virtual bool SetCustomKnots(const TArray<float>& NewKnots, int32 InNumValues) override
+    {
+        // Update number of segments based on control points (4 points per segment)
+        NumSegments = InNumValues / 4;
+        
+        // Generate appropriate knot vector for PolyBezier
+        TArray<float> BezierKnots;
+        const int32 RequiredKnots = InNumValues + 4;  // n + p + 1 where p = 3
+        BezierKnots.SetNum(RequiredKnots);
+
+        // Set up knot vector with multiple knots at segment boundaries
+        int32 CurrentKnot = 0;
+
+        // Start knots with multiplicity 4
+        for (int32 i = 0; i < 4; ++i)
+        {
+            BezierKnots[CurrentKnot++] = 0.0f;
+        }
+
+        // Interior knots at segment boundaries with multiplicity 3
+        for (int32 i = 1; i < NumSegments; ++i)
+        {
+            const float t = static_cast<float>(i);
+            for (int32 j = 0; j < 3; ++j)
+            {
+                BezierKnots[CurrentKnot++] = t;
+            }
+        }
+
+        // End knots with multiplicity 4
+        const float EndKnot = static_cast<float>(NumSegments);
+        while (CurrentKnot < RequiredKnots)
+        {
+            BezierKnots[CurrentKnot++] = EndKnot;
+        }
+
+        return Base::SetCustomKnots(BezierKnots, InNumValues);
+    }
+    
+    FValue EvaluateDerivative(TArrayView<const FStorageEntry<FValue>> Entries, float Parameter) const override
+    {
+        TArray<const FValue*> Window = FindWindow(Entries, Parameter);
+        if (Window.Num() < 4)  // Need all 4 control points for a cubic segment
+        {
+            return FValue();
+        }
+
+        // Get local parameter within segment [0,1]
+        const float t = FMath::Frac(Base::MapToParentSpace(Parameter));
+
+        // First derivative coefficients for cubic Bezier
+        const float c0 = -3.0f * (1.0f-t) * (1.0f-t);
+        const float c1 = 3.0f * (1.0f-t) * (1.0f-3.0f*t);
+        const float c2 = 3.0f * t * (2.0f-3.0f*t);
+        const float c3 = 3.0f * t * t;
+
+        // Build basis array
+        TArray<float> DerivBasis = { c0, c1, c2, c3 };
+
+        // Scale for parameter space and apply basis
+        return TSplineInterpolationPolicy<FValue>::InterpolateWithBasis(Window, DerivBasis) * NumSegments;
+    }
+
+    FValue EvaluateSecondDerivative(TArrayView<const FStorageEntry<FValue>> Entries, float Parameter) const override
+    {
+        TArray<const FValue*> Window = FindWindow(Entries, Parameter);
+        if (Window.Num() < 4)
+        {
+            return FValue();
+        }
+
+        // Get local parameter for this segment
+        const float t = FMath::Frac(Base::MapToParentSpace(Parameter));
+
+        // Second derivative coefficients for cubic Bezier
+        const float c0 = 6.0f * (1.0f-t);
+        const float c1 = -12.0f + 18.0f*t;
+        const float c2 = 6.0f - 18.0f*t;
+        const float c3 = 6.0f * t;
+
+        // Build basis array
+        TArray<float> SecondDerivBasis = { c0, c1, c2, c3 };
+
+        // Scale for parameter space and apply basis
+        const float SegmentScale = NumSegments;
+        return TSplineInterpolationPolicy<FValue>::InterpolateWithBasis(Window, SecondDerivBasis) 
+            * (SegmentScale * SegmentScale);
+    }
+
+protected:
+    virtual float GetParentSpaceMapping(float NormalizedParam) const override
+    {
+        return NormalizedParam * NumSegments;
+    }
+
+    virtual float GetNormalizedSpaceMapping(float ParentParam) const override
+    {
+        return NumSegments > 0.0f ? ParentParam / NumSegments : 0.0f;
+    }
+};
+
+/**
+ * A PolyBezier spline that represents a sequence of connected cubic Bezier curves.
+ * Uses composition with BSpline and custom evaluator for implementation.
+ */
+template<typename ScalarType, int32 Dimensions>
+class TPolyBezierSpline: public TSplineBase<ScalarType, Dimensions, TPolyBezierEvaluator<ScalarType, Dimensions>>
+{
+public:
+    using FValue = typename TChooseVectorType<ScalarType, Dimensions>::Type;
+    using FEvaluator = TPolyBezierEvaluator<ScalarType, Dimensions>;
+    using Base = TSplineBase<ScalarType, Dimensions, TPolyBezierEvaluator<ScalarType, Dimensions>>;
+
+    TPolyBezierSpline()
+        : Base(MakeUnique<TPolyBezierEvaluator<ScalarType, Dimensions>>())
+    {
+    }
+    
+    // Virtual destructor
+    virtual ~TPolyBezierSpline() = default;
+
+    virtual TSharedPtr<Base> CreateNewInstance() const override
+    {
+        return MakeShared<TPolyBezierSpline>();
+    }
+    /**
+     * Adds a new Bezier segment to the spline
+     * @param P0 Start point
+     * @param P1 First control point
+     * @param P2 Second control point
+     * @param P3 End point
+     * @return Index of the new segment
+     */
+    int32 AddBezierSegment(const FValue& P0, const FValue& P1, const FValue& P2, const FValue& P3)
+    {
+        const int32 NumPoints = Base::GetNumKeys();
+        
+        // Validate continuity with previous segment if this isn't the first one
+        if (NumPoints > 0)
+        {
+            const FValue& LastPoint = Base::GetValueAt(NumPoints - 1);
+            if (!LastPoint.Equals(P0))
+            {
+                UE_LOG(LogSpline, Warning, TEXT("PolyBezier segment start point doesn't match previous end point"));
+                return INDEX_NONE;
+            }
+        }
+
+        // Add control points using base storage
+        const float SegmentT = static_cast<float>(NumPoints / 4);
+        Base::AddValue(P0, SegmentT);
+        Base::AddValue(P1, SegmentT);
+        Base::AddValue(P2, SegmentT);
+        Base::AddValue(P3, SegmentT);
+
+        // Let evaluator handle knot vector update
+        Base::GetEvaluator().SetCustomKnots(TArray<float>(), Base::GetNumKeys());
+        
+        return NumPoints / 4;
+    }
+
+    /** Gets the control points for a specific Bezier segment */
+    bool GetSegmentControlPoints(int32 SegmentIndex, FValue& OutP0, FValue& OutP1, 
+        FValue& OutP2, FValue& OutP3) const
+    {
+        if (!IsValidSegmentIndex(SegmentIndex))
+        {
+            return false;
+        }
+
+        const int32 BaseIndex = SegmentIndex * 4;
+        OutP0 = Base::GetValueAt(BaseIndex);
+        OutP1 = Base::GetValueAt(BaseIndex + 1);
+        OutP2 = Base::GetValueAt(BaseIndex + 2);
+        OutP3 = Base::GetValueAt(BaseIndex + 3);
+        return true;
+    }
+
+    /** Returns the number of Bezier segments in the spline */
+    int32 GetNumSegments() const
+    {
+        return Base::GetNumKeys() / 4;
+    }
+
+    /** Gets the endpoint of a segment */
+    FValue GetSegmentEndPoint(int32 SegmentIndex) const
+    {
+        if (!IsValidSegmentIndex(SegmentIndex))
+        {
+            return FValue();
+        }
+        return Base::GetValueAt(SegmentIndex * 4 + 3);
+    }
+
+private:
+
+     bool IsValidSegmentIndex(int32 Index) const
+    {
+        return Index >= 0 && Index < GetNumSegments();
+    }
+};
+
+// Common type definitions
+using FPolyBezierSpline2f = TPolyBezierSpline<float, 2>;
+using FPolyBezierSpline3f = TPolyBezierSpline<float, 3>;
+using FPolyBezierSpline2d = TPolyBezierSpline<double, 2>;
+using FPolyBezierSpline3d = TPolyBezierSpline<double, 3>;
\ No newline at end of file
diff --git a/Engine/Source/Runtime/GeometryCore/Public/Splines/Spline.h b/Engine/Source/Runtime/GeometryCore/Public/Splines/Spline.h
--- a/Engine/Source/Runtime/GeometryCore/Public/Splines/Spline.h
+++ b/Engine/Source/Runtime/GeometryCore/Public/Splines/Spline.h
@@ -61,37 +61,63 @@ protected:
 /**
  * Base spline class with support for varied attribute channels
  */
-template <typename ScalarType, int32 Dimensions>
+template <typename ScalarType, int32 Dimensions, typename EvaluatorType>
 class TSplineBase
 {
 public:
 	using FValue = typename TChooseVectorType<ScalarType, Dimensions>::Type;
 
+	
 protected:
-	TUniquePtr<FAttributeChannel<FValue>> Values;
+	
+	FAttributeChannel<FValue> Values;
 	TMap<FName, TUniquePtr<IAttributeChannelBase>> AttributeChannels;
 
 public:
-	explicit TSplineBase(
-		TUniquePtr<TBaseEvaluator<FValue>> Evaluator)
-		: Values(MakeUnique<FAttributeChannel<FValue>>(MoveTemp(Evaluator)))
+	explicit TSplineBase(TUniquePtr<EvaluatorType> InEvaluator = nullptr)
+		: Values(InEvaluator ? MoveTemp(InEvaluator) : MakeUnique<EvaluatorType>())
 	{
 	}
 
 	virtual ~TSplineBase() = default;
+	
+	TSplineBase(const TSplineBase& Other)
+		: Values(Other.Values)
+	{
+		// Deep copy the attribute channels
+		for (const auto& Pair : Other.AttributeChannels)
+		{
+			AttributeChannels.Add(Pair.Key, Pair.Value->Clone());
+		}
+	}
+    
+	TSplineBase& operator=(const TSplineBase& Other)
+	{
+		if (this != &Other)
+		{
+			Values = Other.Values;
+			AttributeChannels.Empty();
+			for (const auto& Pair : Other.AttributeChannels)
+			{
+				AttributeChannels.Add(Pair.Key, Pair.Value->Clone());
+			}
+		}
+		return *this;
+	}
 
+	// Move operations can be defaulted
+	TSplineBase(TSplineBase&&) noexcept = default;
+	TSplineBase& operator=(TSplineBase&&) noexcept = default;
 	
 	/**
 	  * Creates an attribute channel with a custom evaluator
 	  */
 	template <typename AttrType>
-	void CreateAttributeChannel(
-		const FName& Name,
-		TUniquePtr<TBaseEvaluator<AttrType>> Evaluator)
+	void CreateAttributeChannel( const FName& Name,	TUniquePtr<TBaseEvaluator<AttrType>> AttrEvaluator)
 	{
 		// Ensure attribute interpolation degree doesn't exceed spline degree
-		const int32 SplineDegree = Values->GetEvaluator().GetDegree();
-		const int32 AttrDegree = Evaluator->GetDegree();
+		const int32 SplineDegree = GetEvaluator().GetDegree();
+		const int32 AttrDegree = AttrEvaluator->GetDegree();
 
 		if (AttrDegree > SplineDegree)
 		{
@@ -102,7 +128,7 @@ public:
 			return;
 		}
 
-		TUniquePtr<FAttributeChannel<AttrType>> Channel = MakeUnique<FAttributeChannel<AttrType>>(MoveTemp(Evaluator));
+		TUniquePtr<FAttributeChannel<AttrType>> Channel = MakeUnique<FAttributeChannel<AttrType>>(MoveTemp(AttrEvaluator));
 		AttributeChannels.Add(Name, MoveTemp(Channel));
 	}
 
@@ -113,22 +139,18 @@ public:
 	void CreateAttributeChannel(const FName& Name)
 	{
 		using AttributeEvaluator = TAttributeChannelEvaluator<AttrType, ScalarType, Dimensions>;
-		const TBaseEvaluator<FValue>& SplineEvaluator = Values->GetEvaluator();
     
 		// Create the evaluator
-		TUniquePtr<AttributeEvaluator> Evaluator = MakeUnique<AttributeEvaluator>(SplineEvaluator);
+		TUniquePtr<AttributeEvaluator> ChannelEvaluator = MakeUnique<AttributeEvaluator>(GetEvaluator());
     
 		// Release ownership and create new unique ptr of base type
-		TUniquePtr<TBaseEvaluator<AttrType>> BaseEvaluator(Evaluator.Release());
+		TUniquePtr<TBaseEvaluator<AttrType>> BaseEvaluator(ChannelEvaluator.Release());
     
 		CreateAttributeChannel(Name, MoveTemp(BaseEvaluator));
 	}
 
 	template <typename AttrType>
-	int32 AddAttributeValue(
-		FName Name,
-		const AttrType& Value,
-		float Parameter)
+	int32 AddAttributeValue(FName Name,	const AttrType& Value, float Parameter)
 	{
 		if (FAttributeChannel<AttrType>* Channel = GetAttributeChannel<AttrType>(Name))
 		{
@@ -139,15 +161,31 @@ public:
 
 	virtual int32 AddValue(const FValue& Value, float Parameter)
 	{
-		Values->GetEvaluator().ExpandRangeToInclude(Parameter);
-		return Values->AddValue(Value, Parameter);		
+		GetEvaluator().ExpandRangeToInclude(Parameter);
+		return Values.AddValue(Value, Parameter);		
+	}
+
+	virtual FValue Evaluate(float Parameter) const
+	{
+		return Values.Evaluate(Parameter);
 	}
 
-	FValue Evaluate(float Parameter) const
+	/** Evaluates first derivative at parameter */
+	virtual FValue EvaluateDerivative(float Parameter) const
 	{
-		return Values->Evaluate(Parameter);
+		return GetEvaluator().EvaluateDerivative(
+			Values.GetStorage().GetEntries(), 
+			Parameter);
 	}
 
+	/** Evaluates second derivative at parameter */
+	virtual FValue EvaluateSecondDerivative(float Parameter) const
+	{
+		return GetEvaluator().EvaluateSecondDerivative(
+			Values.GetStorage().GetEntries(), 
+			Parameter);
+	}
+	
 	/** 
 	 * Evaluates an attribute at the given parameter in parent space
 	 * @param Name - Name of the attribute channel
@@ -168,7 +206,7 @@ public:
 
 	virtual void Clear()
 	{
-		Values->Clear();
+		Values.Clear();
 		for (TTuple<FName, TUniquePtr<IAttributeChannelBase>>& Pair : AttributeChannels)
 		{
 			static_cast<FAttributeChannel<FValue>*>(Pair.Value.Get())->Clear();
@@ -198,7 +236,7 @@ public:
 	{
 		if (GetNumKeys() < 2) return;
 
-		const auto& Entries = Values->GetStorage().GetEntries();
+		const auto& Entries = Values.GetStorage().GetEntries();
 		const float NewStart = Entries[0].Parameter;
 		const float NewEnd = Entries.Last().Parameter;
         
@@ -213,7 +251,7 @@ public:
 	 */
 	void SetParameterRange(float InStart, float InEnd)
 	{
-		Values->GetEvaluator().SetParentSpaceRange(InStart, InEnd);
+		GetEvaluator().SetParentSpaceRange(InStart, InEnd);
 	}
 
 	/**
@@ -223,65 +261,57 @@ public:
 	 */ 
 	void GetParameterRange(float& OutStart, float& OutEnd) const
 	{
-		Values->GetEvaluator().GetParentSpaceRange(OutStart, OutEnd);
+		GetEvaluator().GetParentSpaceRange(OutStart, OutEnd);
 	}
 
 	/** Set Pre-Infinity handling mode when out of bounds */ 
 	void SetPreInfinityMode(EOutOfBoundsHandlingMode InMode)
 	{
-		Values->GetEvaluator().SetPreInfinityMode(InMode);
+		GetEvaluator().SetPreInfinityMode(InMode);
 	}
 
 	/** Set Post-Infinity handling mode when out of bounds */
 	void SetPostInfinityMode(EOutOfBoundsHandlingMode InMode)
 	{
-		Values->GetEvaluator().SetPostInfinityMode(InMode);
+		GetEvaluator().SetPostInfinityMode(InMode);
 	}
 
 	/** Get the current post-infinity mode */
 	EOutOfBoundsHandlingMode GetPreInfinityMode() const
 	{
-		return Values->GetEvaluator().GetPreInfinityMode();
+		return GetEvaluator().GetPreInfinityMode();
 	}
 
 	/** Get the current post-infinity mode */
 	EOutOfBoundsHandlingMode GetPostInfinityMode() const
 	{
-		return Values->GetEvaluator().GetPostInfinityMode();
+		return GetEvaluator().GetPostInfinityMode();
 	}
 	
-	/** Access to the spline's evaluator for attribute channels */
-	const TBaseEvaluator<FValue>& GetSplineEvaluator() const
-	{
-		return Values->GetEvaluator();
-	}
-
-	/** For derived classes that need concrete type access */
-	template <typename EvaluatorType>
-	EvaluatorType& GetEvaluatorAs()
+	/** Retrieves the main evaluator from the Values channel */
+	EvaluatorType& GetEvaluator()
 	{
-		return static_cast<EvaluatorType&>(Values->GetEvaluator());
+		return static_cast<EvaluatorType&>(Values.GetEvaluator());
 	}
 
-	template <typename EvaluatorType>
-	const EvaluatorType& GetEvaluatorAs() const
+	const EvaluatorType& GetEvaluator() const
 	{
-		return static_cast<const EvaluatorType&>(Values->GetEvaluator());
+		return static_cast<const EvaluatorType&>(Values.GetEvaluator());
 	}
 
 	int32 GetNumKeys() const
 	{
-		return Values->GetStorage().GetEntries().Num();
+		return Values.GetStorage().GetEntries().Num();
 	}
 
 	FValue GetValueAt(int32 Index) const
 	{
-		return Values->GetStorage().GetEntries()[Index].Value;
+		return Values.GetStorage().GetEntries()[Index].Value;
 	}
 
 	float GetParameterAt(int32 Index) const
 	{
-		return Values->GetStorage().GetEntries()[Index].Parameter;
+		return Values.GetStorage().GetEntries()[Index].Parameter;
 	}
 
 	/**
@@ -360,7 +390,7 @@ private:
 
 
 // Common type definitions for convenience and readability
-using FSpline2f = TSplineBase<float, 2>;
-using FSpline3f = TSplineBase<float, 3>;
-using FSpline2d = TSplineBase<double, 2>;
-using FSpline3d = TSplineBase<double, 3>;
+using FSpline2f = TSplineBase<float, 2, TBaseEvaluator<FVector2f>>;
+using FSpline3f = TSplineBase<float, 3, TBaseEvaluator<FVector3f>>;
+using FSpline2d = TSplineBase<double, 2, TBaseEvaluator<FVector2D>>;
+using FSpline3d = TSplineBase<double, 3, TBaseEvaluator<FVector>>;
