Spline serialization

diff --git a/Engine/Source/Runtime/GeometryCore/Public/Splines/AttributeChannel.h b/Engine/Source/Runtime/GeometryCore/Public/Splines/AttributeChannel.h
--- a/Engine/Source/Runtime/GeometryCore/Public/Splines/AttributeChannel.h
+++ b/Engine/Source/Runtime/GeometryCore/Public/Splines/AttributeChannel.h
@@ -4,7 +4,11 @@
 
 #include "ParameterizedStorage.h"
 #include "BaseEvaluator.h"
-
+namespace UE
+{
+namespace Geometry
+{
+	
 /** Base interface for attribute channels **/
 class IAttributeChannelBase
 {
@@ -19,23 +23,33 @@ public:
 	 * @return A new instance of the channel
 	 */
 	virtual TUniquePtr<IAttributeChannelBase> Clone() const = 0;
+
+	/**
+	 * Serializes the channel to/from an archive
+	 * @param Ar - Archive to serialize to/from
+	 */
+	virtual FArchive& operator<<(FArchive& Ar) = 0;
 };
 
 /**
  * Represents a single attribute channel with its own evaluation scheme
  */
-template<typename AttrType>
-class FAttributeChannel : public IAttributeChannelBase
+template<typename ValueType>
+class TAttributeChannel : public IAttributeChannelBase
 {
-    FAttributeStorage<AttrType> Storage;
-    TUniquePtr<TBaseEvaluator<AttrType>> Evaluator;
+    TAttributeStorage<ValueType> Storage;
+    TUniquePtr<TBaseEvaluator<ValueType>> Evaluator;
 
     // Optional overrides:
-    TFunction<TArray<const AttrType*>(TArrayView<const FStorageEntry<AttrType>>, float)> FindWindowOverride;
-    TFunction<AttrType(TArrayView<const AttrType* const>, float)> InterpolateOverride;
+    TFunction<TArray<const ValueType*>(TArrayView<const TStorageEntry<ValueType>>, float)> FindWindowOverride;
+    TFunction<ValueType(TArrayView<const ValueType* const>, float)> InterpolateOverride;
 
 public:
-    explicit FAttributeChannel(TUniquePtr<TBaseEvaluator<AttrType>> InEvaluator)
+	TAttributeChannel() 
+		: Evaluator(MakeUnique<TBaseEvaluator<ValueType>>())
+	{}
+	
+    explicit TAttributeChannel(TUniquePtr<TBaseEvaluator<ValueType>> InEvaluator)
         : Evaluator(MoveTemp(InEvaluator))
     {
     }
@@ -43,11 +57,11 @@ public:
 	// Make copying logic live in the channel class
 	virtual TUniquePtr<IAttributeChannelBase> Clone() const override
 	{
-    	return MakeUnique<FAttributeChannel>(*this);
+    	return MakeUnique<TAttributeChannel>(*this);
     }
 	
 	// Add copy constructor to properly clone the evaluator
-	FAttributeChannel(const FAttributeChannel& Other)
+	TAttributeChannel(const TAttributeChannel& Other)
 		: Storage(Other.Storage)
 		, FindWindowOverride(Other.FindWindowOverride)
 		, InterpolateOverride(Other.InterpolateOverride)
@@ -59,7 +73,7 @@ public:
     }
 
 	// Add copy assignment
-	FAttributeChannel& operator=(const FAttributeChannel& Other)
+	TAttributeChannel& operator=(const TAttributeChannel& Other)
     {
     	if (this != &Other)
     	{
@@ -79,6 +93,38 @@ public:
     	return *this;
     }
 
+    virtual FArchive& operator<<(FArchive& Ar) override
+    {
+    	// Serialize storage
+    	Ar << Storage;
+        
+    	// Handle evaluator serialization
+    	bool bHasEvaluator = Evaluator != nullptr;
+    	Ar << bHasEvaluator;
+        
+    	if (Ar.IsSaving())
+    	{
+    		if (bHasEvaluator)
+    		{
+    			*Evaluator << Ar;
+    		}
+    	}
+    	else
+    	{
+    		if (bHasEvaluator)
+    		{
+    			// Create appropriate evaluator type and serialize
+    			if (!Evaluator)
+    			{
+    				Evaluator = MakeUnique<TBaseEvaluator<ValueType>>();
+    			}
+    			*Evaluator << Ar;
+    		}
+    	}
+        
+    	return Ar;
+    }
+
 	/**
 	 * Clears all values from this channel
 	 */
@@ -93,7 +139,7 @@ public:
      * @param Parameter
      * @return Index of the added value
      */
-    int32 AddValue(const AttrType& Value, float Parameter)
+    int32 AddValue(const ValueType& Value, float Parameter)
     {
         return Storage.AddValue(Value, Parameter);
     }
@@ -102,13 +148,13 @@ public:
      * Evaluates the channel at the given parameter in channel space
      * @param Parameter - Parameter value in channel's local space
      */
-    AttrType Evaluate(float Parameter) const
+    ValueType Evaluate(float Parameter) const
     {
-        if (Storage.GetEntries().Num() == 0) return AttrType();
+        if (Storage.GetEntries().Num() == 0) return ValueType();
         if (Storage.GetEntries().Num() == 1) return Storage.GetEntries()[0].Value;
 
         // Get window using override or evaluator
-        TArray<const AttrType*> Window;
+        TArray<const ValueType*> Window;
         if (FindWindowOverride)
         {
             Window = FindWindowOverride(Storage.GetEntries(), Parameter);
@@ -117,7 +163,7 @@ public:
         {
             Window = Evaluator->FindWindow(Storage.GetEntries(), Parameter);
         }
-        if (Window.Num() == 0) return AttrType();
+        if (Window.Num() == 0) return ValueType();
         
         if (InterpolateOverride)
         {
@@ -126,30 +172,30 @@ public:
         return Evaluator->InterpolateWindow(Window, Parameter);
     }
 
-    TBaseEvaluator<AttrType>& GetEvaluator()
+    TBaseEvaluator<ValueType>& GetEvaluator()
     {
         return *Evaluator;
     }
     
-    const TBaseEvaluator<AttrType>& GetEvaluator() const
+    const TBaseEvaluator<ValueType>& GetEvaluator() const
     {
         return *Evaluator;
     }
     
     // Access to underlying storage
-    const FAttributeStorage<AttrType>& GetStorage() const
+    const TAttributeStorage<ValueType>& GetStorage() const
     {
         return Storage;
     }
     
 
     // Set overrides
-    void SetFindWindowOverride(TFunction<TArray<const AttrType*>(TArrayView<const FStorageEntry<AttrType>>, float)> Override)
+    void SetFindWindowOverride(TFunction<TArray<const ValueType*>(TArrayView<const TStorageEntry<ValueType>>, float)> Override)
     {
         FindWindowOverride = MoveTemp(Override);
     }
 
-    void SetInterpolateOverride(TFunction<AttrType(TArrayView<const AttrType* const>, float)> Override)
+    void SetInterpolateOverride(TFunction<ValueType(TArrayView<const ValueType* const>, float)> Override)
     {
         InterpolateOverride = MoveTemp(Override);
     }
@@ -161,3 +207,6 @@ public:
         InterpolateOverride = nullptr;
     }
 };
+
+} // end namespace UE::Geometry
+} // end namespace UE
\ No newline at end of file
diff --git a/Engine/Source/Runtime/GeometryCore/Public/Splines/BSpline.h b/Engine/Source/Runtime/GeometryCore/Public/Splines/BSpline.h
--- a/Engine/Source/Runtime/GeometryCore/Public/Splines/BSpline.h
+++ b/Engine/Source/Runtime/GeometryCore/Public/Splines/BSpline.h
@@ -4,7 +4,10 @@
 
 #include "Spline.h"
 #include "BaseEvaluator.h"
-
+namespace UE
+{
+namespace Geometry
+{
 
 enum class EBSplineKnotMode 
 {
@@ -31,8 +34,8 @@ class TBSplineEvaluator : public TBaseEvaluator<typename TChooseVectorType<Scala
 	static_assert(Degree > 0, "BSpline degree must be positive");
 	
 public:
-	using FValue = typename TChooseVectorType<ScalarType, Dimensions>::Type;
-	using Base = TBaseEvaluator<FValue>;
+	using ValueType = typename TChooseVectorType<ScalarType, Dimensions>::Type;
+	using Base = TBaseEvaluator<ValueType>;
 	
 private:
 	
@@ -52,14 +55,14 @@ public:
 	TBSplineEvaluator() = default;
 	virtual ~TBSplineEvaluator() override = default;
 
-	virtual TUniquePtr<TBaseEvaluator<FValue>> Clone() const override
+	virtual TUniquePtr<TBaseEvaluator<ValueType>> Clone() const override
 	{
 		return MakeUnique<TBSplineEvaluator>(*this);
 	}
 
-	virtual TArray<const FValue*> FindWindow(TArrayView<const FStorageEntry<FValue>> Entries, float Parameter) const override
+	virtual TArray<const ValueType*> FindWindow(TArrayView<const TStorageEntry<ValueType>> Entries, float Parameter) const override
 	{
-		TArray<const FValue*> Window;
+		TArray<const ValueType*> Window;
 		if (Entries.Num() < Degree + 1)
 		{
 			return Window;
@@ -81,21 +84,21 @@ public:
 		return Window;
 	}
 
-	virtual FValue InterpolateWindow(TArrayView<const FValue* const> Window, float Parameter) const override
+	virtual ValueType InterpolateWindow(TArrayView<const ValueType* const> Window, float Parameter) const override
 	{
 		TArray<float> Basis;
 		const float MappedParam = Base::MapToParentSpace(Parameter);
 		ComputeBasisFunctions(FindKnotSpan(MappedParam), MappedParam, Basis);
-		return TSplineInterpolationPolicy<FValue>::InterpolateWithBasis(Window, Basis);
+		return TSplineInterpolationPolicy<ValueType>::InterpolateWithBasis(Window, Basis);
 	}
 
-	virtual FValue EvaluateDerivative(TArrayView<const FStorageEntry<FValue>> Entries, float Parameter) const override 
+	virtual ValueType EvaluateDerivative(TArrayView<const TStorageEntry<ValueType>> Entries, float Parameter) const override 
 	{
 		// Find window using existing window logic
-		TArray<const FValue*> Window = FindWindow(Entries, Parameter);
+		TArray<const ValueType*> Window = FindWindow(Entries, Parameter);
 		if (Window.Num() < GetDegree() + 1)
 		{
-			return FValue();
+			return ValueType();
 		}
 
 		// Calculate derivative basis functions
@@ -105,15 +108,15 @@ public:
 		ComputeDerivativeBasisFunctions(k, u, DerivBasis);
 
 		// Apply basis using type-specific policy interpolation
-		return TSplineInterpolationPolicy<FValue>::InterpolateWithBasis(Window, DerivBasis);
+		return TSplineInterpolationPolicy<ValueType>::InterpolateWithBasis(Window, DerivBasis);
 	}
 
-	virtual FValue EvaluateSecondDerivative(TArrayView<const FStorageEntry<FValue>> Entries, float Parameter) const override
+	virtual ValueType EvaluateSecondDerivative(TArrayView<const TStorageEntry<ValueType>> Entries, float Parameter) const override
 	{
-		TArray<const FValue*> Window = FindWindow(Entries, Parameter);
+		TArray<const ValueType*> Window = FindWindow(Entries, Parameter);
 		if (Window.Num() < GetDegree() + 1)
 		{
-			return FValue();
+			return ValueType();
 		}
 
 		// Calculate second derivative basis functions 
@@ -123,7 +126,7 @@ public:
 		ComputeSecondDerivativeBasisFunctions(k, u, SecondDerivBasis);
 
 		// Apply basis using type-specific policy interpolation
-		return TSplineInterpolationPolicy<FValue>::InterpolateWithBasis(Window, SecondDerivBasis);
+		return TSplineInterpolationPolicy<ValueType>::InterpolateWithBasis(Window, SecondDerivBasis);
 	}
 
 	virtual int32 GetDegree() const override
@@ -131,6 +134,27 @@ public:
 		return Degree;
 	}
 
+	virtual FArchive& operator<<(FArchive& Ar) override
+	{
+		// Serialize base evaluator first
+		Base::operator<<(Ar);
+        
+		// Serialize BSpline-specific data
+		Ar << KnotVector;
+		Ar << NumValues;
+        
+		uint8 KnotModeInt = static_cast<uint8>(this->KnotMode);
+		Ar << KnotModeInt;
+        
+		if (Ar.IsLoading())
+		{
+			this->KnotMode = static_cast<EBSplineKnotMode>(KnotModeInt);
+		}
+        
+		Ar << bClampEnds;
+        
+		return Ar;
+	}
 
 	/**
 	 * Sets a custom knot vector and switches to custom knot mode
@@ -183,7 +207,7 @@ public:
 	 * 
 	 * @param Points - Array of control points and their parameters
 	 */
-	void UpdateKnotVector(const TArrayView<const FStorageEntry<FValue>>& Points)
+	void UpdateKnotVector(const TArrayView<const TStorageEntry<ValueType>>& Points)
 	{
 		NumValues = Points.Num();
 		if (NumValues < Degree + 1)
@@ -537,7 +561,7 @@ private:
      * Generates a knot vector based on chord lengths between control points
      * This gives better parameterization when control points are unevenly spaced
      */
-    void GenerateChordLengthKnots(int32 KnotCount, const TArrayView<const FStorageEntry<FValue>>& Points)
+    void GenerateChordLengthKnots(int32 KnotCount, const TArrayView<const TStorageEntry<ValueType>>& Points)
     {
         if (Points.Num() < 2)
         {
@@ -597,7 +621,7 @@ private:
      * Uses square root of chord lengths which helps prevent cusps and unwanted loops
      * Often provides the most visually pleasing results for interactive curve editing
      */
-    void GenerateCentripetalKnots(int32 KnotCount, const TArrayView<const FStorageEntry<FValue>>& Points)
+    void GenerateCentripetalKnots(int32 KnotCount, const TArrayView<const TStorageEntry<ValueType>>& Points)
 	{
 		if (Points.Num() < 2)
 		{
@@ -665,7 +689,7 @@ protected:
 		return MakeShared<TBSpline>();
 	}
 public:
-	using typename Base::FValue;
+	using typename Base::ValueType;
 	explicit TBSpline(TUniquePtr<EvaluatorType> InEvaluator = nullptr)
 			: Base(MoveTemp(InEvaluator))
 	{
@@ -695,7 +719,7 @@ public:
 	 * @param Parameter - Parameter value for the point
 	 * @return Index of the added value
 	 */
-	virtual int32 AddValue(const FValue& Value, float Parameter) override
+	virtual int32 AddValue(const ValueType& Value, float Parameter) override
 	{
 		UE_LOG(LogSpline, Verbose, TEXT("TBSpline::AddValue - Evaluator: %p, KnotMode: %d"), 
 			&this->GetEvaluator(), static_cast<int32>(this->GetEvaluator().GetKnotMode()));
@@ -740,4 +764,7 @@ using FBSpline3d = TBSpline<double, 3>;
 using FQuadraticBSpline2f = TBSpline<float, 2, 2>;
 using FCubicBSpline3f = TBSpline<float, 3, 3>;
 using FQuarticBSpline3f = TBSpline<float, 3, 4>;
-using FQuinticBSpline2d = TBSpline<double, 2, 5>;
\ No newline at end of file
+using FQuinticBSpline2d = TBSpline<double, 2, 5>;
+
+} // end namespace UE::Geometry
+} // end namespace UE	
\ No newline at end of file
diff --git a/Engine/Source/Runtime/GeometryCore/Public/Splines/BaseEvaluator.h b/Engine/Source/Runtime/GeometryCore/Public/Splines/BaseEvaluator.h
--- a/Engine/Source/Runtime/GeometryCore/Public/Splines/BaseEvaluator.h
+++ b/Engine/Source/Runtime/GeometryCore/Public/Splines/BaseEvaluator.h
@@ -5,6 +5,11 @@
 #include "CoreMinimal.h"
 #include "InterpolationPolicies.h"
 
+namespace UE
+{
+namespace Geometry
+{
+	
 DEFINE_LOG_CATEGORY_STATIC(LogSpline, Log, All);
 /** 
  * Enum for handling out-of-bounds parameter values
@@ -49,7 +54,7 @@ public:
 	 */
 	/**  */
 	virtual TArray<const ValueType*> FindWindow(
-		TArrayView<const FStorageEntry<ValueType>> Entries,
+		TArrayView<const TStorageEntry<ValueType>> Entries,
 		float Parameter) const
 	{
 		TArray<const ValueType*> Window;
@@ -100,7 +105,7 @@ public:
 
 	/** Evaluates first derivative at parameter */
 	virtual ValueType EvaluateDerivative(
-		TArrayView<const FStorageEntry<ValueType>> Entries,
+		TArrayView<const TStorageEntry<ValueType>> Entries,
 		float Parameter) const
 	{
 		TArray<const ValueType*> Window = FindWindow(Entries, Parameter);
@@ -119,7 +124,7 @@ public:
 
 	/** Evaluates second derivative at parameter */
 	virtual ValueType EvaluateSecondDerivative(
-		TArrayView<const FStorageEntry<ValueType>> Entries,
+		TArrayView<const TStorageEntry<ValueType>> Entries,
 		float Parameter) const
 	{
 		TArray<const ValueType*> Window = FindWindow(Entries, Parameter);
@@ -137,9 +142,9 @@ public:
 	}
 
 	/** Clone the evaluator */
-	virtual TUniquePtr<TBaseEvaluator<ValueType>> Clone() const
+	virtual TUniquePtr<TBaseEvaluator> Clone() const
 	{
-		return MakeUnique<TBaseEvaluator<ValueType>>(*this);
+		return MakeUnique<TBaseEvaluator>(*this);
 	}
 	
 	/** Convert from normalized [0,1] to parent space with bounds handling */
@@ -209,6 +214,26 @@ public:
 		OutStart = ParentStart;
 		OutEnd = ParentEnd;
 	}
+
+	virtual FArchive& operator<<(FArchive& Ar)
+	{
+		Ar << ParentStart;
+		Ar << ParentEnd;
+        
+		uint8 PreMode = static_cast<uint8>(PreInfinityMode);
+		uint8 PostMode = static_cast<uint8>(PostInfinityMode);
+        
+		Ar << PreMode;
+		Ar << PostMode;
+        
+		if (Ar.IsLoading())
+		{
+			PreInfinityMode = static_cast<EOutOfBoundsHandlingMode>(PreMode);
+			PostInfinityMode = static_cast<EOutOfBoundsHandlingMode>(PostMode);
+		}
+        
+		return Ar;
+	}
 	
 protected:
 	/** Override to implement specific parent space mapping */
@@ -259,4 +284,7 @@ private:
 		
 		return Parameter;
 	}
-};
\ No newline at end of file
+};
+
+} // end namespace UE::Geometry
+} // end namespace UE
\ No newline at end of file
diff --git a/Engine/Source/Runtime/GeometryCore/Public/Splines/ParameterizedStorage.h b/Engine/Source/Runtime/GeometryCore/Public/Splines/ParameterizedStorage.h
--- a/Engine/Source/Runtime/GeometryCore/Public/Splines/ParameterizedStorage.h
+++ b/Engine/Source/Runtime/GeometryCore/Public/Splines/ParameterizedStorage.h
@@ -3,20 +3,42 @@
 #pragma once
 
 #include "ParameterizedTypes.h"
+#include "Serialization/Archive.h"
+namespace UE
+{
+namespace Geometry
+{
+	
 /**
  * Pure storage class - no evaluation logic
  */
-template<typename AttrType>
-class FAttributeStorage
+template<typename ValueType>
+class TAttributeStorage
 {
-    TArray<FStorageEntry<AttrType>> Entries;
+    TArray<TStorageEntry<ValueType>> Entries;
 
 public:
-    int32 AddValue(const AttrType& Value, float Parameter)
+	/**
+	 * Serialization support for storage entries
+	 */
+	friend FArchive& operator<<(FArchive& Ar, TAttributeStorage& Storage)
+	{
+		Ar << Storage.Entries;
+		return Ar;
+	}
+
+    /**
+     * Add a value to the storage with a given parameter
+     * @param Value - Value to add
+     * @param Parameter - Parameter value in channel's local space
+     * @return Index of the added value
+     */
+    int32 AddValue(const ValueType& Value, float Parameter)
     {
         const int32 Index = Algo::LowerBound(
             Entries, Parameter,
-            [](const FStorageEntry<AttrType>& Entry, float Param) { 
+            [](const TStorageEntry<ValueType>& Entry, float Param)
+            { 
                 return Entry.Parameter <= Param; 
             }
         );
@@ -24,7 +46,7 @@ public:
     	return Index;
     }
 
-    TArrayView<const FStorageEntry<AttrType>> GetEntries() const 
+    TArrayView<const TStorageEntry<ValueType>> GetEntries() const 
     { 
         return Entries; 
     }
@@ -33,4 +55,7 @@ public:
     {
         Entries.Reset();
     }
-};
\ No newline at end of file
+};
+
+} // end namespace UE::Geometry
+} // end namespace UE
\ No newline at end of file
diff --git a/Engine/Source/Runtime/GeometryCore/Public/Splines/ParameterizedTypes.h b/Engine/Source/Runtime/GeometryCore/Public/Splines/ParameterizedTypes.h
--- a/Engine/Source/Runtime/GeometryCore/Public/Splines/ParameterizedTypes.h
+++ b/Engine/Source/Runtime/GeometryCore/Public/Splines/ParameterizedTypes.h
@@ -4,6 +4,10 @@
 
 #include "CoreMinimal.h"
 
+namespace UE
+{
+namespace Geometry
+{
 /**
  * Vector type policies based on dimension and numeric type
  */
@@ -55,12 +59,22 @@ struct TChooseVectorType
 };
 
 // Forward declarations
-template<typename AttrType>
-class FAttributeChannel;
+template<typename ValueType>
+class TAttributeChannel;
 
-template<typename AttrType>
-struct FStorageEntry
+template<typename ValueType>
+struct TStorageEntry
 {
 	float Parameter;
-	AttrType Value;
+	ValueType Value;
+
+	friend FArchive& operator<<(FArchive& Ar, TStorageEntry& Entry)
+	{
+		Ar << Entry.Parameter;
+		Ar << Entry.Value;
+		return Ar;
+	}
 };
+
+} // end namespace UE::Geometry
+} // end namespace UE
\ No newline at end of file
diff --git a/Engine/Source/Runtime/GeometryCore/Public/Splines/PolyBezierSpline.h b/Engine/Source/Runtime/GeometryCore/Public/Splines/PolyBezierSpline.h
--- a/Engine/Source/Runtime/GeometryCore/Public/Splines/PolyBezierSpline.h
+++ b/Engine/Source/Runtime/GeometryCore/Public/Splines/PolyBezierSpline.h
@@ -3,7 +3,10 @@
 #pragma once
 
 #include "BSpline.h"
-
+namespace UE
+{
+namespace Geometry
+{
 /**
  * Custom evaluator for PolyBezier splines.
  * Handles parameter mapping and control point selection for proper Bezier evaluation.
@@ -11,7 +14,7 @@
 template<typename T, int32 D>
 class TPolyBezierEvaluator : public TBSplineEvaluator<T, D, 3>
 {
-	using FValue = typename TChooseVectorType<T, D>::Type;
+	using ValueType = typename TChooseVectorType<T, D>::Type;
 	using Base = TBSplineEvaluator<T, D, 3>;
 
 private:
@@ -24,17 +27,28 @@ public:
         this->SetKnotGenerationMode(EBSplineKnotMode::Custom);
         this->SetClampedEnds(true);
     }
+
+    virtual FArchive& operator<<(FArchive& Ar) override
+    {
+        // Serialize base evaluator first
+        Base::operator<<(Ar);
+        
+        // Serialize PolyBezier-specific data
+        Ar << NumSegments;
+        
+        return Ar;
+    }
     
-    virtual TUniquePtr<TBaseEvaluator<FValue>> Clone() const override
+    virtual TUniquePtr<TBaseEvaluator<ValueType>> Clone() const override
     {
         return MakeUnique<TPolyBezierEvaluator<T, D>>(*this);
     }
     
     virtual int32 GetDegree() const override { return 3; }
 
-    virtual TArray<const FValue*> FindWindow(TArrayView<const FStorageEntry<FValue>> Entries, float Parameter) const override
+    virtual TArray<const ValueType*> FindWindow(TArrayView<const TStorageEntry<ValueType>> Entries, float Parameter) const override
     {
-        TArray<const FValue*> Window;
+        TArray<const ValueType*> Window;
         if (Entries.Num() < 4)  // Need at least one complete Bezier segment
         {
             return Window;
@@ -63,11 +77,11 @@ public:
         return Window;
     }
 
-    virtual FValue InterpolateWindow(TArrayView<const FValue* const> Window, float Parameter) const override
+    virtual ValueType InterpolateWindow(TArrayView<const ValueType* const> Window, float Parameter) const override
     {
         if (Window.Num() < 4)
         {
-            return FValue();
+            return ValueType();
         }
 
         // Map to parent space (segment space)
@@ -97,7 +111,7 @@ public:
         Basis[2] = 3.0f * t2 * OneMinusT;      // 3t²(1-t)
         Basis[3] = t * t2;                     // t³
 
-        return TSplineInterpolationPolicy<FValue>::InterpolateWithBasis(Window, Basis);
+        return TSplineInterpolationPolicy<ValueType>::InterpolateWithBasis(Window, Basis);
     }
 
     virtual bool SetCustomKnots(const TArray<float>& NewKnots, int32 InNumValues) override
@@ -139,12 +153,12 @@ public:
         return Base::SetCustomKnots(BezierKnots, InNumValues);
     }
     
-    FValue EvaluateDerivative(TArrayView<const FStorageEntry<FValue>> Entries, float Parameter) const override
+    ValueType EvaluateDerivative(TArrayView<const TStorageEntry<ValueType>> Entries, float Parameter) const override
     {
-        TArray<const FValue*> Window = FindWindow(Entries, Parameter);
+        TArray<const ValueType*> Window = FindWindow(Entries, Parameter);
         if (Window.Num() < 4)  // Need all 4 control points for a cubic segment
         {
-            return FValue();
+            return ValueType();
         }
 
         // Get local parameter within segment [0,1]
@@ -160,15 +174,15 @@ public:
         TArray<float> DerivBasis = { c0, c1, c2, c3 };
 
         // Scale for parameter space and apply basis
-        return TSplineInterpolationPolicy<FValue>::InterpolateWithBasis(Window, DerivBasis) * NumSegments;
+        return TSplineInterpolationPolicy<ValueType>::InterpolateWithBasis(Window, DerivBasis) * NumSegments;
     }
 
-    FValue EvaluateSecondDerivative(TArrayView<const FStorageEntry<FValue>> Entries, float Parameter) const override
+    ValueType EvaluateSecondDerivative(TArrayView<const TStorageEntry<ValueType>> Entries, float Parameter) const override
     {
-        TArray<const FValue*> Window = FindWindow(Entries, Parameter);
+        TArray<const ValueType*> Window = FindWindow(Entries, Parameter);
         if (Window.Num() < 4)
         {
-            return FValue();
+            return ValueType();
         }
 
         // Get local parameter for this segment
@@ -185,7 +199,7 @@ public:
 
         // Scale for parameter space and apply basis
         const float SegmentScale = NumSegments;
-        return TSplineInterpolationPolicy<FValue>::InterpolateWithBasis(Window, SecondDerivBasis) 
+        return TSplineInterpolationPolicy<ValueType>::InterpolateWithBasis(Window, SecondDerivBasis) 
             * (SegmentScale * SegmentScale);
     }
 
@@ -209,7 +223,7 @@ template<typename ScalarType, int32 Dimensions>
 class TPolyBezierSpline: public TSplineBase<ScalarType, Dimensions, TPolyBezierEvaluator<ScalarType, Dimensions>>
 {
 public:
-    using FValue = typename TChooseVectorType<ScalarType, Dimensions>::Type;
+    using ValueType = typename TChooseVectorType<ScalarType, Dimensions>::Type;
     using FEvaluator = TPolyBezierEvaluator<ScalarType, Dimensions>;
     using Base = TSplineBase<ScalarType, Dimensions, TPolyBezierEvaluator<ScalarType, Dimensions>>;
 
@@ -233,14 +247,14 @@ public:
      * @param P3 End point
      * @return Index of the new segment
      */
-    int32 AddBezierSegment(const FValue& P0, const FValue& P1, const FValue& P2, const FValue& P3)
+    int32 AddBezierSegment(const ValueType& P0, const ValueType& P1, const ValueType& P2, const ValueType& P3)
     {
         const int32 NumPoints = Base::GetNumKeys();
         
         // Validate continuity with previous segment if this isn't the first one
         if (NumPoints > 0)
         {
-            const FValue& LastPoint = Base::GetValueAt(NumPoints - 1);
+            const ValueType& LastPoint = Base::GetValueAt(NumPoints - 1);
             if (!LastPoint.Equals(P0))
             {
                 UE_LOG(LogSpline, Warning, TEXT("PolyBezier segment start point doesn't match previous end point"));
@@ -262,8 +276,8 @@ public:
     }
 
     /** Gets the control points for a specific Bezier segment */
-    bool GetSegmentControlPoints(int32 SegmentIndex, FValue& OutP0, FValue& OutP1, 
-        FValue& OutP2, FValue& OutP3) const
+    bool GetSegmentControlPoints(int32 SegmentIndex, ValueType& OutP0, ValueType& OutP1, 
+        ValueType& OutP2, ValueType& OutP3) const
     {
         if (!IsValidSegmentIndex(SegmentIndex))
         {
@@ -285,11 +299,11 @@ public:
     }
 
     /** Gets the endpoint of a segment */
-    FValue GetSegmentEndPoint(int32 SegmentIndex) const
+    ValueType GetSegmentEndPoint(int32 SegmentIndex) const
     {
         if (!IsValidSegmentIndex(SegmentIndex))
         {
-            return FValue();
+            return ValueType();
         }
         return Base::GetValueAt(SegmentIndex * 4 + 3);
     }
@@ -306,4 +320,7 @@ private:
 using FPolyBezierSpline2f = TPolyBezierSpline<float, 2>;
 using FPolyBezierSpline3f = TPolyBezierSpline<float, 3>;
 using FPolyBezierSpline2d = TPolyBezierSpline<double, 2>;
-using FPolyBezierSpline3d = TPolyBezierSpline<double, 3>;
\ No newline at end of file
+using FPolyBezierSpline3d = TPolyBezierSpline<double, 3>;
+
+} // end namespace UE::Geometry
+} // end namespace UE
\ No newline at end of file
diff --git a/Engine/Source/Runtime/GeometryCore/Public/Splines/Spline.h b/Engine/Source/Runtime/GeometryCore/Public/Splines/Spline.h
--- a/Engine/Source/Runtime/GeometryCore/Public/Splines/Spline.h
+++ b/Engine/Source/Runtime/GeometryCore/Public/Splines/Spline.h
@@ -5,13 +5,18 @@
 #include "AttributeChannel.h"
 #include "BaseEvaluator.h"
 
+namespace UE
+{
+namespace Geometry
+{
+
 /** Flags controlling what data gets copied during spline cloning */
 enum class ESplineCloneFlags : uint32
 {
 	None               = 0,
 	CoreData           = 1 << 0,   // Core spline data (points, tangents etc.)
 	Attributes		   = 1 << 1,   // Attribute channels
-    
+
 	// Common combinations
 	All = CoreData | Attributes,
 	CoreOnly = CoreData
@@ -39,7 +44,7 @@ public:
 	{
 		return MakeUnique<TAttributeChannelEvaluator<ValueType, SplineT, SplineD>>(ParentEvaluator);
 	}
-	
+
 	virtual int32 GetDegree() const override
 	{
 		return ParentEvaluator.GetDegree();
@@ -65,12 +70,10 @@ template <typename ScalarType, int32 Dimensions, typename EvaluatorType>
 class TSplineBase
 {
 public:
-	using FValue = typename TChooseVectorType<ScalarType, Dimensions>::Type;
+	using ValueType = typename TChooseVectorType<ScalarType, Dimensions>::Type;
 
-	
 protected:
-	
-	FAttributeChannel<FValue> Values;
+	TAttributeChannel<ValueType> Values;
 	TMap<FName, TUniquePtr<IAttributeChannelBase>> AttributeChannels;
 
 public:
@@ -108,7 +111,45 @@ public:
 	// Move operations can be defaulted
 	TSplineBase(TSplineBase&&) noexcept = default;
 	TSplineBase& operator=(TSplineBase&&) noexcept = default;
-	
+
+	FArchive& operator<<(FArchive& Ar)
+	{
+		// Serialize main value channel
+		Values << Ar;
+        
+		// Serialize attribute channels
+		int32 NumChannels = AttributeChannels.Num();
+		Ar << NumChannels;
+        
+		if (Ar.IsSaving())
+		{
+			for (auto& Pair : AttributeChannels)
+			{
+				// Save channel name and data
+				Ar << Pair.Key;
+				*Pair.Value << Ar;
+			}
+		}
+		else
+		{
+			AttributeChannels.Empty();
+            
+			for (int32 i = 0; i < NumChannels; ++i)
+			{
+				FName ChannelName;
+				Ar << ChannelName;
+                
+				// Create new channel
+				TUniquePtr<IAttributeChannelBase> Channel = 
+					MakeUnique<TAttributeChannel<ValueType>>();
+				*Channel << Ar;
+                
+				AttributeChannels.Add(ChannelName, MoveTemp(Channel));
+			}
+		}
+        
+		return Ar;
+	}
 	/**
 	  * Creates an attribute channel with a custom evaluator
 	  */
@@ -124,11 +165,11 @@ public:
 			// Log an explicit error
 			UE_LOG(LogSpline, Error, TEXT("Attribute channel '%s' has a degree (%d) greater than the spline degree (%d). This is not allowed."),
 				*Name.ToString(), AttrDegree, SplineDegree);
-			
+	
 			return;
 		}
 
-		TUniquePtr<FAttributeChannel<AttrType>> Channel = MakeUnique<FAttributeChannel<AttrType>>(MoveTemp(AttrEvaluator));
+		TUniquePtr<TAttributeChannel<AttrType>> Channel = MakeUnique<TAttributeChannel<AttrType>>(MoveTemp(AttrEvaluator));
 		AttributeChannels.Add(Name, MoveTemp(Channel));
 	}
 
@@ -139,39 +180,39 @@ public:
 	void CreateAttributeChannel(const FName& Name)
 	{
 		using AttributeEvaluator = TAttributeChannelEvaluator<AttrType, ScalarType, Dimensions>;
-    
+
 		// Create the evaluator
 		TUniquePtr<AttributeEvaluator> ChannelEvaluator = MakeUnique<AttributeEvaluator>(GetEvaluator());
-    
+
 		// Release ownership and create new unique ptr of base type
 		TUniquePtr<TBaseEvaluator<AttrType>> BaseEvaluator(ChannelEvaluator.Release());
-    
+
 		CreateAttributeChannel(Name, MoveTemp(BaseEvaluator));
 	}
 
 	template <typename AttrType>
 	int32 AddAttributeValue(FName Name,	const AttrType& Value, float Parameter)
 	{
-		if (FAttributeChannel<AttrType>* Channel = GetAttributeChannel<AttrType>(Name))
+		if (TAttributeChannel<AttrType>* Channel = GetAttributeChannel<AttrType>(Name))
 		{
 			return Channel->AddValue(Value, Parameter);
 		}
 		return INDEX_NONE;
 	}
 
-	virtual int32 AddValue(const FValue& Value, float Parameter)
+	virtual int32 AddValue(const ValueType& Value, float Parameter)
 	{
 		GetEvaluator().ExpandRangeToInclude(Parameter);
 		return Values.AddValue(Value, Parameter);		
 	}
 
-	virtual FValue Evaluate(float Parameter) const
+	virtual ValueType Evaluate(float Parameter) const
 	{
 		return Values.Evaluate(Parameter);
 	}
 
 	/** Evaluates first derivative at parameter */
-	virtual FValue EvaluateDerivative(float Parameter) const
+	virtual ValueType EvaluateDerivative(float Parameter) const
 	{
 		return GetEvaluator().EvaluateDerivative(
 			Values.GetStorage().GetEntries(), 
@@ -179,7 +220,7 @@ public:
 	}
 
 	/** Evaluates second derivative at parameter */
-	virtual FValue EvaluateSecondDerivative(float Parameter) const
+	virtual ValueType EvaluateSecondDerivative(float Parameter) const
 	{
 		return GetEvaluator().EvaluateSecondDerivative(
 			Values.GetStorage().GetEntries(), 
@@ -194,7 +235,7 @@ public:
 	template <typename AttrType>
 	AttrType EvaluateAttribute(FName Name, float Parameter) const
 	{
-		if (const FAttributeChannel<AttrType>* Channel = GetAttributeChannel<AttrType>(Name))
+		if (const TAttributeChannel<AttrType>* Channel = GetAttributeChannel<AttrType>(Name))
 		{
 			// Convert from parent space to channel's local space
 			const float ChannelParam = Channel->GetEvaluator().MapFromParentSpace(Parameter);
@@ -209,7 +250,7 @@ public:
 		Values.Clear();
 		for (TTuple<FName, TUniquePtr<IAttributeChannelBase>>& Pair : AttributeChannels)
 		{
-			static_cast<FAttributeChannel<FValue>*>(Pair.Value.Get())->Clear();
+			static_cast<TAttributeChannel<ValueType>*>(Pair.Value.Get())->Clear();
 		}
 	}
 
@@ -239,10 +280,10 @@ public:
 		const auto& Entries = Values.GetStorage().GetEntries();
 		const float NewStart = Entries[0].Parameter;
 		const float NewEnd = Entries.Last().Parameter;
-        
+
 		SetParameterRange(NewStart, NewEnd);
 	}
-	
+
 	/** 
 	 * Sets the parameter range for evaluation. This remaps how parameters are interpreted
 	 * but does not move control points.
@@ -287,7 +328,7 @@ public:
 	{
 		return GetEvaluator().GetPostInfinityMode();
 	}
-	
+
 	/** Retrieves the main evaluator from the Values channel */
 	EvaluatorType& GetEvaluator()
 	{
@@ -304,7 +345,7 @@ public:
 		return Values.GetStorage().GetEntries().Num();
 	}
 
-	FValue GetValueAt(int32 Index) const
+	ValueType GetValueAt(int32 Index) const
 	{
 		return Values.GetStorage().GetEntries()[Index].Value;
 	}
@@ -325,46 +366,46 @@ public:
 		OnClone(NewSpline, CloneFlags);
 		return NewSpline;
 	}
-	
+
 protected:
 	template <typename AttrType>
-	FAttributeChannel<AttrType>* GetAttributeChannel(const FName Name)
+	TAttributeChannel<AttrType>* GetAttributeChannel(const FName Name)
 	{
 		TUniquePtr<IAttributeChannelBase>* Found = AttributeChannels.Find(Name);
-		return Found ? static_cast<FAttributeChannel<AttrType>*>(Found->Get()) : nullptr;
+		return Found ? static_cast<TAttributeChannel<AttrType>*>(Found->Get()) : nullptr;
 	}
 
 	template <typename AttrType>
-	const FAttributeChannel<AttrType>* GetAttributeChannel(const FName Name) const
+	const TAttributeChannel<AttrType>* GetAttributeChannel(const FName Name) const
 	{
 		const TUniquePtr<IAttributeChannelBase>* Found = AttributeChannels.Find(Name);
-		return Found ? static_cast<const FAttributeChannel<AttrType>*>(Found->Get()) : nullptr;
+		return Found ? static_cast<const TAttributeChannel<AttrType>*>(Found->Get()) : nullptr;
 	}
 
 	/** Hook for derived classes to copy their specific data after base data is copied */
 	virtual void OnClone(TSharedPtr<TSplineBase> NewSpline, ESplineCloneFlags CloneFlags) const {}
-	
+
 private:
 	/**
 	 * Creates a new instance of the correct derived type for cloning.
 	 */
 	virtual TSharedPtr<TSplineBase> CreateNewInstance() const = 0;
-	
+
 	/** Creates a new instance and copies common base class state */
 	TSharedPtr<TSplineBase> CopyBaseData(ESplineCloneFlags CloneFlags) const
 	{
 		// Core data is always required for a valid clone
 		check(EnumHasAnyFlags(CloneFlags, ESplineCloneFlags::CoreData));
-        
+
 		// Create new instance of correct derived type
-        TSharedPtr<TSplineBase> NewSpline = CreateNewInstance();
-        NewSpline->Values = MakeUnique<FAttributeChannel<FValue>>(*Values);
-		
+		TSharedPtr<TSplineBase> NewSpline = CreateNewInstance();
+		NewSpline->Values = MakeUnique<TAttributeChannel<ValueType>>(*Values);
+
 		// Copy parameter range settings
 		float Start, End;
 		GetParameterRange(Start, End);
 		NewSpline->SetParameterRange(Start, End);
-			
+	
 		// Copy out-of-bounds handling mode
 		NewSpline->SetPostInfinityMode(GetPostInfinityMode());
 		NewSpline->SetPreInfinityMode(GetPreInfinityMode());
@@ -378,9 +419,9 @@ private:
 				const TUniquePtr<IAttributeChannelBase>& Channel = AttributePair.Value;
 
 				// Use type erasure to copy the channel without knowing its concrete type
-				if (const FAttributeChannel<FValue>* TypedChannel = static_cast<const FAttributeChannel<FValue>*>(Channel.Get()))
+				if (const TAttributeChannel<ValueType>* TypedChannel = static_cast<const TAttributeChannel<ValueType>*>(Channel.Get()))
 				{
-					NewSpline->AttributeChannels.Add(Name, MakeUnique<FAttributeChannel<FValue>>(*TypedChannel));
+					NewSpline->AttributeChannels.Add(Name, MakeUnique<TAttributeChannel<ValueType>>(*TypedChannel));
 				}
 			}
 		}
@@ -394,3 +435,6 @@ using FSpline2f = TSplineBase<float, 2, TBaseEvaluator<FVector2f>>;
 using FSpline3f = TSplineBase<float, 3, TBaseEvaluator<FVector3f>>;
 using FSpline2d = TSplineBase<double, 2, TBaseEvaluator<FVector2D>>;
 using FSpline3d = TSplineBase<double, 3, TBaseEvaluator<FVector>>;
+
+} // end namespace UE::Geometry
+} // end namespace UE
\ No newline at end of file
