fix convex hull getfaces floodfill logic to properly add neighbors-of-neighbors, and fix the coplanarity test for this case

diff --git a/Engine/Source/Runtime/GeometryCore/Private/CompGeom/ConvexHull3.cpp b/Engine/Source/Runtime/GeometryCore/Private/CompGeom/ConvexHull3.cpp
--- a/Engine/Source/Runtime/GeometryCore/Private/CompGeom/ConvexHull3.cpp
+++ b/Engine/Source/Runtime/GeometryCore/Private/CompGeom/ConvexHull3.cpp
@@ -907,39 +907,54 @@ void TConvexHull3<RealType>::GetFaces(TFunctionRef<void(TArray<int32>&, TVector<
 		while (ToProcess.Num() > 0)
 		{
 			int32 NbrTriIdx = ToProcess.Pop(EAllowShrinking::No);
+			// skip the triangle if it's already assigned a group or we've already visited it in this pass
 			if (GroupIDs[NbrTriIdx] >= 0 || GroupIDs[NbrTriIdx] < -CurGroupID - 1)
 			{
 				continue;
 			}
 
+			// region grow to the triangle if its vertices are all coplanar w/ the source triangle
 			const FIndex3i& NbrTri = Hull[NbrTriIdx];
-			int32 OtherV = -1;
+			bool bFoundOnlyNewCoplanar = false;
 			for (int32 SubIdx = 0; SubIdx < 3; ++SubIdx)
 			{
 				int32 V = NbrTri[SubIdx];
+				// only bother testing the vertex if it's new (wasn't on the source triangles)
 				if (V != Tri.A && V != Tri.B && V != Tri.C)
 				{
-					OtherV = V;
+					TVector<RealType> OtherPos = GetPointFunc(V);
+					RealType O3DVal = ExactPredicates::Orient3<RealType>(TriPosns[0], TriPosns[1], TriPosns[2], OtherPos);
+					if (O3DVal == 0)
+					{
+						// found a coplanar vertex; consider growing to this tri
+						bFoundOnlyNewCoplanar = true;
+					}
+					else
+					{
+						// found a non-coplanar vertex; stop searching and do not grow to this tri
+						bFoundOnlyNewCoplanar = false;
+						break;
+					}
 				}
 			}
-			check(OtherV >= 0);
-			TVector<RealType> OtherPos = GetPointFunc(OtherV);
-			RealType O3DVal = ExactPredicates::Orient3<RealType>(TriPosns[0], TriPosns[1], TriPosns[2], OtherPos);
-			if (O3DVal == 0)
+			if (bFoundOnlyNewCoplanar)
 			{
 				GroupIDs[NbrTriIdx] = CurGroupID;
 
+				// if we failed to compute a stable normal from the source triangle, compute one now
 				if (!bValidNormal)
 				{
 					FaceNormal = VectorUtil::Normal(GetPointFunc(NbrTri.A), GetPointFunc(NbrTri.B), GetPointFunc(NbrTri.C));
 					bValidNormal = FaceNormal.SquaredLength() > 0;
 				}
 
+				// grow to unlabelled neighbors
 				const FIndex3i& NbrNbrInds = HullNeighbors[NbrTriIdx];
 				for (int32 SubIdx = 0; SubIdx < 3; ++SubIdx)
 				{
 					int32 NbrNbrTriIdx = NbrNbrInds[SubIdx];
-					if (GroupIDs[NbrNbrTriIdx] >= 0 || GroupIDs[NbrNbrTriIdx] < -CurGroupID - 1)
+					// consider the triangle for region growing if it's not assigned a group and we haven't visited it yet this pass
+					if (GroupIDs[NbrNbrTriIdx] < 0 && GroupIDs[NbrNbrTriIdx] >= -CurGroupID - 1)
 					{
 						ToProcess.Add(NbrNbrTriIdx);
 					}
