n-th derivative support
API changes to allow updating specific points
Updated serialization calls to be more standard

diff --git a/Engine/Source/Runtime/GeometryCore/Public/Splines/AttributeChannel.h b/Engine/Source/Runtime/GeometryCore/Public/Splines/AttributeChannel.h
--- a/Engine/Source/Runtime/GeometryCore/Public/Splines/AttributeChannel.h
+++ b/Engine/Source/Runtime/GeometryCore/Public/Splines/AttributeChannel.h
@@ -8,7 +8,8 @@ namespace UE
 {
 namespace Geometry
 {
-	
+namespace Spline
+{
 /** Base interface for attribute channels **/
 class IAttributeChannelBase
 {
@@ -24,11 +25,12 @@ public:
 	 */
 	virtual TUniquePtr<IAttributeChannelBase> Clone() const = 0;
 
+	
 	/**
 	 * Serializes the channel to/from an archive
 	 * @param Ar - Archive to serialize to/from
 	 */
-	virtual FArchive& operator<<(FArchive& Ar) = 0;
+	virtual void Serialize(FArchive& Ar) {};
 };
 
 /**
@@ -93,20 +95,20 @@ public:
     	return *this;
     }
 
-    virtual FArchive& operator<<(FArchive& Ar) override
+    friend FArchive& operator<<(FArchive& Ar, TAttributeChannel& Channel)
     {
     	// Serialize storage
-    	Ar << Storage;
+    	Ar << Channel.Storage;
         
     	// Handle evaluator serialization
-    	bool bHasEvaluator = Evaluator != nullptr;
+    	bool bHasEvaluator = Channel.Evaluator != nullptr;
     	Ar << bHasEvaluator;
         
     	if (Ar.IsSaving())
     	{
     		if (bHasEvaluator)
     		{
-    			*Evaluator << Ar;
+    			Ar << *Channel.Evaluator;
     		}
     	}
     	else
@@ -114,11 +116,11 @@ public:
     		if (bHasEvaluator)
     		{
     			// Create appropriate evaluator type and serialize
-    			if (!Evaluator)
+    			if (!Channel.Evaluator)
     			{
-    				Evaluator = MakeUnique<TBaseEvaluator<ValueType>>();
+    				Channel.Evaluator = MakeUnique<TBaseEvaluator<ValueType>>();
     			}
-    			*Evaluator << Ar;
+    			Ar << *Channel.Evaluator;
     		}
     	}
         
@@ -172,6 +174,29 @@ public:
         return Evaluator->InterpolateWindow(Window, Parameter);
     }
 
+	// Evaluate nth derivative of attribute
+	template<int32 Order>
+	ValueType EvaluateDerivative(float Parameter) const
+    {
+    	if (Storage.GetEntries().Num() == 0) return ValueType();
+        
+    	// Use policy for derivative calculation
+    	auto Window = Evaluator->FindWindow(Storage.GetEntries(), Parameter);
+
+    	// Get parameter range from evaluator
+    	float Start, End;
+    	Evaluator->GetParentSpaceRange(Start, End);
+    	const float ParamScale = End - Start;
+        
+    	if (FMath::IsNearlyZero(ParamScale))
+    	{
+    		return ValueType();
+    	}
+        
+    	return TSplineInterpolationPolicy<ValueType>::CalculateDerivative<Order>(
+			Window, Parameter, 1.0f / ParamScale);
+    }
+
     TBaseEvaluator<ValueType>& GetEvaluator()
     {
         return *Evaluator;
@@ -183,6 +208,11 @@ public:
     }
     
     // Access to underlying storage
+	TAttributeStorage<ValueType>& GetStorage() 
+    {
+        return Storage;
+    }
+    
     const TAttributeStorage<ValueType>& GetStorage() const
     {
         return Storage;
@@ -206,7 +236,15 @@ public:
         FindWindowOverride = nullptr;
         InterpolateOverride = nullptr;
     }
+
+    virtual void Serialize(FArchive& Ar) override
+    {
+		IAttributeChannelBase::Serialize(Ar);
+		
+		Ar << *this;
+    }
 };
 
+} // end namespace UE::Geometry::Spline
 } // end namespace UE::Geometry
 } // end namespace UE
\ No newline at end of file
diff --git a/Engine/Source/Runtime/GeometryCore/Public/Splines/BSpline.h b/Engine/Source/Runtime/GeometryCore/Public/Splines/BSpline.h
--- a/Engine/Source/Runtime/GeometryCore/Public/Splines/BSpline.h
+++ b/Engine/Source/Runtime/GeometryCore/Public/Splines/BSpline.h
@@ -4,11 +4,13 @@
 
 #include "Spline.h"
 #include "BaseEvaluator.h"
+#include "SplineMath.h"
 namespace UE
 {
 namespace Geometry
 {
-
+namespace Spline
+{
 enum class EBSplineKnotMode 
 {
 	/** Equal spacing between knots regardless of control point positions */
@@ -91,67 +93,30 @@ public:
 		ComputeBasisFunctions(FindKnotSpan(MappedParam), MappedParam, Basis);
 		return TSplineInterpolationPolicy<ValueType>::InterpolateWithBasis(Window, Basis);
 	}
-
-	virtual ValueType EvaluateDerivative(TArrayView<const TStorageEntry<ValueType>> Entries, float Parameter) const override 
-	{
-		// Find window using existing window logic
-		TArray<const ValueType*> Window = FindWindow(Entries, Parameter);
-		if (Window.Num() < GetDegree() + 1)
-		{
-			return ValueType();
-		}
-
-		// Calculate derivative basis functions
-		TArray<float> DerivBasis;
-		const float u = Base::MapToParentSpace(Parameter);
-		const int32 k = FindKnotSpan(u);
-		ComputeDerivativeBasisFunctions(k, u, DerivBasis);
-
-		// Apply basis using type-specific policy interpolation
-		return TSplineInterpolationPolicy<ValueType>::InterpolateWithBasis(Window, DerivBasis);
-	}
-
-	virtual ValueType EvaluateSecondDerivative(TArrayView<const TStorageEntry<ValueType>> Entries, float Parameter) const override
-	{
-		TArray<const ValueType*> Window = FindWindow(Entries, Parameter);
-		if (Window.Num() < GetDegree() + 1)
-		{
-			return ValueType();
-		}
-
-		// Calculate second derivative basis functions 
-		TArray<float> SecondDerivBasis;
-		const float u = Base::MapToParentSpace(Parameter);
-		const int32 k = FindKnotSpan(u);
-		ComputeSecondDerivativeBasisFunctions(k, u, SecondDerivBasis);
-
-		// Apply basis using type-specific policy interpolation
-		return TSplineInterpolationPolicy<ValueType>::InterpolateWithBasis(Window, SecondDerivBasis);
-	}
-
+	
 	virtual int32 GetDegree() const override
 	{
 		return Degree;
 	}
 
-	virtual FArchive& operator<<(FArchive& Ar) override
+	friend FArchive& operator<<(FArchive& Ar, TBSplineEvaluator& Evaluator)
 	{
 		// Serialize base evaluator first
-		Base::operator<<(Ar);
+		Ar << static_cast<Base&>(Evaluator);
         
 		// Serialize BSpline-specific data
-		Ar << KnotVector;
-		Ar << NumValues;
+		Ar << Evaluator.KnotVector;
+		Ar << Evaluator.NumValues;
         
-		uint8 KnotModeInt = static_cast<uint8>(this->KnotMode);
+		uint8 KnotModeInt = static_cast<uint8>(Evaluator.KnotMode);
 		Ar << KnotModeInt;
         
 		if (Ar.IsLoading())
 		{
-			this->KnotMode = static_cast<EBSplineKnotMode>(KnotModeInt);
+			Evaluator.KnotMode = static_cast<EBSplineKnotMode>(KnotModeInt);
 		}
         
-		Ar << bClampEnds;
+		Ar << Evaluator.bClampEnds;
         
 		return Ar;
 	}
@@ -293,6 +258,29 @@ public:
 		GetKnotRange(Start, End);
 		return FMath::Lerp(Start, End, NormalizedParam);
 	}
+
+	/** 
+	 * Evaluate nth derivative at parameter
+	 * @param Order - The derivative order (0 = position, 1 = first derivative, etc.)
+	 * @param Entries - Spline control points
+	 * @param Parameter - Parameter in normalized space [0,1]
+	 * @return nth derivative vector
+	 */
+	template<int32 Order>
+	ValueType EvaluateDerivative(
+		TArrayView<const TStorageEntry<ValueType>> Entries,
+		float Parameter) const
+	{
+		// Order 0 is just regular evaluation
+		if constexpr (Order == 0)
+		{
+			return InterpolateWindow(FindWindow(Entries, Parameter), Parameter);
+		}
+
+		TArray<const ValueType*> Window = FindWindow(Entries, Parameter);
+		const float MappedParam = Base::MapToParentSpace(Parameter);
+		return Math::TBSplineDerivativeCalculator<ValueType, Order>::Compute(Window, MappedParam);
+	}
 	
 protected:
 	virtual float GetParentSpaceMapping(float NormalizedParam) const override
@@ -732,6 +720,27 @@ public:
 		}
 		return Index;
 	}
+
+	/**
+	 * Updates an existing value at the specified index and updates knot vector accordingly
+	 * @param Index - Index of the value to update
+	 * @param NewValue - New value to set
+	 * @return true if the index was valid and value was updated
+	 */
+	virtual bool UpdateValue(int32 Index, const ValueType& NewValue) override
+	{
+		if (!Base::UpdateValue(Index, NewValue))
+		{
+			return false;
+		}
+    
+		// Update knot vector if needed
+		if (Base::GetEvaluator().GetKnotMode() != EBSplineKnotMode::Custom)
+		{
+			Base::GetEvaluator().UpdateKnotVector(Base::Values.GetStorage().GetEntries());
+		}
+		return true;
+	}
 	
 	/**
 	 * Changes how the knot vector is generated
@@ -765,6 +774,6 @@ using FQuadraticBSpline2f = TBSpline<float, 2, 2>;
 using FCubicBSpline3f = TBSpline<float, 3, 3>;
 using FQuarticBSpline3f = TBSpline<float, 3, 4>;
 using FQuinticBSpline2d = TBSpline<double, 2, 5>;
-
+} // end namespace UE::Geometry::Spline
 } // end namespace UE::Geometry
 } // end namespace UE	
\ No newline at end of file
diff --git a/Engine/Source/Runtime/GeometryCore/Public/Splines/BaseEvaluator.h b/Engine/Source/Runtime/GeometryCore/Public/Splines/BaseEvaluator.h
--- a/Engine/Source/Runtime/GeometryCore/Public/Splines/BaseEvaluator.h
+++ b/Engine/Source/Runtime/GeometryCore/Public/Splines/BaseEvaluator.h
@@ -4,13 +4,16 @@
 
 #include "CoreMinimal.h"
 #include "InterpolationPolicies.h"
+#include "SplineMath.h"
+DEFINE_LOG_CATEGORY_STATIC(LogSpline, Log, All);
 
 namespace UE
 {
 namespace Geometry
 {
-	
-DEFINE_LOG_CATEGORY_STATIC(LogSpline, Log, All);
+namespace Spline
+{
+
 /** 
  * Enum for handling out-of-bounds parameter values
  */
@@ -20,7 +23,7 @@ enum class EOutOfBoundsHandlingMode
 	Constant,   // Choose the nearest of the first/last value
 	Cycle       // Repeat the curve indefinitely (mod of the parameter space)
 };
-
+	
 /**
  * Base evaluator class providing parameter mapping and bounds handling
  */
@@ -44,6 +47,34 @@ public:
 	}
 	virtual ~TBaseEvaluator() = default;
 	
+	/**
+	 * Serializes evaluator data to/from an archive
+	 * @param Ar - Archive to serialize to/from
+	 * @param Evaluator 
+	 */
+	friend void operator<<(FArchive& Ar, TBaseEvaluator& Evaluator)
+	{
+		// Serialize parameter range
+		Ar << Evaluator.ParentStart;
+		Ar << Evaluator.ParentEnd;
+        
+		// Serialize infinity modes
+		uint8 PreMode = static_cast<uint8>(Evaluator.PreInfinityMode);
+		uint8 PostMode = static_cast<uint8>(Evaluator.PostInfinityMode);
+		Ar << PreMode;
+		Ar << PostMode;
+        
+		if (Ar.IsLoading())
+		{
+			Evaluator.PreInfinityMode = static_cast<EOutOfBoundsHandlingMode>(PreMode);
+			Evaluator.PostInfinityMode = static_cast<EOutOfBoundsHandlingMode>(PostMode);
+		}
+	}
+	void Serialize(FArchive& Ar)
+	{
+		Ar << *this;
+	}
+	
 	/** Get the polynomial degree of the evaluator */
 	virtual int32 GetDegree() const{ return 1; }
 
@@ -103,43 +134,28 @@ public:
 		return TSplineInterpolationPolicy<ValueType>::Interpolate(Window, MappedParam);
 	}
 
-	/** Evaluates first derivative at parameter */
-	virtual ValueType EvaluateDerivative(
-		TArrayView<const TStorageEntry<ValueType>> Entries,
-		float Parameter) const
-	{
-		TArray<const ValueType*> Window = FindWindow(Entries, Parameter);
-		if (Window.Num() < 2)
-		{
-			return ValueType();
-		}
-
-		float Start, End;
-		GetParentSpaceRange(Start, End);
-		const float ParamSpaceScale = 1.0f / (End - Start);
-
-		return TSplineInterpolationPolicy<ValueType>::CalculateDerivative(
-			*Window[0], *Window[1], ParamSpaceScale);
-	}
+	/** 
+     * Evaluate nth derivative at parameter
+     * @param Order - The derivative order (0 = position, 1 = first derivative, etc.)
+     * @param Entries - Spline control points
+     * @param Parameter - Parameter in normalized space [0,1]
+     * @return nth derivative vector
+     */
+    template<int32 Order>
+    ValueType EvaluateDerivative(
+        TArrayView<const TStorageEntry<ValueType>> Entries,
+        float Parameter) const
+    {
+        // Order 0 is just regular evaluation
+        if constexpr (Order == 0)
+        {
+            return InterpolateWindow(FindWindow(Entries, Parameter), Parameter);
+        }
 
-	/** Evaluates second derivative at parameter */
-	virtual ValueType EvaluateSecondDerivative(
-		TArrayView<const TStorageEntry<ValueType>> Entries,
-		float Parameter) const
-	{
 		TArray<const ValueType*> Window = FindWindow(Entries, Parameter);
-		if (Window.Num() < 3)
-		{
-			return ValueType();
-		}
-
-		float Start, End;
-		GetParentSpaceRange(Start, End);
-		const float ParamSpaceScale = 1.0f / FMath::Square(End - Start);
-
-		return TSplineInterpolationPolicy<ValueType>::CalculateSecondDerivative(
-			*Window[0], *Window[1], *Window[2], ParamSpaceScale);
-	}
+		const float MappedParam = MapToParentSpace(Parameter);
+		return TDerivativePolicy<ValueType, Order>::Compute(Window, MappedParam);
+    }
 
 	/** Clone the evaluator */
 	virtual TUniquePtr<TBaseEvaluator> Clone() const
@@ -285,6 +301,6 @@ private:
 		return Parameter;
 	}
 };
-
+} // end namespace UE::Geometry::Spline
 } // end namespace UE::Geometry
 } // end namespace UE
\ No newline at end of file
diff --git a/Engine/Source/Runtime/GeometryCore/Public/Splines/InterpolationPolicies.h b/Engine/Source/Runtime/GeometryCore/Public/Splines/InterpolationPolicies.h
--- a/Engine/Source/Runtime/GeometryCore/Public/Splines/InterpolationPolicies.h
+++ b/Engine/Source/Runtime/GeometryCore/Public/Splines/InterpolationPolicies.h
@@ -1,12 +1,16 @@
 // Copyright Epic Games, Inc. All Rights Reserved.
 
 #pragma once
-
+#include "CoreMinimal.h"
+#include "SplineMath.h"
 #include "ParameterizedTypes.h"
 
-/**
- * Common interpolation policies for built-in types
- */
+namespace UE
+{
+namespace Geometry
+{
+namespace Spline
+{
 /**
  * Base spline interpolation policy that works across all spline types
  * Extends the existing parameter-based interpolation with window-based methods
@@ -15,40 +19,37 @@ template<typename T>
 class TSplineInterpolationPolicy
 {
 public:
-    // Keep existing parameter-based interpolation (used by attribute channels)
-    static T Interpolate(TArrayView<const T* const> Window, float Parameter)
-    {
-        if (Window.Num() < 2)
-        {
-            return Window.Num() == 1 ? *Window[0] : T();
-        }
-        return FMath::Lerp(*Window[0], *Window[1], Parameter);
-    }
-
-    // Add window-based interpolation for spline implementations
-    static T InterpolateWithBasis(TArrayView<const T* const> Window, TArrayView<const float> Basis)
-    {
-        T Result = T();
-        for (int32 i = 0; i < FMath::Min(Window.Num(), Basis.Num()); ++i)
-        {
-            Result += *Window[i] * Basis[i];
-        }
-        return Result;
-    }
-
-	// Default derivative calculation
-	static T CalculateDerivative(const T& Start, const T& End, float ParamSpaceScale)
-    {
-    	return (End - Start) * ParamSpaceScale;
-    }
-
-	// Default second derivative calculation
-	static T CalculateSecondDerivative(const T& Start, const T& Mid, const T& End, float ParamSpaceScale)
-    {
-    	T FirstDeriv1 = Mid - Start;
-    	T FirstDeriv2 = End - Mid;
-    	return (FirstDeriv2 - FirstDeriv1) * ParamSpaceScale;
-    }
+	// Keep existing parameter-based interpolation (used by attribute channels)
+	static T Interpolate(TArrayView<const T* const> Window, float Parameter)
+	{
+		if (Window.Num() < 2)
+		{
+			return Window.Num() == 1 ? *Window[0] : T();
+		}
+		return FMath::Lerp(*Window[0], *Window[1], Parameter);
+	}
+
+	// Add window-based interpolation for spline implementations
+	static T InterpolateWithBasis(TArrayView<const T* const> Window, TArrayView<const float> Basis)
+	{
+		T Result = T();
+		for (int32 i = 0; i < FMath::Min(Window.Num(), Basis.Num()); ++i)
+		{
+			Result += *Window[i] * Basis[i];
+		}
+		return Result;
+	}
+
+	// Template-based n-th derivative calculation
+	template<int32 Order>
+	static T EvaluateDerivative(TArrayView<const T* const> Window, float Parameter)
+	{
+		if (Order == 0)
+		{
+			return Interpolate(Window, Parameter);
+		}
+		return Math::TGenericDerivativeHelper<T>::Compute<Order>(Window, Parameter);
+	}
 };
 
 // FLinearColor interpolation policy
@@ -56,53 +57,48 @@ template<>
 class TSplineInterpolationPolicy<FLinearColor>
 {
 public:
-    static FLinearColor Interpolate(TArrayView<const FLinearColor* const> Window, float Parameter)
-    {
-        if (Window.Num() < 2)
-        {
-            return Window.Num() == 1 ? *Window[0] : FLinearColor();
-        }
-        
-        // Convert to linear color for better interpolation
-        return FLinearColor::LerpUsingHSV(
-            *Window[0],
-            *Window[1],
-            Parameter
-        ).ToFColor(true);
-    }
-
-    static FLinearColor InterpolateWithBasis(TArrayView<const FLinearColor* const> Window, TArrayView<const float> Basis)
-    {
-        // Convert all colors to linear space first
-        TArray<FLinearColor> LinearColors;
-        LinearColors.Reserve(Window.Num());
-        for (const FLinearColor* Color : Window)
-        {
-            LinearColors.Add(*Color);
-        }
-
-        // Perform weighted sum in linear space
-        FLinearColor Result = FLinearColor::Black;
-        for (int32 i = 0; i < FMath::Min(Window.Num(), Basis.Num()); ++i)
-        {
-            Result += LinearColors[i] * Basis[i];
-        }
-
-        // Convert back to FColor
-        return Result;
-    }
-
-	static FLinearColor CalculateDerivative(const FLinearColor& Start, const FLinearColor& End, float ParamSpaceScale)
-    {
-    	return (End - Start) * ParamSpaceScale;
-    }
-
-	static FLinearColor CalculateSecondDerivative(const FLinearColor& Start, const FLinearColor& Mid, const FLinearColor& End, float ParamSpaceScale)
-    {
-    	FLinearColor FirstDeriv1 = Mid - Start;
-    	FLinearColor FirstDeriv2 = End - Mid;
-    	return (FirstDeriv2 - FirstDeriv1) * ParamSpaceScale;
-    }
+	static FLinearColor Interpolate(TArrayView<const FLinearColor* const> Window, float Parameter)
+	{
+		if (Window.Num() < 2)
+		{
+			return Window.Num() == 1 ? *Window[0] : FLinearColor();
+		}
+
+		// Convert to linear color for better interpolation
+		return FLinearColor::LerpUsingHSV(*Window[0],*Window[1], Parameter).ToFColor(true);
+	}
+
+	static FLinearColor InterpolateWithBasis(TArrayView<const FLinearColor* const> Window, TArrayView<const float> Basis)
+	{
+		// Convert all colors to linear space first
+		TArray<FLinearColor> LinearColors;
+		LinearColors.Reserve(Window.Num());
+		for (const FLinearColor* Color : Window)
+		{
+			LinearColors.Add(*Color);
+		}
+
+		// Perform weighted sum in linear space
+		FLinearColor Result = FLinearColor::Black;
+		for (int32 i = 0; i < FMath::Min(Window.Num(), Basis.Num()); ++i)
+		{
+			Result += LinearColors[i] * Basis[i];
+		}
+
+		// Convert back to FColor
+		return Result;
+	}
+
+	// Template-based n-th derivative calculation
+	template<int32 Order>
+	static FLinearColor EvaluateDerivative(TArrayView<const FLinearColor* const> Window, float Parameter)
+	{
+		if (Order == 0)
+		{
+			return Interpolate(Window, Parameter);
+		}
+		return Math::TGenericDerivativeHelper<FLinearColor>::Compute<Order>(Window, Parameter);
+	}
 };
 
 // FQuat interpolation policy
@@ -110,132 +106,83 @@ template<>
 class TSplineInterpolationPolicy<FQuat>
 {
 public:
-    // Existing parameter-based interpolation
-    static FQuat Interpolate(TArrayView<const FQuat* const> Window, float Parameter)
-    {
-        if (Window.Num() < 2)
-        {
-            return Window.Num() == 1 ? *Window[0] : FQuat::Identity;
-        }
-        return FQuat::Slerp(*Window[0], *Window[1], Parameter);
-    }
-
-    // Window-based interpolation for spline implementations
-    static FQuat InterpolateWithBasis(TArrayView<const FQuat* const> Window, TArrayView<const float> Basis)
-    {
-        FQuat Result = FQuat::Identity;
-        float TotalWeight = 0.0f;
-        
-        for (int32 i = 0; i < FMath::Min(Window.Num(), Basis.Num()); ++i)
-        {
-            const float Weight = Basis[i];
-            if (Weight > UE_SMALL_NUMBER)
-            {
-                const FQuat& Q = *Window[i];
-                const double Dot = Result.W * Q.W + Result.X * Q.X + Result.Y * Q.Y + Result.Z * Q.Z;
-                Result += (Dot >= 0.0f ? Q : -Q) * Weight;
-                TotalWeight += Weight;
-            }
-        }
-
-        if (TotalWeight > UE_SMALL_NUMBER)
-        {
-            Result.Normalize();
-        }
-        return Result;
-    }
-
-	static FQuat CalculateDerivative(const FQuat& Start, const FQuat& End, float ParamSpaceScale)
-    {
-    	const double Dot = Start | End;
-    	const FQuat TargetQuat = (Dot < 0) ? -End : End;
-    	FQuat Delta = TargetQuat * Start.Inverse();
-        
-    	FVector Axis;
-    	float Angle;
-    	Delta.ToAxisAndAngle(Axis, Angle);
-        
-    	const FVector AngularVelocity = Axis * Angle * ParamSpaceScale;
-    	return FQuat(AngularVelocity.X, AngularVelocity.Y, AngularVelocity.Z, 0.0f);
-    }
-
-	static FQuat CalculateSecondDerivative(const FQuat& Start, const FQuat& Mid, const FQuat& End, float ParamSpaceScale)
-    {
-    	auto CalcAngularVelocity = [](const FQuat& q0, const FQuat& q1) -> FVector
-    	{
-    		const double Dot = q0 | q1;
-    		const FQuat TargetQuat = (Dot < 0) ? -q1 : q1;
-    		FQuat Delta = TargetQuat * q0.Inverse();
-            
-    		FVector Axis;
-    		float Angle;
-    		Delta.ToAxisAndAngle(Axis, Angle);
-    		return Axis * Angle;
-    	};
-
-    	FVector AngVel1 = CalcAngularVelocity(Start, Mid);
-    	FVector AngVel2 = CalcAngularVelocity(Mid, End);
-    	const FVector AngularVelocity = (AngVel2 - AngVel1) * ParamSpaceScale;
-    	return FQuat(AngularVelocity.X, AngularVelocity.Y, AngularVelocity.Z, 0.0f);
-    }
-	
-};
+	// Existing parameter-based interpolation
+	static FQuat Interpolate(TArrayView<const FQuat* const> Window, float Parameter)
+	{
+		if (Window.Num() < 2)
+		{
+			return Window.Num() == 1 ? *Window[0] : FQuat::Identity;
+		}
+		return FQuat::Slerp(*Window[0], *Window[1], Parameter);
+	}
+
+	// Window-based interpolation for spline implementations
+	static FQuat InterpolateWithBasis(TArrayView<const FQuat* const> Window, TArrayView<const float> Basis)
+	{
+		FQuat Result = FQuat::Identity;
+		float TotalWeight = 0.0f;
+
+		for (int32 i = 0; i < FMath::Min(Window.Num(), Basis.Num()); ++i)
+		{
+			const float Weight = Basis[i];
+			if (Weight > UE_SMALL_NUMBER)
+			{
+				const FQuat& Q = *Window[i];
+				const double Dot = Result.W * Q.W + Result.X * Q.X + Result.Y * Q.Y + Result.Z * Q.Z;
+				Result += (Dot >= 0.0f ? Q : -Q) * Weight;
+				TotalWeight += Weight;
+			}
+		}
+
+		if (TotalWeight > UE_SMALL_NUMBER)
+		{
+			Result.Normalize();
+		}
+		return Result;
+	}
+
+	// Template-based n-th derivative calculation
+	template<int32 Order>
+	static FQuat EvaluateDerivative(TArrayView<const FQuat* const> Window, float Parameter)
+	{
+		if (Order == 0)
+		{
+			return Interpolate(Window, Parameter);
+		}
+
+		// First derivative - angular velocity
+		if (Order == 1)
+		{
+			// Adaptive h based on the difference between quaternions at Parameter and Parameter + base_h
+			float base_h = 0.01f; // Start with a reasonable base value
+			FQuat Q0 = Interpolate(Window, Parameter);
+			FQuat Q1base = Interpolate(Window, Parameter + base_h);
+			float diff = static_cast<float>(Q0.AngularDistance(Q1base)); // Use a suitable distance metric
+			float h = base_h * FMath::Clamp(diff, 0.001f, 0.1f); // Adjust h based on the difference
+
+			FQuat Q1 = Interpolate(Window, Parameter + h);
+			return (Q1 * Q0.Inverse()).Log() * (1.0f / h);
+		}
+
+		// Second derivative - angular acceleration
+		if (Order == 2)
+		{
+			// Use the same adaptive h as the first derivative, centered around the parameter
+			float base_h = 0.01f;
+			FQuat Q_m1 = Interpolate(Window, Parameter - base_h);
+			FQuat Q_p1 = Interpolate(Window, Parameter + base_h);
+			float diff = static_cast<float>(Q_m1.AngularDistance(Q_p1));
+			float h = base_h * FMath::Clamp(diff / 2.0f, 0.001f, 0.1f);  // Divide diff by 2 for centering
+
+			FQuat V0 = EvaluateDerivative<1>(Window, Parameter - h);
+			FQuat V1 = EvaluateDerivative<1>(Window, Parameter + h);
+			return (V1 * V0.Inverse()).Log() * (1.0f / (2.0f * h));
+		}
+
+		// Higher orders generally not meaningful for rotations
+		return FQuat::Identity;
+	}
 
-// FRotator interpolation policy
-template<>
-class TSplineInterpolationPolicy<FRotator>
-{
-public:
-    static FRotator Interpolate(TArrayView<const FRotator* const> Window, float Parameter)
-    {
-        if (Window.Num() < 2)
-        {
-            return Window.Num() == 1 ? *Window[0] : FRotator::ZeroRotator;
-        }
-        
-        return FMath::LerpRange(*Window[0], *Window[1], Parameter);
-    }
-
-    static FRotator InterpolateWithBasis(TArrayView<const FRotator* const> Window, TArrayView<const float> Basis)
-    {
-        // Convert to quaternions for better interpolation
-        TArray<FQuat> Quats;
-        Quats.Reserve(Window.Num());
-        for (const FRotator* Rot : Window)
-        {
-            Quats.Add(Rot->Quaternion());
-        }
-
-        // Use quaternion interpolation
-        FQuat Result = TSplineInterpolationPolicy<FQuat>::InterpolateWithBasis(
-            MakeArrayView((const FQuat**)Quats.GetData(), Quats.Num()),
-            Basis
-        );
-
-        return FRotator(Result);
-    }
-	static FRotator CalculateDerivative(const FRotator& Start, const FRotator& End, float ParamSpaceScale)
-    {
-    	// Convert to quaternions for better interpolation
-    	FQuat StartQuat = Start.Quaternion();
-    	FQuat EndQuat = End.Quaternion();
-        
-    	// Use quaternion derivative calculation and convert back to rotator
-    	FQuat QuatDeriv = TSplineInterpolationPolicy<FQuat>::CalculateDerivative(StartQuat, EndQuat, ParamSpaceScale);
-    	return FRotator(QuatDeriv);
-    }
-
-	static FRotator CalculateSecondDerivative(const FRotator& Start, const FRotator& Mid, const FRotator& End, float ParamSpaceScale)
-    {
-    	// Convert to quaternions for better interpolation
-    	FQuat StartQuat = Start.Quaternion();
-    	FQuat MidQuat = Mid.Quaternion();
-    	FQuat EndQuat = End.Quaternion();
-        
-    	// Use quaternion second derivative calculation and convert back to rotator
-    	FQuat QuatSecondDeriv = TSplineInterpolationPolicy<FQuat>::CalculateSecondDerivative(StartQuat, MidQuat, EndQuat, ParamSpaceScale);
-    	return FRotator(QuatSecondDeriv);
-    }
 };
 
 // FTransform interpolation policy
@@ -243,98 +190,112 @@ template<>
 class TSplineInterpolationPolicy<FTransform>
 {
 public:
-    static FTransform Interpolate(TArrayView<const FTransform* const> Window, float Parameter)
-    {
-        if (Window.Num() < 2)
-        {
-            return Window.Num() == 1 ? *Window[0] : FTransform::Identity;
-        }
-        
-        const FTransform& A = *Window[0];
-        const FTransform& B = *Window[1];
-        
-        FTransform Result;
-        Result.SetLocation(FMath::Lerp(A.GetLocation(), B.GetLocation(), Parameter));
-        Result.SetRotation(FQuat::Slerp(A.GetRotation(), B.GetRotation(), Parameter));
-        Result.SetScale3D(FMath::Lerp(A.GetScale3D(), B.GetScale3D(), Parameter));
-        
-        return Result;
-    }
-
-    static FTransform InterpolateWithBasis(TArrayView<const FTransform* const> Window, TArrayView<const float> Basis)
-    {
-        // Handle empty case
-        if (Window.Num() == 0 || Basis.Num() == 0)
-        {
-            return FTransform::Identity;
-        }
-
-        // Interpolate components separately
-        FVector Location = FVector::ZeroVector;
-        FQuat Rotation = FQuat::Identity;
-        FVector Scale = FVector(1.0f);
-
-        float TotalRotWeight = 0.0f;
+	static FTransform Interpolate(TArrayView<const FTransform* const> Window, float Parameter)
+	{
+		if (Window.Num() < 2)
+		{
+			return Window.Num() == 1 ? *Window[0] : FTransform::Identity;
+		}
+
+		const FTransform& A = *Window[0];
+		const FTransform& B = *Window[1];
+
+		FTransform Result;
+		Result.SetLocation(FMath::Lerp(A.GetLocation(), B.GetLocation(), Parameter));
+		Result.SetRotation(FQuat::Slerp(A.GetRotation(), B.GetRotation(), Parameter));
+		Result.SetScale3D(FMath::Lerp(A.GetScale3D(), B.GetScale3D(), Parameter));
+
+		return Result;
+	}
+
+	static FTransform InterpolateWithBasis(TArrayView<const FTransform* const> Window, TArrayView<const float> Basis)
+	{
+		// Handle empty case
+		if (Window.Num() == 0 || Basis.Num() == 0)
+		{
+			return FTransform::Identity;
+		}
+
+		// Interpolate components separately
+		FVector Location = FVector::ZeroVector;
+		FQuat Rotation = FQuat::Identity;
+		FVector Scale = FVector(1.0f);
+
+		float TotalRotWeight = 0.0f;
+
+		// Weighted sum of locations
+		for (int32 i = 0; i < FMath::Min(Window.Num(), Basis.Num()); ++i)
+		{
+			const float Weight = Basis[i];
+			if (Weight > UE_SMALL_NUMBER)
+			{
+				Location += Window[i]->GetLocation() * Weight;
+				Scale += Window[i]->GetScale3D() * Weight;
+
+				// Handle rotation separately with proper quaternion blending
+				const FQuat& Q = Window[i]->GetRotation();
+				const double Dot = Rotation.W * Q.W + Rotation.X * Q.X + Rotation.Y * Q.Y + Rotation.Z * Q.Z;
+				Rotation += (Dot >= 0.0f ? Q : -Q) * Weight;
+				TotalRotWeight += Weight;
+			}
+		}
+
+		// Normalize the quaternion if we accumulated any weights
+		if (TotalRotWeight > UE_SMALL_NUMBER)
+		{
+			Rotation.Normalize();
+		}
+
+		return FTransform(Rotation, Location, Scale);
+	}
+
+	template<int32 Order>
+	static FTransform EvaluateDerivative(TArrayView<const FTransform* const> Window, float Parameter)
+	{
+		if (Order == 0)
+		{
+			return Interpolate(Window, Parameter);
+		}
+
+		// Split into components
+		TArray<FVector> Translations;
+		TArray<FQuat> Rotations;
+		TArray<FVector> Scales;
         
-        // Weighted sum of locations
-        for (int32 i = 0; i < FMath::Min(Window.Num(), Basis.Num()); ++i)
-        {
-            const float Weight = Basis[i];
-            if (Weight > UE_SMALL_NUMBER)
-            {
-                Location += Window[i]->GetLocation() * Weight;
-                Scale += Window[i]->GetScale3D() * Weight;
-
-                // Handle rotation separately with proper quaternion blending
-                const FQuat& Q = Window[i]->GetRotation();
-                const double Dot = Rotation.W * Q.W + Rotation.X * Q.X + Rotation.Y * Q.Y + Rotation.Z * Q.Z;
-                Rotation += (Dot >= 0.0f ? Q : -Q) * Weight;
-                TotalRotWeight += Weight;
-            }
-        }
-
-        // Normalize the quaternion if we accumulated any weights
-        if (TotalRotWeight > UE_SMALL_NUMBER)
-        {
-            Rotation.Normalize();
-        }
-
-        return FTransform(Rotation, Location, Scale);
-    }
-
-	static FTransform CalculateDerivative(const FTransform& Start, const FTransform& End, float ParamSpaceScale)
-    {
-    	FTransform Result;
-    	Result.SetLocation((End.GetLocation() - Start.GetLocation()) * ParamSpaceScale);
+		Translations.SetNum(Window.Num());
+		Rotations.SetNum(Window.Num());
+		Scales.SetNum(Window.Num());
         
-    	// Use quaternion derivative for rotation
-    	const FQuat& StartRot = Start.GetRotation();
-    	const FQuat& EndRot = End.GetRotation();
-    	Result.SetRotation(TSplineInterpolationPolicy<FQuat>::CalculateDerivative(StartRot, EndRot, ParamSpaceScale));
-        
-    	Result.SetScale3D((End.GetScale3D() - Start.GetScale3D()) * ParamSpaceScale);
-    	return Result;
-    }
+		for (int32 i = 0; i < Window.Num(); ++i)
+		{
+			Translations[i] = Window[i]->GetTranslation();
+			Rotations[i] = Window[i]->GetRotation();
+			Scales[i] = Window[i]->GetScale3D();
+		}
+
+		// Calculate derivatives for each component
+		TArrayView<const FVector> TransView(Translations.GetData(), Translations.Num());
+		TArrayView<const FQuat> RotView(Rotations.GetData(), Rotations.Num());
+		TArrayView<const FVector> ScaleView(Scales.GetData(), Scales.Num());
+
+		FVector TransResult = TSplineInterpolationPolicy<FVector>::EvaluateDerivative<Order>(TransView, Parameter);
+		FQuat RotResult = TSplineInterpolationPolicy<FQuat>::EvaluateDerivative<Order>(RotView, Parameter);
+		FVector ScaleResult = TSplineInterpolationPolicy<FVector>::EvaluateDerivative<Order>(ScaleView, Parameter);
+
+		return FTransform(RotResult, TransResult, ScaleResult);
+	}
+};
 
-	static FTransform CalculateSecondDerivative(const FTransform& Start, const FTransform& Mid, const FTransform& End, float ParamSpaceScale)
-    {
-    	FTransform Result;
-        
-    	// Location and scale use standard second derivative
-    	FVector LocationDeriv1 = Mid.GetLocation() - Start.GetLocation();
-    	FVector LocationDeriv2 = End.GetLocation() - Mid.GetLocation();
-    	Result.SetLocation((LocationDeriv2 - LocationDeriv1) * ParamSpaceScale);
-
-    	// Use quaternion second derivative for rotation
-    	const FQuat& StartRot = Start.GetRotation();
-    	const FQuat& MidRot = Mid.GetRotation();
-    	const FQuat& EndRot = End.GetRotation();
-    	Result.SetRotation(TSplineInterpolationPolicy<FQuat>::CalculateSecondDerivative(StartRot, MidRot, EndRot, ParamSpaceScale));
-
-    	FVector ScaleDeriv1 = Mid.GetScale3D() - Start.GetScale3D();
-    	FVector ScaleDeriv2 = End.GetScale3D() - Mid.GetScale3D();
-    	Result.SetScale3D((ScaleDeriv2 - ScaleDeriv1) * ParamSpaceScale);
-
-    	return Result;
-    }
-};
\ No newline at end of file
+
+// Base derivative calculator
+template <typename T, int32 Order>
+struct TDerivativePolicy
+{
+	static T Compute(TArrayView<const T* const> Window, float Parameter)
+	{
+		return TSplineInterpolationPolicy<T>::template EvaluateDerivative<Order>(Window, Parameter);
+	}
+};
+} // end namespace UE::Geometry::Spline
+} // end namespace UE::Geometry
+} // end namespace UE
\ No newline at end of file
diff --git a/Engine/Source/Runtime/GeometryCore/Public/Splines/ParameterizedStorage.h b/Engine/Source/Runtime/GeometryCore/Public/Splines/ParameterizedStorage.h
--- a/Engine/Source/Runtime/GeometryCore/Public/Splines/ParameterizedStorage.h
+++ b/Engine/Source/Runtime/GeometryCore/Public/Splines/ParameterizedStorage.h
@@ -8,7 +8,8 @@ namespace UE
 {
 namespace Geometry
 {
-	
+namespace Spline
+{
 /**
  * Pure storage class - no evaluation logic
  */
@@ -46,16 +47,23 @@ public:
     	return Index;
     }
 
-    TArrayView<const TStorageEntry<ValueType>> GetEntries() const 
-    { 
-        return Entries; 
-    }
+	// Read-only access through view
+	TArrayView<const TStorageEntry<ValueType>> GetEntries() const
+	{
+		return Entries;
+	}
+
+	// Modifiable access through reference
+	TArray<TStorageEntry<ValueType>>& GetEntries()
+	{
+		return Entries;
+	}
 
     void Clear()
     {
         Entries.Reset();
     }
 };
-
+} // end namespace UE::Geometry::Spline
 } // end namespace UE::Geometry
-} // end namespace UE
\ No newline at end of file
+} // end namespace UE
diff --git a/Engine/Source/Runtime/GeometryCore/Public/Splines/ParameterizedTypes.h b/Engine/Source/Runtime/GeometryCore/Public/Splines/ParameterizedTypes.h
--- a/Engine/Source/Runtime/GeometryCore/Public/Splines/ParameterizedTypes.h
+++ b/Engine/Source/Runtime/GeometryCore/Public/Splines/ParameterizedTypes.h
@@ -8,6 +8,9 @@ namespace UE
 {
 namespace Geometry
 {
+namespace Spline
+{
+
 /**
  * Vector type policies based on dimension and numeric type
  */
@@ -76,5 +79,6 @@ struct TStorageEntry
 	}
 };
 
+} // end namespace UE::Geometry::Spline
 } // end namespace UE::Geometry
 } // end namespace UE
\ No newline at end of file
diff --git a/Engine/Source/Runtime/GeometryCore/Public/Splines/PolyBezierSpline.h b/Engine/Source/Runtime/GeometryCore/Public/Splines/PolyBezierSpline.h
--- a/Engine/Source/Runtime/GeometryCore/Public/Splines/PolyBezierSpline.h
+++ b/Engine/Source/Runtime/GeometryCore/Public/Splines/PolyBezierSpline.h
@@ -7,6 +7,8 @@ namespace UE
 {
 namespace Geometry
 {
+namespace Spline
+{
 /**
  * Custom evaluator for PolyBezier splines.
  * Handles parameter mapping and control point selection for proper Bezier evaluation.
@@ -28,13 +30,13 @@ public:
         this->SetClampedEnds(true);
     }
 
-    virtual FArchive& operator<<(FArchive& Ar) override
+    friend FArchive& operator<<(FArchive& Ar, TPolyBezierEvaluator& Evaluator)
     {
         // Serialize base evaluator first
-        Base::operator<<(Ar);
+        Ar << static_cast<Base&>(Evaluator);
         
         // Serialize PolyBezier-specific data
-        Ar << NumSegments;
+        Ar << Evaluator.NumSegments;
         
         return Ar;
     }
@@ -152,58 +154,29 @@ public:
 
         return Base::SetCustomKnots(BezierKnots, InNumValues);
     }
-    
-    ValueType EvaluateDerivative(TArrayView<const TStorageEntry<ValueType>> Entries, float Parameter) const override
+    /** 
+     * Evaluate nth derivative at parameter
+     * @param Order - The derivative order (0 = position, 1 = first derivative, etc.)
+     * @param Entries - Spline control points
+     * @param Parameter - Parameter in normalized space [0,1]
+     * @return nth derivative vector
+     */
+    template<int32 Order>
+    ValueType EvaluateDerivative(
+        TArrayView<const TStorageEntry<ValueType>> Entries,
+        float Parameter) const
     {
-        TArray<const ValueType*> Window = FindWindow(Entries, Parameter);
-        if (Window.Num() < 4)  // Need all 4 control points for a cubic segment
+        // Order 0 is just regular evaluation
+        if constexpr (Order == 0)
         {
-            return ValueType();
+            return InterpolateWindow(FindWindow(Entries, Parameter), Parameter);
         }
 
-        // Get local parameter within segment [0,1]
-        const float t = FMath::Frac(Base::MapToParentSpace(Parameter));
-
-        // First derivative coefficients for cubic Bezier
-        const float c0 = -3.0f * (1.0f-t) * (1.0f-t);
-        const float c1 = 3.0f * (1.0f-t) * (1.0f-3.0f*t);
-        const float c2 = 3.0f * t * (2.0f-3.0f*t);
-        const float c3 = 3.0f * t * t;
-
-        // Build basis array
-        TArray<float> DerivBasis = { c0, c1, c2, c3 };
-
-        // Scale for parameter space and apply basis
-        return TSplineInterpolationPolicy<ValueType>::InterpolateWithBasis(Window, DerivBasis) * NumSegments;
-    }
-
-    ValueType EvaluateSecondDerivative(TArrayView<const TStorageEntry<ValueType>> Entries, float Parameter) const override
-    {
         TArray<const ValueType*> Window = FindWindow(Entries, Parameter);
-        if (Window.Num() < 4)
-        {
-            return ValueType();
-        }
-
-        // Get local parameter for this segment
-        const float t = FMath::Frac(Base::MapToParentSpace(Parameter));
-
-        // Second derivative coefficients for cubic Bezier
-        const float c0 = 6.0f * (1.0f-t);
-        const float c1 = -12.0f + 18.0f*t;
-        const float c2 = 6.0f - 18.0f*t;
-        const float c3 = 6.0f * t;
-
-        // Build basis array
-        TArray<float> SecondDerivBasis = { c0, c1, c2, c3 };
-
-        // Scale for parameter space and apply basis
-        const float SegmentScale = NumSegments;
-        return TSplineInterpolationPolicy<ValueType>::InterpolateWithBasis(Window, SecondDerivBasis) 
-            * (SegmentScale * SegmentScale);
+        const float MappedParam = Base::MapToParentSpace(Parameter);
+        return Math::TBezierDerivativeCalculator<ValueType, Order>::Compute(Window, MappedParam);
     }
-
-protected:
+   
     virtual float GetParentSpaceMapping(float NormalizedParam) const override
     {
         return NormalizedParam * NumSegments;
@@ -275,23 +248,56 @@ public:
         return NumPoints / 4;
     }
 
-    /** Gets the control points for a specific Bezier segment */
-    bool GetSegmentControlPoints(int32 SegmentIndex, ValueType& OutP0, ValueType& OutP1, 
-        ValueType& OutP2, ValueType& OutP3) const
+    /**
+ * Updates a single control point within a Bezier segment
+ * @param SegmentIndex - Which cubic Bezier segment (each has 4 points)
+ * @param PointIndex - Which point in the segment (0=P0, 1=P1, 2=P2, 3=P3)
+ * @param NewValue - New position for the control point
+ * @return true if successfully updated
+ */
+    bool UpdateSegmentPoint(int32 SegmentIndex, int32 PointIndex, const ValueType& NewValue)
     {
-        if (!IsValidSegmentIndex(SegmentIndex))
+        const int32 NumSegments = GetNumSegments();
+        if (SegmentIndex < 0 || SegmentIndex >= NumSegments ||
+            PointIndex < 0 || PointIndex > 3)
+        {
+            return false;
+        }
+
+        // Convert segment+point index to global index in spline values
+        const int32 GlobalIndex = SegmentIndex * 4 + PointIndex;
+        return UpdateValue(GlobalIndex, NewValue);
+    }
+
+    /**
+     * Updates all control points of a Bezier segment at once
+     * @param SegmentIndex - Which cubic Bezier segment to update
+     * @param P0 - Start point of segment
+     * @param P1 - First control point
+     * @param P2 - Second control point
+     * @param P3 - End point of segment
+     * @return true if successfully updated
+     */
+    bool UpdateSegment(int32 SegmentIndex, const ValueType& P0, const ValueType& P1, 
+        const ValueType& P2, const ValueType& P3)
+    {
+        const int32 NumSegments = GetNumSegments();
+        if (SegmentIndex < 0 || SegmentIndex >= NumSegments)
         {
             return false;
         }
 
         const int32 BaseIndex = SegmentIndex * 4;
-        OutP0 = Base::GetValueAt(BaseIndex);
-        OutP1 = Base::GetValueAt(BaseIndex + 1);
-        OutP2 = Base::GetValueAt(BaseIndex + 2);
-        OutP3 = Base::GetValueAt(BaseIndex + 3);
-        return true;
+        bool bSuccess = true;
+        bSuccess &= Base::UpdateValue(BaseIndex, P0);
+        bSuccess &= Base::UpdateValue(BaseIndex + 1, P1);
+        bSuccess &= Base::UpdateValue(BaseIndex + 2, P2);
+        bSuccess &= Base::UpdateValue(BaseIndex + 3, P3);
+    
+        return bSuccess;
     }
 
+
     /** Returns the number of Bezier segments in the spline */
     int32 GetNumSegments() const
     {
@@ -307,10 +313,7 @@ public:
         }
         return Base::GetValueAt(SegmentIndex * 4 + 3);
     }
-
-private:
-
-     bool IsValidSegmentIndex(int32 Index) const
+    bool IsValidSegmentIndex(int32 Index) const
     {
         return Index >= 0 && Index < GetNumSegments();
     }
@@ -321,6 +324,6 @@ using FPolyBezierSpline2f = TPolyBezierSpline<float, 2>;
 using FPolyBezierSpline3f = TPolyBezierSpline<float, 3>;
 using FPolyBezierSpline2d = TPolyBezierSpline<double, 2>;
 using FPolyBezierSpline3d = TPolyBezierSpline<double, 3>;
-
+} // end namespace UE::Geometry::Spline
 } // end namespace UE::Geometry
 } // end namespace UE
\ No newline at end of file
diff --git a/Engine/Source/Runtime/GeometryCore/Public/Splines/Spline.h b/Engine/Source/Runtime/GeometryCore/Public/Splines/Spline.h
--- a/Engine/Source/Runtime/GeometryCore/Public/Splines/Spline.h
+++ b/Engine/Source/Runtime/GeometryCore/Public/Splines/Spline.h
@@ -9,7 +9,8 @@ namespace UE
 {
 namespace Geometry
 {
-
+namespace Spline
+{
 /** Flags controlling what data gets copied during spline cloning */
 enum class ESplineCloneFlags : uint32
 {
@@ -112,27 +113,27 @@ public:
 	TSplineBase(TSplineBase&&) noexcept = default;
 	TSplineBase& operator=(TSplineBase&&) noexcept = default;
 
-	FArchive& operator<<(FArchive& Ar)
+	friend FArchive& operator<<(FArchive& Ar, TSplineBase& Spline)
 	{
 		// Serialize main value channel
-		Values << Ar;
+		Ar << Spline.Values;
         
 		// Serialize attribute channels
-		int32 NumChannels = AttributeChannels.Num();
+		int32 NumChannels = Spline.AttributeChannels.Num();
 		Ar << NumChannels;
         
 		if (Ar.IsSaving())
 		{
-			for (auto& Pair : AttributeChannels)
+			for (auto& Pair : Spline.AttributeChannels)
 			{
 				// Save channel name and data
 				Ar << Pair.Key;
-				*Pair.Value << Ar;
+				Pair.Value->Serialize(Ar);
 			}
 		}
 		else
 		{
-			AttributeChannels.Empty();
+			Spline.AttributeChannels.Empty();
             
 			for (int32 i = 0; i < NumChannels; ++i)
 			{
@@ -142,9 +143,9 @@ public:
 				// Create new channel
 				TUniquePtr<IAttributeChannelBase> Channel = 
 					MakeUnique<TAttributeChannel<ValueType>>();
-				*Channel << Ar;
+				Channel->Serialize(Ar);
                 
-				AttributeChannels.Add(ChannelName, MoveTemp(Channel));
+				Spline.AttributeChannels.Add(ChannelName, MoveTemp(Channel));
 			}
 		}
         
@@ -152,6 +153,9 @@ public:
 	}
 	/**
 	  * Creates an attribute channel with a custom evaluator
+	  * @tparam AttrType - Custom attribute type
+	  * @param Name - Name of the attribute channel
+	  * @param AttrEvaluator - Custom evaluator for the attribute channel
 	  */
 	template <typename AttrType>
 	void CreateAttributeChannel( const FName& Name,	TUniquePtr<TBaseEvaluator<AttrType>> AttrEvaluator)
@@ -175,6 +179,8 @@ public:
 
 	/**
 	 * Creates an attribute channel that follows the spline's evaluation pattern
+	 * @tparam AttrType - Custom attribute type
+	 * @param Name - Name of the attribute channel 
 	 */
 	template <typename AttrType>
 	void CreateAttributeChannel(const FName& Name)
@@ -190,6 +196,14 @@ public:
 		CreateAttributeChannel(Name, MoveTemp(BaseEvaluator));
 	}
 
+	/**
+	 * Adds a new value to an attribute channel
+	 * @tparam AttrType - Custom attribute type
+	 * @param Name - Name of the attribute channel
+	 * @param Value - Value to add
+	 * @param Parameter - Parameter value in parent spline space
+	 * @return Index of the added value
+	 */
 	template <typename AttrType>
 	int32 AddAttributeValue(FName Name,	const AttrType& Value, float Parameter)
 	{
@@ -200,35 +214,84 @@ public:
 		return INDEX_NONE;
 	}
 
+	/**
+	 * Adds a new value to the spline's main storage channel.
+	 * Also expands the parameter range if necessary.
+	 * @param Value - The value to add
+	 * @param Parameter - The parameter value in parent spline space
+	 * @return 
+	 */
 	virtual int32 AddValue(const ValueType& Value, float Parameter)
 	{
 		GetEvaluator().ExpandRangeToInclude(Parameter);
 		return Values.AddValue(Value, Parameter);		
 	}
 
+	/**
+     * Updates an existing value at the specified index
+     * @param Index - Index of the value to update
+     * @param NewValue - New value to set
+     * @return true if the index was valid and value was updated
+     */
+    virtual bool UpdateValue(int32 Index, const ValueType& NewValue)
+    {
+        TArray<TStorageEntry<ValueType>>& Entries = Values.GetStorage().GetEntries();
+        if (!Entries.IsValidIndex(Index))
+        {
+            return false;
+        }
+        
+        Entries[Index].Value = NewValue;
+        return true;
+    }
+
+	/**
+	 * Evaluates the spline at the given parameter in parent space
+	 * @param Parameter - Parameter value in parent spline space
+	 * @return Value at the given parameter
+	 */ 
 	virtual ValueType Evaluate(float Parameter) const
 	{
 		return Values.Evaluate(Parameter);
 	}
 
-	/** Evaluates first derivative at parameter */
-	virtual ValueType EvaluateDerivative(float Parameter) const
+	/**
+	 * Evaluates a derivative at the given parameter in parent space
+	 * @tparam Order Derivative order
+	 * @param Parameter Parameter value in parent spline space
+	 * @return Derivative value
+	 */
+	template<int32 Order>
+	ValueType EvaluateDerivative(float Parameter) const
 	{
-		return GetEvaluator().EvaluateDerivative(
-			Values.GetStorage().GetEntries(), 
-			Parameter);
+		return Values.GetEvaluator().template EvaluateDerivative<Order>(
+			Values.GetStorage().GetEntries(), Parameter);
 	}
 
-	/** Evaluates second derivative at parameter */
-	virtual ValueType EvaluateSecondDerivative(float Parameter) const
+	/**
+	 * Evaluates an attribute derivative at the given parameter in parent space
+	 * @tparam AttrType Custom attribute type
+	 * @tparam Order Derivative order
+	 * @param Name Name of the channel to evaluate
+	 * @param Parameter Parameter value in parent spline space
+	 * @return Derivative value
+	 */
+	template<typename AttrType, int32 Order>
+	AttrType EvaluateAttributeDerivative(FName Name, float Parameter) const
 	{
-		return GetEvaluator().EvaluateSecondDerivative(
-			Values.GetStorage().GetEntries(), 
-			Parameter);
+		if (const TAttributeChannel<AttrType>* Channel = GetAttributeChannel<AttrType>(Name))
+		{
+			// Convert from parent space to channel's local space
+			const float ChannelParam = Channel->GetEvaluator().MapFromParentSpace(Parameter);
+			// Evaluate in channel's local space
+			return Channel->template EvaluateDerivative<Order>(ChannelParam);
+		}
+		return AttrType();
 	}
 	
 	/** 
 	 * Evaluates an attribute at the given parameter in parent space
+	 * @tparam AttrType - Custom attribute type
 	 * @param Name - Name of the attribute channel
 	 * @param Parameter - Parameter value in parent spline space
 	 */
@@ -245,6 +308,9 @@ public:
 		return AttrType();
 	}
 
+	/**
+	 * Clears all values from the spline including custom attribute channels
+	 */
 	virtual void Clear()
 	{
 		Values.Clear();
@@ -399,7 +465,7 @@ private:
 
 		// Create new instance of correct derived type
 		TSharedPtr<TSplineBase> NewSpline = CreateNewInstance();
-		NewSpline->Values = MakeUnique<TAttributeChannel<ValueType>>(*Values);
+        NewSpline->Values = MakeUnique<TAttributeChannel<ValueType>>(*Values);
 
 		// Copy parameter range settings
 		float Start, End;
@@ -435,6 +501,6 @@ using FSpline2f = TSplineBase<float, 2, TBaseEvaluator<FVector2f>>;
 using FSpline3f = TSplineBase<float, 3, TBaseEvaluator<FVector3f>>;
 using FSpline2d = TSplineBase<double, 2, TBaseEvaluator<FVector2D>>;
 using FSpline3d = TSplineBase<double, 3, TBaseEvaluator<FVector>>;
-
+} // end namespace UE::Geometry::Spline
 } // end namespace UE::Geometry
 } // end namespace UE
\ No newline at end of file
diff --git a/Engine/Source/Runtime/GeometryCore/Public/Splines/SplineMath.h b/Engine/Source/Runtime/GeometryCore/Public/Splines/SplineMath.h
new file mode 100644
--- /dev/null
+++ b/Engine/Source/Runtime/GeometryCore/Public/Splines/SplineMath.h
@@ -0,0 +1,430 @@
+// Copyright Epic Games, Inc. All Rights Reserved.
+
+#pragma once
+#include "CoreMinimal.h"
+namespace UE
+{
+namespace Geometry
+{
+namespace Spline
+{
+namespace Math
+{
+
+/** Helper for validating types and inputs */
+struct FSplineValidation
+{
+	/** Get the default/zero value for a type */
+	template<typename T>
+	static T GetDefaultValue()
+	{
+		if constexpr (TIsConstructible<T, EForceInit>::Value)
+		{
+			return T(ForceInit);
+		}
+		return T(); 
+	}
+
+	template<typename T>
+	static bool IsValidWindow(TArrayView<const T* const> Window, int32 RequiredSize)
+	{
+		if (Window.Num() < RequiredSize)
+		{
+			return false;
+		}
+    
+		// Check for null pointers
+		for (const T* Ptr : Window)
+		{
+			if (!Ptr)
+			{
+				return false;
+			}
+		}
+    
+		return true;
+	}
+
+	static bool IsValidKnotVector(const TArray<float>& Knots, int32 Degree)
+	{
+		// Need at least 2*(degree+1) knots
+		if (Knots.Num() < 2*(Degree + 1))
+		{
+			return false;
+		}
+
+		// Check monotonicity
+		for (int32 i = 1; i < Knots.Num(); ++i)
+		{
+			if (Knots[i] < Knots[i-1])
+			{
+				return false;
+			}
+		}
+
+		return true;
+	}
+
+	static bool IsValidParameter(float Parameter)
+	{
+		return !FPlatformMath::IsNaN(Parameter) && 
+			   FPlatformMath::IsFinite(Parameter);
+	}
+
+	template<typename T>
+	static bool IsValidGeometry(TArrayView<const T* const> Window, float Tolerance = KINDA_SMALL_NUMBER)
+	{
+		if (Window.Num() < 2) return false;
+
+		// Check for degenerate/repeated points
+		for (int32 i = 1; i < Window.Num(); ++i)
+		{
+			if ((*Window[i] - *Window[i-1]).SizeSquared() < Tolerance)
+			{
+				return false;
+			}
+		}
+
+		return true;
+	}
+};
+	
+// Cached factorial table access
+template <int32 N>
+const TArray<double>& GetFactorialTable()
+{
+	static TArray<double> Table;
+	if (Table.Num() == 0)
+	{
+		Table.SetNum(N + 1);
+		Table[0] = 1.0;
+		for (int32 i = 1; i <= N; ++i)
+		{
+			Table[i] = Table[i - 1] * i;
+		}
+	}
+	return Table;
+}
+
+// Cached binomial coefficient table
+template<int32 N>
+const TArray<TArray<double>>& GetBinomialTable()
+{
+	static TArray<TArray<double>> Table;
+	if (Table.Num() == 0)
+	{
+		Table.SetNum(N + 1);
+		for (int32 i = 0; i <= N; ++i)
+		{
+			Table[i].SetNum(i + 1);
+			Table[i][0] = Table[i][i] = 1.0;
+			for (int32 j = 1; j < i; ++j)
+			{
+				Table[i][j] = Table[i-1][j-1] + Table[i-1][j];
+			}
+		}
+	}
+	return Table;
+}
+
+/** Template-based derivative calculators for spline types */
+
+// Helper struct for shared derivative implementation
+template<typename T>
+struct TGenericDerivativeHelper
+{
+	template<int32 Order>
+	static T Compute(TArrayView<const T* const> Window, float Parameter)
+	{
+		// Basic validation
+		if (!FSplineValidation::IsValidWindow(Window, Order + 1))
+		{
+			return FSplineValidation::GetDefaultValue<T>();
+		}
+
+		if (!FSplineValidation::IsValidParameter(Parameter))
+		{
+#if ENABLE_NAN_DIAGNOSTIC
+			logOrEnsureNanError(TEXT("Invalid parameter in derivative calculation"));
+#endif
+			return FSplineValidation::GetDefaultValue<T>();
+		}
+		
+		// Get pre-computed binomial coefficients
+		static const auto& BinomialTable = GetBinomialTable<20>();
+		static const auto& FactorialTable = GetFactorialTable<20>();
+
+		const int32 WindowSize = Window.Num();
+		if (Order >= WindowSize || Order >= BinomialTable.Num())
+		{
+			return FSplineValidation::GetDefaultValue<T>();
+		}
+
+		// Optional geometry validation for higher derivatives
+		if (Order > 1 && !FSplineValidation::IsValidGeometry(Window))
+		{
+			return FSplineValidation::GetDefaultValue<T>();
+		}
+
+		T Result = FSplineValidation::GetDefaultValue<T>();
+		const float t = Parameter;
+		const float mt = 1.0f - t;
+
+		// Compute derivative using binomial expansion
+		for (int32 i = 0; i <= WindowSize - Order - 1; ++i)
+		{
+			double Coeff = BinomialTable[WindowSize - 1][i];
+
+			// Protect against overflow
+			if (Coeff > DBL_MAX / (WindowSize - 1))
+			{
+				return FSplineValidation::GetDefaultValue<T>();
+			}
+			
+			// Apply derivative scaling
+			Coeff *= FactorialTable[WindowSize - 1] / FactorialTable[WindowSize - Order - 1];
+        
+			// Apply parameter terms with underflow protection
+			if (mt > 0.0f) 
+			{
+				const float PowMT = FMath::Pow(mt, WindowSize - Order - 1 - i);
+				if (FPlatformMath::IsFinite(PowMT))
+				{
+					Coeff *= PowMT;
+				}
+			}
+			if (t > 0.0f)
+			{
+				const float PowT = FMath::Pow(t, i);
+				if (FPlatformMath::IsFinite(PowT))
+				{
+					Coeff *= PowT;
+				}
+			}
+
+			Result += *Window[i] * Coeff;
+		}
+		return Result;
+	}
+};
+
+// B-spline specific calculations
+template <typename T, int32 Order>
+struct TBSplineDerivativeCalculator
+{
+	static T Compute(
+		TArrayView<const T* const> Window,
+		const TArray<float>& Knots,
+		int32 Span,
+		float Parameter,
+		int32 Degree,
+		const T& DefaultValue = FSplineValidation::GetDefaultValue<T>())
+	{
+		// Extended validation
+		if (!FSplineValidation::IsValidWindow(Window, Degree + 1) ||
+			!FSplineValidation::IsValidKnotVector(Knots, Degree) ||
+			!FSplineValidation::IsValidParameter(Parameter))
+		{
+			return FSplineValidation::GetDefaultValue<T>();
+		}
+
+		// Validate span range
+		if (Span < 0 || Span >= Knots.Num() - Degree - 1)
+		{
+			return FSplineValidation::GetDefaultValue<T>();
+		}
+
+		// For higher derivatives, check geometry
+		if (Order > 1 && !FSplineValidation::IsValidGeometry(Window))
+		{
+			return FSplineValidation::GetDefaultValue<T>();
+		}
+		
+		// Compute basis functions and derivatives using Oslo algorithm
+		TArray<TArray<float>> Basis;
+		ComputeBSplineBasisFunctions<Order>(
+			Span, Parameter, Degree, Knots, Basis);
+
+		// Validate basis computation succeeded
+		if (Basis.Num() == 0 || Basis[Order].Num() == 0)
+		{
+			return FSplineValidation::GetDefaultValue<T>();
+		}
+		
+		// Apply basis to control points
+		T Result = FSplineValidation::GetDefaultValue<T>();
+		for (int32 i = 0; i <= Degree - Order; ++i)
+		{
+			Result += *Window[i] * Basis[Order][i];
+		}
+
+		// Scale by factorial ratio
+		static const auto& FactorialTable = GetFactorialTable<20>();
+		const double Scale = FactorialTable[Degree] / FactorialTable[Degree - Order];
+		if (Scale > 0.0 && Scale < DBL_MAX)
+		{
+			return Result * Scale;
+		}
+		
+		return FSplineValidation::GetDefaultValue<T>();
+	}
+};
+
+// Bezier curve specific calculations
+template <typename T, int32 Order>
+struct TBezierDerivativeCalculator
+{
+	static T Compute(TArrayView<const T* const> Window, float Parameter, float SegmentScale)
+	{
+		if constexpr (Order > 3) 
+		{
+			return FSplineValidation::GetDefaultValue<T>();
+		}
+
+		if (!FSplineValidation::IsValidWindow(Window, 4) ||  // Cubic Bezier needs 4 points
+			!FSplineValidation::IsValidParameter(Parameter))
+		{
+			return FSplineValidation::GetDefaultValue<T>();
+		}
+
+		// Use existing basis computation
+		TArray<float> Basis;
+		ComputeBezierBasis<Order>(Parameter, Basis);
+    
+		// Apply basis to control points
+		T Result = FSplineValidation::GetDefaultValue<T>();
+		for (int32 i = 0; i < Window.Num(); ++i)
+		{
+			Result += *Window[i] * Basis[i];
+		}
+
+		// Apply segment scaling based on derivative order
+		return Result * FMath::Pow(SegmentScale, Order);
+	}
+};
+
+// B-spline basis functions with derivatives
+template <int32 Order>
+void ComputeBSplineBasisFunctions(
+	int32 Span,
+	float Parameter,
+	int32 Degree,
+	const TArray<float>& Knots,
+	TArray<TArray<float>>& OutBasis)
+{
+	// Initialize workspace for Oslo algorithm
+	OutBasis.SetNum(Order + 1);
+	for (auto& Row : OutBasis)
+	{
+		Row.SetNum(Degree + 1);
+	}
+
+	// Allocate workspace arrays
+	TArray<float> Nders;
+	Nders.SetNum((Degree + 1) * 2);
+
+	// Initialize zeroth-degree functions
+	for (int32 j = 0; j <= Degree; ++j)
+	{
+		Nders[j] = 0.0f;
+	}
+
+	// Find index of knot interval containing Parameter
+	const float eps = 1e-6f;
+	int32 Index = Span - Degree;
+	for (int32 j = 0; j <= Degree; ++j, ++Index)
+	{
+		if (Parameter >= Knots[Index] - eps &&
+			Parameter < Knots[Index + 1] + eps)
+		{
+			Nders[j] = 1.0f;
+			break;
+		}
+	}
+
+	// Compute derivatives using optimized Oslo algorithm
+	for (int32 k = 1; k <= Order; ++k)
+	{
+		const int32 Stride = Degree + 1;
+		float* Previous = &Nders[0];
+		float* Current = &Nders[Stride];
+
+		// Clear current row
+		FMemory::Memzero(Current, Stride * sizeof(float));
+
+		for (int32 j = 0; j < Degree - k + 1; ++j)
+		{
+			const float OneOverDeltaKnot = 1.0f / (
+				Knots[Span + j + Degree - k + 1] - Knots[Span + j]);
+			Current[j] = Previous[j] * OneOverDeltaKnot;
+
+			if (j > 0)
+			{
+				Current[j] -= Previous[j - 1] * OneOverDeltaKnot;
+			}
+		}
+
+		// Scale by degree
+		const float Scale = static_cast<const float>(Degree - k + 1);
+		for (int32 j = 0; j < Degree - k + 1; ++j)
+		{
+			Current[j] *= Scale;
+		}
+
+		// Swap buffers for next iteration
+		Swap(Previous, Current);
+	}
+
+	// Copy results to output basis
+	const float* FinalDerivs = &Nders[((Order + 1) % 2) * (Degree + 1)];
+	for (int32 i = 0; i <= Degree - Order; ++i)
+	{
+		OutBasis[Order][i] = FinalDerivs[i];
+	}
+}
+
+// Bezier basis computation
+template <int32 Order>
+static void ComputeBezierBasis(float t, TArray<float>& OutBasis)
+{
+	OutBasis.SetNumZeroed(4); // Cubic Bezier
+
+	const float mt = 1.0f - t;
+
+	if constexpr (Order == 0)
+	{
+		// Position basis
+		OutBasis[0] = mt * mt * mt;
+		OutBasis[1] = 3.0f * mt * mt * t;
+		OutBasis[2] = 3.0f * mt * t * t;
+		OutBasis[3] = t * t * t;
+	}
+	else if constexpr (Order == 1)
+	{
+		// First derivative basis
+		OutBasis[0] = -3.0f * mt * mt;
+		OutBasis[1] = 3.0f * mt * (1.0f - 3.0f * t);
+		OutBasis[2] = 3.0f * t * (2.0f - 3.0f * t);
+		OutBasis[3] = 3.0f * t * t;
+	}
+	else if constexpr (Order == 2)
+	{
+		// Second derivative basis
+		OutBasis[0] = 6.0f * mt;
+		OutBasis[1] = -12.0f + 18.0f * t;
+		OutBasis[2] = 6.0f - 18.0f * t;
+		OutBasis[3] = 6.0f * t;
+	}
+	else if constexpr (Order == 3)
+	{
+		// Third derivative basis (constant)
+		OutBasis[0] = -6.0f;
+		OutBasis[1] = 18.0f;
+		OutBasis[2] = -18.0f;
+		OutBasis[3] = 6.0f;
+	}
+}
+} // end namespace UE::Geometry::Spline::Math
+} // end namespace UE::Geometry::Spline
+} // end namespace UE::Geometry
+} // end namespace UE
