template call fixes

diff --git a/Engine/Source/Runtime/GeometryCore/Public/Splines/InterpolationPolicies.h b/Engine/Source/Runtime/GeometryCore/Public/Splines/InterpolationPolicies.h
--- a/Engine/Source/Runtime/GeometryCore/Public/Splines/InterpolationPolicies.h
+++ b/Engine/Source/Runtime/GeometryCore/Public/Splines/InterpolationPolicies.h
@@ -48,7 +48,7 @@ public:
 		{
 			return Interpolate(Window, Parameter);
 		}
-		return Math::TGenericDerivativeHelper<T>::Compute<Order>(Window, Parameter);
+		return Math::TGenericDerivativeHelper<T, Order>::Compute(Window, Parameter);
 	}
 };
 
@@ -97,7 +97,7 @@ public:
 		{
 			return Interpolate(Window, Parameter);
 		}
-		return Math::TGenericDerivativeHelper<FLinearColor>::Compute<Order>(Window, Parameter);
+		return Math::TGenericDerivativeHelper<FLinearColor, Order>::Compute(Window, Parameter);
 	}
 };
 
diff --git a/Engine/Source/Runtime/GeometryCore/Public/Splines/SplineMath.h b/Engine/Source/Runtime/GeometryCore/Public/Splines/SplineMath.h
--- a/Engine/Source/Runtime/GeometryCore/Public/Splines/SplineMath.h
+++ b/Engine/Source/Runtime/GeometryCore/Public/Splines/SplineMath.h
@@ -130,10 +130,9 @@ const TArray<TArray<double>>& GetBinomialTable()
 /** Template-based derivative calculators for spline types */
 
 // Helper struct for shared derivative implementation
-template<typename T>
+template<typename T, int32 Order>
 struct TGenericDerivativeHelper
 {
-	template<int32 Order>
 	static T Compute(TArrayView<const T* const> Window, float Parameter)
 	{
 		// Basic validation
