Splines: closed loop

diff --git a/Engine/Source/Runtime/GeometryCore/Public/Splines/BSpline.h b/Engine/Source/Runtime/GeometryCore/Public/Splines/BSpline.h
--- a/Engine/Source/Runtime/GeometryCore/Public/Splines/BSpline.h
+++ b/Engine/Source/Runtime/GeometryCore/Public/Splines/BSpline.h
@@ -71,17 +71,33 @@ public:
 		}
 
 		const float MappedParam = Base::MapToParentSpace(Parameter);
-		const int32 Span = FindKnotSpan(MappedParam);
-
-		// Get control points based on knot vector
-		for (int32 i = 0; i <= Degree; ++i)
+		int32 Span = FindKnotSpan(MappedParam);
+		// For closed loop, we need to ensure the span wraps properly
+		if (this->IsClosedLoop())
 		{
-			const int32 Index = Span - Degree + i;
-			if (Index >= 0 && Index < Entries.Num())
+			// For closed loop:
+			// - Span needs to wrap around number of control points
+			// - Window indices wrap around array
+			const int32 NumPoints = Entries.Num();
+			for (int32 i = 0; i <= Degree; ++i)
 			{
+				int32 Index = (Span - Degree + i) % NumPoints;
+				if (Index < 0) Index += NumPoints;
 				Window.Add(&Entries[Index].Value);
 			}
 		}
+		else
+		{
+			// Get control points for the window
+			for (int32 i = 0; i <= Degree; ++i)
+			{
+				const int32 Index = Span - Degree + i;
+				if (Index >= 0 && Index < Entries.Num())
+				{
+					Window.Add(&Entries[Index].Value);
+				}
+			}
+		}
 
 		return Window;
 	}
@@ -99,38 +115,45 @@ public:
 		return Degree;
 	}
 
-	friend FArchive& operator<<(FArchive& Ar, TBSplineEvaluator& Evaluator)
+	/**
+	 * Serializes evaluator data to/from an archive
+	 * @param Ar - Archive to serialize to/from
+	 */
+	virtual void Serialize(FArchive& Ar) override
 	{
 		// Serialize base evaluator first
-		Ar << static_cast<Base&>(Evaluator);
+		Base::Serialize(Ar);
         
 		// Serialize BSpline-specific data
-		Ar << Evaluator.KnotVector;
-		Ar << Evaluator.NumValues;
+		Ar << KnotVector;
+		Ar << NumValues;
         
-		uint8 KnotModeInt = static_cast<uint8>(Evaluator.KnotMode);
+		uint8 KnotModeInt = static_cast<uint8>(KnotMode);
 		Ar << KnotModeInt;
         
 		if (Ar.IsLoading())
 		{
-			Evaluator.KnotMode = static_cast<EBSplineKnotMode>(KnotModeInt);
+			KnotMode = static_cast<EBSplineKnotMode>(KnotModeInt);
 		}
         
-		Ar << Evaluator.bClampEnds;
+		Ar << bClampEnds;
         
-		return Ar;
 	}
-
 	/**
 	 * Sets a custom knot vector and switches to custom knot mode
 	 * @param NewKnots - Array of knot values. Must be non-decreasing and have correct length (n + p + 1)
 	 * @param InNumValues - Number of values (control points) in the spline
 	 * @return True if the knot vector was valid and set successfully
 	 */
-	virtual bool SetCustomKnots(const TArray<float>& NewKnots, int32 InNumValues)
+	bool SetCustomKnots(const TArray<float>& NewKnots, int32 InNumValues)
 	{
 		NumValues = InNumValues;
-		const int32 RequiredKnots = NumValues + Degree + 1;
+		
+		// For closed loops, we repeat the first Degree points at the end
+		const int32 RequiredKnots = this->IsClosedLoop() ? 
+			(NumValues + 2 * Degree + 1) : // Extra knots for closure
+			(NumValues + Degree + 1);      // Standard case
+		
 		if (NewKnots.Num() < RequiredKnots)
 		{
 			UE_LOG(LogSpline, Warning, TEXT("Invalid knot vector size. Need %d knots, got %d."), 
@@ -510,31 +533,54 @@ private:
      */
 	void GenerateUniformKnots(int32 KnotCount)
 	{
-		for (int32 i = 0; i < KnotCount; ++i)
+		if (this->IsClosedLoop())
 		{
-			if (bClampEnds)
+			// For a closed B-spline of degree d with n points:
+			// - Need additional d points wrapping from start
+			// - Need n+2d+1 knots total
+			// - Knot spacing should be uniform to maintain periodicity
+			const int32 n = NumValues;
+			const int32 d = Degree; 
+			const int32 NumKnots = n + 2*d + 1;
+        
+			KnotVector.SetNum(NumKnots);
+
+			// For periodic B-splines, use uniform knot spacing
+			for (int32 i = 0; i < NumKnots; ++i)
 			{
-				if (i <= Degree)
-				{
-					// Multiple knots at start for endpoint interpolation
-					KnotVector[i] = 0.0f;
-				}
-				else if (i >= KnotCount - Degree - 1)
+				// Map to [0,n] range for n segments
+				KnotVector[i] = static_cast<float>(i);
+			}
+		}
+		else
+		{
+			KnotVector.SetNum(KnotCount);
+			for (int32 i = 0; i < KnotCount; ++i)
+			{
+				if (bClampEnds)
 				{
-					// Multiple knots at end for endpoint interpolation
-					KnotVector[i] = KnotCount - 2 * Degree;
+					if (i <= Degree)
+					{
+						// Multiple knots at start for endpoint interpolation
+						KnotVector[i] = 0.0f;
+					}
+					else if (i >= KnotCount - Degree - 1)
+					{
+						// Multiple knots at end for endpoint interpolation
+						KnotVector[i] = KnotCount - 2 * Degree;
+					}
+					else
+					{
+						// Uniform spacing for internal knots
+						KnotVector[i] = i - Degree;
+					}
 				}
 				else
 				{
-					// Uniform spacing for internal knots
-					KnotVector[i] = i - Degree;
+					// Simple uniform spacing when not clamped
+					KnotVector[i] = static_cast<float>(i);
 				}
 			}
-			else
-			{
-				// Simple uniform spacing when not clamped
-				KnotVector[i] = static_cast<float>(i);
-			}
 		}
 
 		UE_LOG(LogSpline, Verbose, TEXT("Updated knot vector - Points: %d, Degree: %d, Knots: %d"),
@@ -556,6 +602,11 @@ private:
             GenerateUniformKnots(KnotCount);
             return;
         }
+		
+		const int32 n = NumValues;
+		const int32 d = Degree;
+		const int32 NumKnots = this->IsClosedLoop() ? (n + 2*d + 1) : KnotCount;
+		KnotVector.SetNum(NumKnots);
 
         // Calculate chord lengths between consecutive points
         TArray<float> Chords;
@@ -569,39 +620,58 @@ private:
             TotalLength += Length;
         }
 
-        // Generate knot vector based on accumulated chord lengths
-        for (int32 i = 0; i < KnotCount; ++i)
-        {
-            if (bClampEnds)
-            {
-                if (i <= Degree)
-                {
-                    KnotVector[i] = 0.0f;
-                }
-                else if (i >= KnotCount - Degree - 1)
-                {
-                    KnotVector[i] = 1.0f;
-                }
-                else
-                {
-                    // Accumulate normalized chord lengths
-                    const int32 Segment = (i - Degree);
-                    float Param = 0.0f;
-                    for (int32 j = 0; j < Segment; ++j)
-                    {
-                        Param += (TotalLength > UE_SMALL_NUMBER) ? 
-                            (Chords[j] / TotalLength) : 
-                            (1.0f / (float)(Points.Num() - 1));
-                    }
-                    KnotVector[i] = Param;
-                }
-            }
-            else
-            {
-                const float Param = static_cast<float>(i) / static_cast<float>(KnotCount - 1);
-                KnotVector[i] = Param * TotalLength;
-            }
-        }
+		// For closed loop, add the closing segment length
+		if (this->IsClosedLoop() && Points.Num() > 0)
+		{
+			const float ClosureLength = static_cast<float>((Points[0].Value - Points.Last().Value).Size());
+			Chords.Add(ClosureLength);
+			TotalLength += ClosureLength;
+		}
+
+		if (this->IsClosedLoop())
+		{
+			// For closed loop, generate periodic knot vector based on chord lengths
+			for (int32 i = 0; i < NumKnots; ++i)
+			{
+				KnotVector[i] = static_cast<float>(i) * TotalLength / (n + d);
+			}
+		}
+		else
+		{
+			// Generate knot vector based on accumulated chord lengths
+			for (int32 i = 0; i < KnotCount; ++i)
+			{
+				if (bClampEnds)
+				{
+					if (i <= Degree)
+					{
+						KnotVector[i] = 0.0f;
+					}
+					else if (i >= KnotCount - Degree - 1)
+					{
+						KnotVector[i] = 1.0f;
+					}
+					else
+					{
+						// Accumulate normalized chord lengths
+						const int32 Segment = (i - Degree);
+						float Param = 0.0f;
+						for (int32 j = 0; j < Segment; ++j)
+						{
+							Param += (TotalLength > UE_SMALL_NUMBER) ? 
+								(Chords[j] / TotalLength) : 
+								(1.0f / (float)(Points.Num() - 1));
+						}
+						KnotVector[i] = Param;
+					}
+				}
+				else
+				{
+					const float Param = static_cast<float>(i) / static_cast<float>(KnotCount - 1);
+					KnotVector[i] = Param * TotalLength;
+				}
+			}
+		}
     }
 
     /**
@@ -617,6 +687,11 @@ private:
 			return;
 		}
 
+		const int32 n = NumValues;
+		const int32 d = Degree;
+		const int32 NumKnots = this->IsClosedLoop() ? (n + 2*d + 1) : KnotCount;
+		KnotVector.SetNum(NumKnots);
+		
 		TArray<float> Lengths;
 		Lengths.Reserve(Points.Num() - 1);
 		float TotalLength = 0.0f;
@@ -627,37 +702,55 @@ private:
 			Lengths.Add(Length);
 			TotalLength += Length;
 		}
+		// For closed loop, add the closing segment
+		if (this->IsClosedLoop() && Points.Num() > 0)
+		{
+			const float ClosureLength = FMath::Sqrt(static_cast<float>((Points[0].Value - Points.Last().Value).Size()));
+			Lengths.Add(ClosureLength);
+			TotalLength += ClosureLength;
+		}
 
-		for (int32 i = 0; i < KnotCount; ++i)
+		if (this->IsClosedLoop())
 		{
-			if (bClampEnds)
+			// Generate periodic knot vector using centripetal parameterization
+			for (int32 i = 0; i < NumKnots; ++i)
 			{
-				if (i <= Degree)
-				{
-					KnotVector[i] = 0.0f;
-				}
-				else if (i >= KnotCount - Degree - 1)
+				KnotVector[i] = static_cast<float>(i) * TotalLength / (n + d);
+			}
+		}
+		else
+		{
+			for (int32 i = 0; i < KnotCount; ++i)
+			{
+				if (bClampEnds)
 				{
-					KnotVector[i] = 1.0f;
+					if (i <= Degree)
+					{
+						KnotVector[i] = 0.0f;
+					}
+					else if (i >= KnotCount - Degree - 1)
+					{
+						KnotVector[i] = 1.0f;
+					}
+					else
+					{
+						const int32 Segment = (i - Degree);
+						float Param = 0.0f;
+						for (int32 j = 0; j < Segment; ++j)
+						{
+							Param += (TotalLength > UE_SMALL_NUMBER) ? 
+								(Lengths[j] / TotalLength) : 
+								(1.0f / (float)(Points.Num() - 1));
+						}
+						KnotVector[i] = Param;
+					}
 				}
 				else
 				{
-					const int32 Segment = (i - Degree);
-					float Param = 0.0f;
-					for (int32 j = 0; j < Segment; ++j)
-					{
-						Param += (TotalLength > UE_SMALL_NUMBER) ? 
-							(Lengths[j] / TotalLength) : 
-							(1.0f / (float)(Points.Num() - 1));
-					}
-					KnotVector[i] = Param;
+					const float Param = static_cast<float>(i) / static_cast<float>(KnotCount - 1);
+					KnotVector[i] = Param * TotalLength;
 				}
 			}
-			else
-			{
-				const float Param = static_cast<float>(i) / static_cast<float>(KnotCount - 1);
-				KnotVector[i] = Param * TotalLength;
-			}
 		}
 	}
 
@@ -761,6 +854,15 @@ public:
 		Base::GetEvaluator().SetClampedEnds(bInClampEnds);
 		Base::GetEvaluator().UpdateKnotVector(Base::Values.GetStorage().GetEntries());
 	}
+
+	virtual void SetClosedLoop(bool bInClosedLoop) override
+	{
+		if (this->IsClosedLoop() != bInClosedLoop)
+		{
+			this->GetEvaluator().SetClosedLoop(bInClosedLoop);
+			this->GetEvaluator().UpdateKnotVector(Base::Values.GetStorage().GetEntries());
+		}
+	}
 };
 
 // Common type definitions
diff --git a/Engine/Source/Runtime/GeometryCore/Public/Splines/BaseEvaluator.h b/Engine/Source/Runtime/GeometryCore/Public/Splines/BaseEvaluator.h
--- a/Engine/Source/Runtime/GeometryCore/Public/Splines/BaseEvaluator.h
+++ b/Engine/Source/Runtime/GeometryCore/Public/Splines/BaseEvaluator.h
@@ -31,6 +31,10 @@ template<typename ValueType>
 class TBaseEvaluator
 {
 protected:
+
+	/** whether the spline is closed */
+	bool bClosedLoop = false;
+	
 	/** Start point in parent space */
 	float ParentStart = 0.0f;
     
@@ -50,29 +54,31 @@ public:
 	/**
 	 * Serializes evaluator data to/from an archive
 	 * @param Ar - Archive to serialize to/from
-	 * @param Evaluator 
 	 */
-	friend void operator<<(FArchive& Ar, TBaseEvaluator& Evaluator)
+	virtual void Serialize(FArchive& Ar)
 	{
+		Ar << bClosedLoop;
+
 		// Serialize parameter range
-		Ar << Evaluator.ParentStart;
-		Ar << Evaluator.ParentEnd;
-        
+		Ar << ParentStart;
+		Ar << ParentEnd;
+
 		// Serialize infinity modes
-		uint8 PreMode = static_cast<uint8>(Evaluator.PreInfinityMode);
-		uint8 PostMode = static_cast<uint8>(Evaluator.PostInfinityMode);
+		uint8 PreMode = static_cast<uint8>(PreInfinityMode);
+		uint8 PostMode = static_cast<uint8>(PostInfinityMode);
 		Ar << PreMode;
 		Ar << PostMode;
         
 		if (Ar.IsLoading())
 		{
-			Evaluator.PreInfinityMode = static_cast<EOutOfBoundsHandlingMode>(PreMode);
-			Evaluator.PostInfinityMode = static_cast<EOutOfBoundsHandlingMode>(PostMode);
+			PreInfinityMode = static_cast<EOutOfBoundsHandlingMode>(PreMode);
+			PostInfinityMode = static_cast<EOutOfBoundsHandlingMode>(PostMode);
 		}
 	}
-	void Serialize(FArchive& Ar)
+	friend FArchive& operator<<(FArchive& Ar, TBaseEvaluator& Evaluator)
 	{
-		Ar << *this;
+		Evaluator.Serialize(Ar);
+		return Ar;	
 	}
 	
 	/** Get the polynomial degree of the evaluator */
@@ -185,6 +191,9 @@ public:
 		return Result;
 	}
 
+	virtual void SetClosedLoop(bool bInClosedLoop){ bClosedLoop = bInClosedLoop; };
+	bool IsClosedLoop() const { return bClosedLoop; }
+
 	/** Set the mapping range in parent space */
 	void SetParentSpaceRange(float InStart, float InEnd)
 	{
@@ -241,6 +250,7 @@ public:
         
 		Ar << PreMode;
 		Ar << PostMode;
+		Ar << bClosedLoop;
         
 		if (Ar.IsLoading())
 		{
diff --git a/Engine/Source/Runtime/GeometryCore/Public/Splines/PolyBezierSpline.h b/Engine/Source/Runtime/GeometryCore/Public/Splines/PolyBezierSpline.h
--- a/Engine/Source/Runtime/GeometryCore/Public/Splines/PolyBezierSpline.h
+++ b/Engine/Source/Runtime/GeometryCore/Public/Splines/PolyBezierSpline.h
@@ -20,7 +20,7 @@ class TPolyBezierEvaluator : public TBSplineEvaluator<T, D, 3>
 	using Base = TBSplineEvaluator<T, D, 3>;
 
 private:
-    float NumSegments = 0.0f;
+    int32 NumSegments = 0;
     
 public:
     TPolyBezierEvaluator():
@@ -30,15 +30,65 @@ public:
         this->SetClampedEnds(true);
     }
 
-    friend FArchive& operator<<(FArchive& Ar, TPolyBezierEvaluator& Evaluator)
+    bool UpdateSegmentStructure(int32 InNumSegments, int32 InNumValues)
+    {
+        NumSegments = InNumSegments;
+
+        // For Bezier curves:
+        // - Each segment has degree 3
+        // - For closed loop, we need n + 2*p + 1 knots where:
+        //   n = number of control points
+        //   p = degree (3 for cubic Bezier)
+        const int32 NumKnots = this->IsClosedLoop() ? 
+            (InNumValues + (2 * 3) + 1) :  // closed loop: n + 2p + 1
+            (InNumValues + 4);         // normal case: n + p + 1 where p = 3
+        
+        // Generate appropriate knot vector for PolyBezier
+        TArray<float> BezierKnots;
+        BezierKnots.SetNum(NumKnots);
+
+        // Set up knot vector with multiple knots at segment boundaries
+        int32 CurrentKnot = 0;
+
+        // Start knots with multiplicity 4
+        for (int32 i = 0; i < 4; ++i)
+        {
+            BezierKnots[CurrentKnot++] = 0.0f;
+        }
+
+        // For closed loop, add an extra segment to connect back to start
+        const int32 NumSegmentsToProcess = this->IsClosedLoop() ? (NumSegments + 1) : NumSegments;
+
+        // Interior knots at segment boundaries with multiplicity 3
+        for (int32 i = 1; i < NumSegmentsToProcess; ++i)
+        {
+            const float t = static_cast<float>(i);
+            for (int32 j = 0; j < 3; ++j)
+            {
+                BezierKnots[CurrentKnot++] = t;
+            }
+        }
+
+        // End knots with multiplicity 4
+        const float EndKnot = static_cast<float>(NumSegmentsToProcess);
+        while (CurrentKnot < NumKnots)
+        {
+            BezierKnots[CurrentKnot++] = EndKnot;
+        }
+
+        return Base::SetCustomKnots(BezierKnots, InNumValues);
+    }
+    /**
+     * Serializes evaluator data to/from an archive
+     * @param Ar - Archive to serialize to/from
+     */
+    virtual void Serialize(FArchive& Ar) override
     {
         // Serialize base evaluator first
-        Ar << static_cast<Base&>(Evaluator);
+        Base::Serialize(Ar);
         
         // Serialize PolyBezier-specific data
-        Ar << Evaluator.NumSegments;
-        
-        return Ar;
+        Ar << NumSegments;
     }
     
     virtual TUniquePtr<TBaseEvaluator<ValueType>> Clone() const override
@@ -61,20 +111,45 @@ public:
         
         // Find which segment we're in
         const float ClampedParam = FMath::Clamp(ParentParam, 0.0f, NumSegments);
-        const int32 SegmentIndex = FMath::Min(FMath::FloorToInt(ClampedParam), static_cast<int32>(NumSegments) - 1);
+        int32 SegmentIndex = FMath::Min(FMath::FloorToInt(ClampedParam), static_cast<int32>(NumSegments) - 1);
+
+        if (this->IsClosedLoop())
+        {
+            // For closed loop, wrap the segment index
+            SegmentIndex = SegmentIndex % NumSegments;
+            if (SegmentIndex < 0) SegmentIndex += NumSegments;
+        }
         
-        // Get the four control points for this segment
+        // Each Bezier segment has exactly 4 control points
         const int32 BaseIndex = SegmentIndex * 4;
-        if (BaseIndex + 3 >= Entries.Num())
+        Window.Reserve(4);
+
+        if (this->IsClosedLoop())
         {
-            return Window;
-        }
+            // Handle wrapping for closed loop
+            const int32 NumPoints = Entries.Num();
+            auto WrapIndex = [NumPoints](int32 Index) -> int32
+            {
+                return ((Index % NumPoints) + NumPoints) % NumPoints;
+            };
 
-        // Add the four Bezier control points to the window
-        Window.Add(&Entries[BaseIndex].Value);     // P0 need 
-        Window.Add(&Entries[BaseIndex + 1].Value); // P1
-        Window.Add(&Entries[BaseIndex + 2].Value); // P2
-        Window.Add(&Entries[BaseIndex + 3].Value); // P3
+            Window.Add(&Entries[WrapIndex(BaseIndex)].Value);     // P0 
+            Window.Add(&Entries[WrapIndex(BaseIndex + 1)].Value); // P1
+            Window.Add(&Entries[WrapIndex(BaseIndex + 2)].Value); // P2
+            Window.Add(&Entries[WrapIndex(BaseIndex + 3)].Value); // P3
+        }
+        else
+        {
+            // Standard case - no wrapping
+            if (BaseIndex + 3 >= Entries.Num())
+            {
+                return Window;
+            }
+            Window.Add(&Entries[BaseIndex].Value);     // P0 
+            Window.Add(&Entries[BaseIndex + 1].Value); // P1
+            Window.Add(&Entries[BaseIndex + 2].Value); // P2
+            Window.Add(&Entries[BaseIndex + 3].Value); // P3
+        }
 
         return Window;
     }
@@ -104,56 +179,18 @@ public:
         const float t2 = t * t;
         const float OneMinusT2 = OneMinusT * OneMinusT;
 
-        TArray<float> Basis;
-        Basis.SetNum(4);
-        
         // Cubic Bernstein polynomials
-        Basis[0] = OneMinusT * OneMinusT2;     // (1-t)³
-        Basis[1] = 3.0f * t * OneMinusT2;      // 3t(1-t)²
-        Basis[2] = 3.0f * t2 * OneMinusT;      // 3t²(1-t)
-        Basis[3] = t * t2;                     // t³
+        const float Basis[4] =
+        {
+            OneMinusT * OneMinusT2,     // (1-t)³
+            3.0f * t * OneMinusT2,      // 3t(1-t)²
+            3.0f * t2 * OneMinusT,      // 3t²(1-t)
+            t * t2                       // t³
+        };
 
         return TSplineInterpolationPolicy<ValueType>::InterpolateWithBasis(Window, Basis);
     }
 
-    virtual bool SetCustomKnots(const TArray<float>& NewKnots, int32 InNumValues) override
-    {
-        // Update number of segments based on control points (4 points per segment)
-        NumSegments = InNumValues / 4;
-        
-        // Generate appropriate knot vector for PolyBezier
-        TArray<float> BezierKnots;
-        const int32 RequiredKnots = InNumValues + 4;  // n + p + 1 where p = 3
-        BezierKnots.SetNum(RequiredKnots);
-
-        // Set up knot vector with multiple knots at segment boundaries
-        int32 CurrentKnot = 0;
-
-        // Start knots with multiplicity 4
-        for (int32 i = 0; i < 4; ++i)
-        {
-            BezierKnots[CurrentKnot++] = 0.0f;
-        }
-
-        // Interior knots at segment boundaries with multiplicity 3
-        for (int32 i = 1; i < NumSegments; ++i)
-        {
-            const float t = static_cast<float>(i);
-            for (int32 j = 0; j < 3; ++j)
-            {
-                BezierKnots[CurrentKnot++] = t;
-            }
-        }
-
-        // End knots with multiplicity 4
-        const float EndKnot = static_cast<float>(NumSegments);
-        while (CurrentKnot < RequiredKnots)
-        {
-            BezierKnots[CurrentKnot++] = EndKnot;
-        }
-
-        return Base::SetCustomKnots(BezierKnots, InNumValues);
-    }
     /** 
      * Evaluate nth derivative at parameter
      * @param Order - The derivative order (0 = position, 1 = first derivative, etc.)
@@ -179,12 +216,12 @@ public:
    
     virtual float GetParentSpaceMapping(float NormalizedParam) const override
     {
-        return NormalizedParam * NumSegments;
+        return NormalizedParam * static_cast<float>(NumSegments);
     }
 
     virtual float GetNormalizedSpaceMapping(float ParentParam) const override
     {
-        return NumSegments > 0.0f ? ParentParam / NumSegments : 0.0f;
+        return NumSegments > 0 ? ParentParam / static_cast<float>(NumSegments) : 0.0f;
     }
 };
 
@@ -208,10 +245,15 @@ public:
     // Virtual destructor
     virtual ~TPolyBezierSpline() = default;
 
-    virtual TSharedPtr<Base> CreateNewInstance() const override
+    virtual void SetClosedLoop(bool bInClosedLoop) override
     {
-        return MakeShared<TPolyBezierSpline>();
+        if (this->IsClosedLoop() != bInClosedLoop)
+        {
+            this->GetEvaluator().SetClosedLoop(bInClosedLoop);
+            UpdateSegmentStructure();
+        }
     }
+    
     /**
      * Adds a new Bezier segment to the spline
      * @param P0 Start point
@@ -242,8 +284,8 @@ public:
         Base::AddValue(P2, SegmentT);
         Base::AddValue(P3, SegmentT);
 
-        // Let evaluator handle knot vector update
-        Base::GetEvaluator().SetCustomKnots(TArray<float>(), Base::GetNumKeys());
+        // Update segment structure
+        UpdateSegmentStructure();
         
         return NumPoints / 4;
     }
@@ -301,7 +343,8 @@ public:
     /** Returns the number of Bezier segments in the spline */
     int32 GetNumSegments() const
     {
-        return Base::GetNumKeys() / 4;
+        const int32 NumPoints = this->GetNumKeys();
+        return NumPoints/4;//this->IsClosedLoop() ? NumPoints/4 : (NumPoints/4 - 1);
     }
 
     /** Gets the endpoint of a segment */
@@ -317,6 +360,19 @@ public:
     {
         return Index >= 0 && Index < GetNumSegments();
     }
+    
+protected:
+    virtual TSharedPtr<Base> CreateNewInstance() const override
+    {
+        return MakeShared<TPolyBezierSpline>();
+    }
+private:
+    void UpdateSegmentStructure()
+    {
+        const int32 NumPoints = this->GetNumKeys();
+        
+        this->GetEvaluator().UpdateSegmentStructure(GetNumSegments(), NumPoints);
+    }
 };
 
 // Common type definitions
diff --git a/Engine/Source/Runtime/GeometryCore/Public/Splines/Spline.h b/Engine/Source/Runtime/GeometryCore/Public/Splines/Spline.h
--- a/Engine/Source/Runtime/GeometryCore/Public/Splines/Spline.h
+++ b/Engine/Source/Runtime/GeometryCore/Public/Splines/Spline.h
@@ -151,6 +151,7 @@ public:
         
 		return Ar;
 	}
+
 	/**
 	  * Creates an attribute channel with a custom evaluator
 	  * @tparam AttrType - Custom attribute type
@@ -395,6 +396,16 @@ public:
 		return GetEvaluator().GetPostInfinityMode();
 	}
 
+	virtual void SetClosedLoop(bool bInClosedLoop)
+	{
+		GetEvaluator().SetClosedLoop(bInClosedLoop);
+	}
+	
+	bool IsClosedLoop() const
+	{
+		return GetEvaluator().IsClosedLoop();
+	}
+
 	/** Retrieves the main evaluator from the Values channel */
 	EvaluatorType& GetEvaluator()
 	{
