GeometryCore: Mark box clamp method as const.

diff --git a/Engine/Source/Runtime/GeometryCore/Public/BoxTypes.h b/Engine/Source/Runtime/GeometryCore/Public/BoxTypes.h
--- a/Engine/Source/Runtime/GeometryCore/Public/BoxTypes.h
+++ b/Engine/Source/Runtime/GeometryCore/Public/BoxTypes.h
@@ -488,7 +488,7 @@ struct TAxisAlignedBox3
 		return Contains(Box.Min) && Contains(Box.Max);
 	}
 
-	TVector<RealType> Clamp(const TVector<RealType>& V)
+	TVector<RealType> Clamp(const TVector<RealType>& V) const
 	{
 		return TVector<RealType>(
 			FMath::Clamp(V.X, Min.X, Max.X),
