Make dynamic mesh's mesh transform methods also transform sculpt layers if present

diff --git a/Engine/Source/Runtime/GeometryCore/Private/DynamicMesh/MeshTransforms.cpp b/Engine/Source/Runtime/GeometryCore/Private/DynamicMesh/MeshTransforms.cpp
--- a/Engine/Source/Runtime/GeometryCore/Private/DynamicMesh/MeshTransforms.cpp
+++ b/Engine/Source/Runtime/GeometryCore/Private/DynamicMesh/MeshTransforms.cpp
@@ -8,6 +8,44 @@
 
 using namespace UE::Geometry;
 
+namespace MeshTransformsLocal
+{
+	enum class ESkipOffset
+	{
+		SkipOffset
+	};
+
+	template<typename TransformPosF, typename TransformVecF>
+	void TransformSculptLayers(FDynamicMesh3& Mesh, TransformPosF TransformPosition, TransformVecF TransformOffset, EParallelForFlags ParallelForFlags = EParallelForFlags::None)
+	{
+		if (Mesh.HasAttributes() && Mesh.Attributes()->NumSculptLayers() > 0)
+		{
+			FDynamicMeshSculptLayers* SculptLayers = Mesh.Attributes()->GetSculptLayers();
+			ParallelFor(Mesh.MaxVertexID(), [&Mesh, SculptLayers, &TransformPosition, &TransformOffset](int VID)
+				{
+					if (Mesh.IsVertex(VID))
+					{
+						FVector3d BaseLayerPos;
+						SculptLayers->GetLayer(0)->GetValue(VID, BaseLayerPos);
+						SculptLayers->GetLayer(0)->SetValue(VID, TransformPosition(BaseLayerPos));
+						if constexpr (!std::is_same<TransformVecF, MeshTransformsLocal::ESkipOffset>())
+						{
+							for (int32 LayerIdx = 1, NumLayers = Mesh.Attributes()->NumSculptLayers(); LayerIdx < NumLayers; ++LayerIdx)
+							{
+								FVector3d Offset;
+								SculptLayers->GetLayer(LayerIdx)->GetValue(VID, Offset);
+								SculptLayers->GetLayer(LayerIdx)->SetValue(VID, TransformOffset(Offset));
+							}
+						}
+					}
+				},
+				ParallelForFlags
+			);
+		}
+	}
+}
+
+
 void MeshTransforms::Translate(FDynamicMesh3& Mesh, const FVector3d& Translation)
 {
 	int NumVertices = Mesh.MaxVertexID();
@@ -18,6 +56,9 @@ void MeshTransforms::Translate(FDynamicMesh3& Mesh, const FVector3d& Translation
 			Mesh.SetVertex(vid, Mesh.GetVertex(vid) + Translation);
 		}
 	});
+	MeshTransformsLocal::TransformSculptLayers(Mesh, 
+		[&Translation](const FVector3d& Pos) {return Pos + Translation;}, 
+		MeshTransformsLocal::ESkipOffset::SkipOffset);
 }
 
 
@@ -75,6 +116,10 @@ void MeshTransforms::Scale(FDynamicMesh3& Mesh, const FVector3d& Scale, const FV
 				});
 		}
 	}
+	
+	MeshTransformsLocal::TransformSculptLayers(Mesh, 
+		[&Origin, &Scale](const FVector3d& Pos) {return (Pos - Origin) * Scale + Origin;}, 
+		[&Scale](const FVector3d& Offset) { return Offset * Scale;});
 
 	if (bReverseOrientationIfNeeded && Scale.X * Scale.Y * Scale.Z < 0)
 	{
@@ -120,6 +165,10 @@ void MeshTransforms::WorldToFrameCoords(FDynamicMesh3& Mesh, const FFrame3d& Fra
 				}
 			});
 		}
+
+		MeshTransformsLocal::TransformSculptLayers(Mesh, 
+			[&Frame](const FVector3d& Pos) {return Frame.ToFramePoint(Pos); }, 
+			[&Frame](const FVector3d& Offset) { return Frame.ToFrameVector(Offset); });
 	}
 }
 
@@ -163,6 +212,10 @@ void MeshTransforms::FrameCoordsToWorld(FDynamicMesh3& Mesh, const FFrame3d& Fra
 				}
 			});
 		}
+
+		MeshTransformsLocal::TransformSculptLayers(Mesh, 
+			[&Frame](const FVector3d& Pos) {return Frame.FromFramePoint(Pos); }, 
+			[&Frame](const FVector3d& Offset) { return Frame.FromFrameVector(Offset); });
 	}
 }
 
@@ -204,6 +257,10 @@ void MeshTransforms::Rotate(FDynamicMesh3& Mesh, const FRotator& Rotation, const
 				}
 			});
 		}
+		
+		MeshTransformsLocal::TransformSculptLayers(Mesh, 
+			[&Rotation, &RotationOrigin](const FVector3d& Pos) {return Rotation.RotateVector((Pos - RotationOrigin)) + RotationOrigin; }, 
+			[&Rotation](const FVector3d& Offset) { return Rotation.RotateVector(Offset); });
 	}
 }
 
@@ -262,6 +319,10 @@ void MeshTransforms::ApplyTransform(FDynamicMesh3& Mesh, const FTransformSRT3d&
 				});
 			}
 		}
+
+		MeshTransformsLocal::TransformSculptLayers(Mesh,
+			[&Transform](const FVector3d& Pos) {return Transform.TransformPosition(Pos); },
+			[&Transform](const FVector3d& Offset) { return Transform.TransformVector(Offset); });
 	}
 
 	if (bReverseOrientationIfNeeded && Transform.GetDeterminant() < 0)
@@ -327,6 +388,10 @@ void MeshTransforms::ApplyTransformInverse(FDynamicMesh3& Mesh, const FTransform
 				});
 			}
 		}
+		
+		MeshTransformsLocal::TransformSculptLayers(Mesh,
+			[&Transform](const FVector3d& Pos) {return Transform.InverseTransformPosition(Pos); },
+			[&Transform](const FVector3d& Offset) { return Transform.InverseTransformVector(Offset); });
 	}
 
 	if (bReverseOrientationIfNeeded && Transform.GetDeterminant() < 0)
@@ -386,6 +451,11 @@ void MeshTransforms::ApplyTransform(FDynamicMesh3& Mesh,
 				}
 			});
 		}
+
+		FVector3d ZeroTransformed = PositionTransform(FVector::ZeroVector);
+		MeshTransformsLocal::TransformSculptLayers(Mesh,
+			[&PositionTransform](const FVector3d& Pos) {return PositionTransform(Pos); },
+			[&PositionTransform, ZeroTransformed](const FVector3d& Offset) { return PositionTransform(Offset) - ZeroTransformed; });
 	}
 }
 
@@ -447,6 +517,11 @@ void MeshTransforms::ApplyTransform(FDynamicMesh3& Mesh,
 				});
 			}
 		}
+		
+		FVector3d ZeroTransformed = PositionTransform(FVector::ZeroVector);
+		MeshTransformsLocal::TransformSculptLayers(Mesh,
+			[&PositionTransform](const FVector3d& Pos) {return PositionTransform(Pos); },
+			[&PositionTransform, ZeroTransformed](const FVector3d& Offset) { return PositionTransform(Offset) - ZeroTransformed; });
 	}
 }
 
