make point hash grid / scale indexing clamp the cell size above zero, since a zero cell size will not work
+ move ensures for too-small cellsize to the places that actually specify cellsize
+ change the radius search to a <= to still find exactly coincident points w/ a zero-radius search

diff --git a/Engine/Source/Runtime/GeometryCore/Public/Spatial/PointHashGrid3.h b/Engine/Source/Runtime/GeometryCore/Public/Spatial/PointHashGrid3.h
--- a/Engine/Source/Runtime/GeometryCore/Public/Spatial/PointHashGrid3.h
+++ b/Engine/Source/Runtime/GeometryCore/Public/Spatial/PointHashGrid3.h
@@ -40,10 +40,10 @@ public:
 
 	/**
 	 * Construct 3D hash grid
-	 * @param cellSize size of grid cells
+	 * @param CellSize size of grid cells. Clamped to at least a small tolerance size.
 	 * @param InvalidValue this value will be returned by queries if no valid result is found (e.g. bounded-distance query)
 	 */
-	TPointHashGrid3(RealType cellSize, PointDataType InvalidValue) : Indexer(cellSize), InvalidValue(InvalidValue)
+	TPointHashGrid3(RealType CellSize, PointDataType InvalidValue) : Indexer(FMath::Max(TMathUtil<RealType>::ZeroTolerance, CellSize)), InvalidValue(InvalidValue)
 	{
 	}
 
@@ -58,12 +58,12 @@ public:
 
 	/**
 	 * Remove all entries in the underlying hash map, without resizing or releasing any allocations, and reset the cell size.
-	 * @param NewCellSize The size of grid cells
+	 * @param NewCellSize The size of grid cells. Clamped to at least a small tolerance size.
 	 */
 	void Reset(RealType NewCellSize)
 	{
 		Hash.Reset();
-		Indexer.CellSize = NewCellSize;
+		Indexer.CellSize = FMath::Max(TMathUtil<RealType>::ZeroTolerance, NewCellSize);
 	}
 
 	/** Invalid grid value */
@@ -346,7 +346,7 @@ public:
 							continue;
 						}
 						RealType DistSq = DistanceSqFunc(Value);
-						if (DistSq < RadiusSquared)
+						if (DistSq <= RadiusSquared)
 						{
 							if (!ProcessPointFunc(Value, DistSq))
 							{
diff --git a/Engine/Source/Runtime/GeometryCore/Public/Util/GridIndexing3.h b/Engine/Source/Runtime/GeometryCore/Public/Util/GridIndexing3.h
--- a/Engine/Source/Runtime/GeometryCore/Public/Util/GridIndexing3.h
+++ b/Engine/Source/Runtime/GeometryCore/Public/Util/GridIndexing3.h
@@ -26,11 +26,12 @@ struct TScaleGridIndexer3
 
 	TScaleGridIndexer3() : CellSize((RealType)1) 
 	{
-		ensure(CellSize >= TMathUtil<RealType>::ZeroTolerance);
 	}
 	
-	TScaleGridIndexer3(RealType CellSize) : CellSize(CellSize) 
+	TScaleGridIndexer3(RealType CellSize) : CellSize(CellSize)
 	{
+		// Note a very small cell size is likely to overflow the integer grid coordinates
+		ensure(CellSize >= TMathUtil<RealType>::ZeroTolerance);
 	}
 
 	/** Convert real-valued point to integer grid coordinates */
@@ -67,7 +68,6 @@ struct TShiftGridIndexer3
 	TShiftGridIndexer3() 
 		: CellSize((RealType)1), Origin(TVector<RealType>::Zero())
 	{
-		ensure(CellSize >= TMathUtil<RealType>::ZeroTolerance);
 	}
 
 	TShiftGridIndexer3(const TVector<RealType>& origin, RealType cellSize)
