Splines: Added NearestPoint, nth order derivative calculation improvements. Removed duplicate data from TangentBezierSpline, it now acts as mostly an interface on top of PolyBezierSpline.
Renamed TangentControlledPolyBezier to TangentBezierSpline and moved it in to GeometryCore/Splines
Added unit tests for PolyBezier spline evaluation, derivatives and nearest point.

diff --git a/Engine/Source/Runtime/GeometryCore/Public/Splines/AttributeChannel.h b/Engine/Source/Runtime/GeometryCore/Public/Splines/AttributeChannel.h
--- a/Engine/Source/Runtime/GeometryCore/Public/Splines/AttributeChannel.h
+++ b/Engine/Source/Runtime/GeometryCore/Public/Splines/AttributeChannel.h
@@ -146,6 +146,16 @@ public:
         return Storage.AddValue(Value, Parameter);
     }
 
+	/**
+	 * Removes a value from the channel at the given index
+	 * @param Index - Index of the value to remove
+	 * @return True if the value was removed
+	 */
+	bool RemoveValue(int32 Index)
+	{
+		return Storage.RemoveValue(Index);
+	}
+	
     /**
      * Evaluates the channel at the given parameter in channel space
      * @param Parameter - Parameter value in channel's local space
diff --git a/Engine/Source/Runtime/GeometryCore/Public/Splines/BSpline.h b/Engine/Source/Runtime/GeometryCore/Public/Splines/BSpline.h
--- a/Engine/Source/Runtime/GeometryCore/Public/Splines/BSpline.h
+++ b/Engine/Source/Runtime/GeometryCore/Public/Splines/BSpline.h
@@ -200,7 +200,7 @@ public:
 		NumValues = Points.Num();
 		if (NumValues < Degree + 1)
 		{
-			UE_LOG(LogSpline, Error, TEXT("Not enough points for B-spline: %d (need >= %d)"),
+			UE_LOG(LogSpline, Warning, TEXT("Not enough points for B-spline: %d (need >= %d)"),
 				   NumValues, Degree + 1);
 			return;
 		}
@@ -302,7 +302,10 @@ public:
 
 		TArray<const ValueType*> Window = FindWindow(Entries, Parameter);
 		const float MappedParam = Base::MapToParentSpace(Parameter);
-		return Math::TBSplineDerivativeCalculator<ValueType, Order>::Compute(Window, MappedParam);
+		const int32 Span = FindKnotSpan(MappedParam);
+		
+		return Math::TBSplineDerivativeCalculator<ValueType, Order>::Compute(
+			Window, KnotVector, Span, MappedParam, Degree);
 	}
 	
 protected:
@@ -386,60 +389,7 @@ private:
 	/** Compute B-spline basis functions and their derivatives */
 	void ComputeBasisFunctions(int32 Span, float Parameter, TArray<float>& OutBasis) const
 	{
-		const int32 MinKnots = NumValues + Degree + 1;
-		if (KnotVector.Num() < MinKnots)
-		{
-			UE_LOG(LogSpline, Warning, TEXT("Not enough knots in ComputeBasisFunctions. Have %d, need %d"),
-			       KnotVector.Num(), MinKnots);
-			return;
-		}
-
-		OutBasis.SetNum(Degree + 1);
-		TArray<float> Left;
-		TArray<float> Right;
-		Left.SetNum(Degree + 1);
-		Right.SetNum(Degree + 1);
-
-		OutBasis[0] = 1.0f;
-
-		for (int32 j = 1; j <= Degree; ++j)
-		{
-			// Validate array indices before access
-			const int32 LeftIndex = Span + 1 - j;
-			const int32 RightIndex = Span + j;
-
-			if (LeftIndex < 0 || LeftIndex >= KnotVector.Num() ||
-				RightIndex < 0 || RightIndex >= KnotVector.Num())
-			{
-				UE_LOG(LogSpline, Warning, TEXT("Invalid knot indices in ComputeBasisFunctions: LeftIndex=%d, RightIndex=%d"),
-				       LeftIndex, RightIndex);
-				continue;
-			}
-
-			Left[j] = Parameter - KnotVector[LeftIndex];
-			Right[j] = KnotVector[RightIndex] - Parameter;
-
-			float Saved = 0.0f;
-
-			for (int32 r = 0; r < j; ++r)
-			{
-				float Temp;
-				const float Denominator = Right[r + 1] + Left[j - r];
-				if (FMath::Abs(Denominator) < UE_SMALL_NUMBER)
-				{
-					Temp = 0.0f;
-				}
-				else
-				{
-					Temp = OutBasis[r] / Denominator;
-				}
-
-				OutBasis[r] = Saved + Right[r + 1] * Temp;
-				Saved = Left[j - r] * Temp;
-			}
-
-			OutBasis[j] = Saved;
-		}
+		Math::ComputeBSplineBasisFunctions(KnotVector, Span, Parameter, Degree, OutBasis);
 
 		UE_LOG(LogSpline, VeryVerbose, TEXT("BSpline basis at t=%f, span=%d:"), Parameter, Span);
 		for (int32 i = 0; i <= Degree; ++i)
@@ -567,12 +517,12 @@ private:
 					else if (i >= KnotCount - Degree - 1)
 					{
 						// Multiple knots at end for endpoint interpolation
-						KnotVector[i] = KnotCount - 2 * Degree;
+						KnotVector[i] = KnotCount - 2 * static_cast<float>(Degree);
 					}
 					else
 					{
 						// Uniform spacing for internal knots
-						KnotVector[i] = i - Degree;
+						KnotVector[i] = i - static_cast<float>(Degree);
 					}
 				}
 				else
@@ -834,7 +784,55 @@ public:
 		}
 		return true;
 	}
-	
+
+	virtual float FindNearest(const ValueType& Point, float& OutSquaredDistance) const override
+	{
+		float BestDistSq = TNumericLimits<float>::Max();
+		float BestParam = 0.0f;
+
+		// Test each segment
+		const int32 NumSegments = GetNumSegments();
+		for (int32 i = 0; i < NumSegments; ++i)
+		{
+			float SegmentDistSq;
+			const float SegmentParam = FindNearestOnSegment(Point, i, SegmentDistSq);
+			if (SegmentDistSq < BestDistSq)
+			{
+				BestDistSq = SegmentDistSq;
+				BestParam = SegmentParam;
+			}
+		}
+
+		OutSquaredDistance = BestDistSq;
+		return BestParam;
+	}
+
+	virtual float FindNearestOnSegment(const ValueType& Point, int32 SegmentIndex, float& OutSquaredDistance) const override
+	{
+		//TODO: Implement this
+		return 0.0f;
+	}
+
+	virtual int32 GetNumSegments() const override
+	{
+		const auto& KnotVector = this->GetEvaluator().GetKnotVector();
+		const int32 NumPoints = this->GetNumKeys();
+
+		// For B-splines, number of segments depends on:
+		// - Whether it's closed loop
+		// - Number of distinct internal knots
+
+		if (this->IsClosedLoop())
+		{
+			return NumPoints;  // For closed loop, number of segments equals number of control points
+		}
+		else
+		{
+			// Number of segments = Number of distinct internal knot spans
+			// For standard clamped B-spline: NumPoints - Degree
+			return NumPoints - Degree;
+		}
+	}
 	/**
 	 * Changes how the knot vector is generated
 	 * @param Mode - The new knot generation mode
diff --git a/Engine/Source/Runtime/GeometryCore/Public/Splines/ParameterizedStorage.h b/Engine/Source/Runtime/GeometryCore/Public/Splines/ParameterizedStorage.h
--- a/Engine/Source/Runtime/GeometryCore/Public/Splines/ParameterizedStorage.h
+++ b/Engine/Source/Runtime/GeometryCore/Public/Splines/ParameterizedStorage.h
@@ -46,6 +46,38 @@ public:
         Entries.Insert({Parameter, Value}, Index);
     	return Index;
     }
+	/**
+	 * Remove value from the storage
+	 * @param Index - Index of the value to remove
+	 * @return true if the index was valid and value was removed
+	 */
+	bool RemoveValue(int32 Index)
+	{
+		if (Entries.IsValidIndex(Index))
+		{
+			Entries.RemoveAt(Index);
+			return true;
+		}
+		return false;
+	}
+
+	/**
+	 * Update an existing value at the specified index
+	 * @param Index - Index of the value to update
+	 * @param NewValue - New value to set
+	 * @return true if the index was valid and value was updated
+	 */
+	bool UpdateValue(int32 Index, const ValueType& NewValue)
+	{
+		if (!Entries.IsValidIndex(Index))
+		{
+			return false;
+		}
+		
+		Entries[Index].Value = NewValue;
+		return true;	
+	}
+	
 
 	// Read-only access through view
 	TArrayView<const TStorageEntry<ValueType>> GetEntries() const
diff --git a/Engine/Source/Runtime/GeometryCore/Public/Splines/PolyBezierSpline.h b/Engine/Source/Runtime/GeometryCore/Public/Splines/PolyBezierSpline.h
--- a/Engine/Source/Runtime/GeometryCore/Public/Splines/PolyBezierSpline.h
+++ b/Engine/Source/Runtime/GeometryCore/Public/Splines/PolyBezierSpline.h
@@ -204,16 +204,92 @@ public:
         float Parameter) const
     {
         // Order 0 is just regular evaluation
-        if constexpr (Order == 0)
-        {
-            return InterpolateWindow(FindWindow(Entries, Parameter), Parameter);
-        }
+	    if constexpr (Order == 0)
+	    {
+	        return InterpolateWindow(FindWindow(Entries, Parameter), Parameter);
+	    }
+
+	    // Map to parent space
+	    const float MappedParam = Base::MapToParentSpace(Parameter);
+	    
+		// Get window for the current parameter
+	    TArray<const ValueType*> Window = FindWindow(Entries, Parameter);
+	    if (Window.Num() < 4)  // Need all 4 control points for a cubic Bezier
+	    {
+	        return ValueType();
+	    }
+	    
+	    // Find segment index and local parameter
+	    const float ClampedParam = FMath::Clamp(MappedParam, 0.0f, static_cast<float>(NumSegments));
+	    int32 SegmentIndex;
+		float LocalT;
+    	
+    	// Correctly handle exact segment boundary cases
+    	if (ClampedParam == static_cast<float>(NumSegments)) 
+    	{
+    		// If exactly at the end of the last segment
+    		LocalT = 1.0f;
+    		SegmentIndex = NumSegments - 1;
+    	}
+    	else
+    	{
+    		SegmentIndex = FMath::FloorToInt(ClampedParam);
+    		LocalT = ClampedParam - static_cast<float>(SegmentIndex);
+    	}
+		        
+	    // For PolyBezier, each segment is normalized to unit parameter space
+	    const float SegmentScale = 1.0f;
+	    
+	    // Compute the derivative with proper scaling
+	    return Math::TBezierDerivativeCalculator<ValueType, Order>::Compute(
+	        Window, LocalT, SegmentScale);
+    }
 
-        TArray<const ValueType*> Window = FindWindow(Entries, Parameter);
-        const float MappedParam = Base::MapToParentSpace(Parameter);
-        return Math::TBezierDerivativeCalculator<ValueType, Order>::Compute(Window, MappedParam);
+	/**
+	 * Alternative implementation that leverages B-spline properties for cubic Bezier curves
+	 * A cubic Bezier curve is a B-spline of degree 3 with a specific knot vector
+	 */
+	template<int32 Order>
+	ValueType EvaluateDerivativeAsBSpline(
+		TArrayView<const TStorageEntry<ValueType>> Entries,
+		float Parameter) const
+    {
+    	// Order 0 is just regular evaluation
+    	if constexpr (Order == 0)
+    	{
+    		return InterpolateWindow(FindWindow(Entries, Parameter), Parameter);
+    	}
+
+    	// Map to parent space
+    	const float MappedParam = Base::MapToParentSpace(Parameter);
+    
+    	// Find segment index and local parameter
+    	const float ClampedParam = FMath::Clamp(MappedParam, 0.0f, static_cast<float>(NumSegments));
+    	const int32 SegmentIndex = FMath::Min(FMath::FloorToInt(ClampedParam), NumSegments - 1);
+    	const float LocalT = ClampedParam - static_cast<float>(SegmentIndex); // Parameter [0,1] within segment
+    
+    	// For a cubic Bezier segment, construct the appropriate knot vector
+    	TArray<float> KnotVector;
+    	KnotVector.SetNum(8); // For a cubic B-spline with one segment
+    
+    	// Knot vector for this segment (clamped at both ends)
+    	for (int32 i = 0; i < 4; ++i)
+    		KnotVector[i] = 0.0f;  // [0,0,0,0]
+    
+    	for (int32 i = 4; i < 8; ++i)
+    		KnotVector[i] = 1.0f;  // [1,1,1,1]
+    
+    	// Get control points for this segment
+    	TArray<const ValueType*> Window = FindWindow(Entries, Parameter);
+    
+    	// For a cubic Bezier curve represented as a B-spline, the span is always 3
+    	const int32 Span = 3;
+    
+    	// Compute the derivative using B-spline methods
+    	return Math::TBSplineDerivativeCalculator<ValueType, Order>::Compute(
+			Window, KnotVector, Span, LocalT, 3);
     }
-   
+	
     virtual float GetParentSpaceMapping(float NormalizedParam) const override
     {
         return NormalizedParam * static_cast<float>(NumSegments);
@@ -260,29 +336,31 @@ public:
      * @param P1 First control point
      * @param P2 Second control point
      * @param P3 End point
+     * @param bEnforceContinuity If true, enforces that the segment connects to the previous one
      * @return Index of the new segment
      */
-    int32 AddBezierSegment(const ValueType& P0, const ValueType& P1, const ValueType& P2, const ValueType& P3)
+    int32 AddBezierSegment(const ValueType& P0, const ValueType& P1, const ValueType& P2, const ValueType& P3, bool bEnforceContinuity = true)
     {
         const int32 NumPoints = Base::GetNumKeys();
-        
-        // Validate continuity with previous segment if this isn't the first one
-        if (NumPoints > 0)
-        {
-            const ValueType& LastPoint = Base::GetValueAt(NumPoints - 1);
-            if (!LastPoint.Equals(P0))
-            {
-                UE_LOG(LogSpline, Warning, TEXT("PolyBezier segment start point doesn't match previous end point"));
-                return INDEX_NONE;
-            }
-        }
-
-        // Add control points using base storage
-        const float SegmentT = static_cast<float>(NumPoints / 4);
-        Base::AddValue(P0, SegmentT);
-        Base::AddValue(P1, SegmentT);
-        Base::AddValue(P2, SegmentT);
-        Base::AddValue(P3, SegmentT);
+        ValueType AdjustedP0 = P0;
+    	
+    	// Check continuity with previous segment if this isn't the first one
+    	if (NumPoints > 0 && bEnforceContinuity)
+    	{
+    		const ValueType& LastPoint = Base::GetValueAt(NumPoints - 1);
+    		if (!LastPoint.Equals(P0, UE_KINDA_SMALL_NUMBER))
+    		{
+    			UE_LOG(LogSpline, Warning, TEXT("PolyBezier segment start point doesn't match previous end point. Adjusting for continuity."));
+    			AdjustedP0 = LastPoint; // Adjust the start point for continuity
+    		}
+    	}
+
+    	// Add control points using base storage
+    	const float SegmentT = static_cast<float>(NumPoints / 4);
+    	Base::AddValue(AdjustedP0, SegmentT);
+    	Base::AddValue(P1, SegmentT);
+    	Base::AddValue(P2, SegmentT);
+    	Base::AddValue(P3, SegmentT);
 
         // Update segment structure
         UpdateSegmentStructure();
@@ -291,13 +369,38 @@ public:
     }
 
     /**
- * Updates a single control point within a Bezier segment
- * @param SegmentIndex - Which cubic Bezier segment (each has 4 points)
- * @param PointIndex - Which point in the segment (0=P0, 1=P1, 2=P2, 3=P3)
- * @param NewValue - New position for the control point
- * @return true if successfully updated
- */
-    bool UpdateSegmentPoint(int32 SegmentIndex, int32 PointIndex, const ValueType& NewValue)
+     * Removes a Bezier segment from the spline
+     * @param SegmentIndex Segment index to remove
+     * @return true if successfully removed
+     */
+    bool RemoveSegment(const int32 SegmentIndex)
+    {
+    	const int32 NumSegments = GetNumSegments();
+    	if (SegmentIndex < 0 || SegmentIndex >= NumSegments)
+    	{
+    		return false;
+    	}
+
+    	const int32 BaseIndex = SegmentIndex * 4;
+    	bool bSuccess = true;
+    	bSuccess &= Base::RemoveValue(BaseIndex + 3);
+    	bSuccess &= Base::RemoveValue(BaseIndex + 2);
+    	bSuccess &= Base::RemoveValue(BaseIndex + 1);
+    	bSuccess &= Base::RemoveValue(BaseIndex);
+
+    	// Update segment structure
+    	UpdateSegmentStructure();
+    	
+    	return bSuccess;
+    }
+    /**
+	 * Updates a single control point within a Bezier segment
+	 * @param SegmentIndex - Which cubic Bezier segment (each has 4 points)
+	 * @param PointIndex - Which point in the segment (0=P0, 1=P1, 2=P2, 3=P3)
+	 * @param NewValue - New position for the control point
+	 * @return true if successfully updated
+	 */
+    bool UpdateSegmentPoint(const int32 SegmentIndex, const int32 PointIndex, const ValueType& NewValue)
     {
         const int32 NumSegments = GetNumSegments();
         if (SegmentIndex < 0 || SegmentIndex >= NumSegments ||
@@ -341,10 +444,65 @@ public:
 
 
     /** Returns the number of Bezier segments in the spline */
-    int32 GetNumSegments() const
+    virtual int32 GetNumSegments() const override
     {
         const int32 NumPoints = this->GetNumKeys();
-        return NumPoints/4;//this->IsClosedLoop() ? NumPoints/4 : (NumPoints/4 - 1);
+        return NumPoints/4;
+    }
+
+	virtual float FindNearest(const ValueType& Point, float& OutSquaredDistance) const override
+    {
+    	// Can use same structure as TBSpline since we're cubic
+    	float BestDistSq = TNumericLimits<float>::Max();
+    	float BestParam = 0.0f;
+    	int32 BestSegment = 0;
+
+    	const int32 NumSegments = GetNumSegments();
+    	for (int32 i = 0; i < NumSegments; ++i)
+    	{
+    		float SegmentDistSq;
+    		const float LocalT = FindNearestOnSegment(Point, i, SegmentDistSq);
+    		// Convert local [0,1] parameter to normalized spline parameter
+    		const float SplineParam = (i + LocalT) / static_cast<float>(NumSegments);
+    		if (SegmentDistSq < BestDistSq)
+    		{
+    			BestDistSq = SegmentDistSq;
+    			BestParam = SplineParam;
+    			BestSegment = i;
+    		}
+    	}
+
+    	OutSquaredDistance = BestDistSq;
+    	
+    	return BestParam;
+    }
+	
+	virtual float FindNearestOnSegment(const ValueType& Point, int32 SegmentIndex, float& OutSquaredDistance) const override
+    {
+    	if (!IsValidSegmentIndex(SegmentIndex))
+    	{
+    		OutSquaredDistance = TNumericLimits<float>::Max();
+    		return 0.0f;
+    	}
+
+    	// Get Bezier control points for this segment
+    	TArray<ValueType> Coeffs;
+    	Coeffs.SetNum(4);
+    	const int32 BaseIndex = SegmentIndex * 4;  // Changed from *3 to *4
+        
+    	// Get the control points
+    	const ValueType& P0 = this->GetValueAt(BaseIndex);
+    	const ValueType& P1 = this->GetValueAt(BaseIndex + 1);
+    	const ValueType& P2 = this->GetValueAt(BaseIndex + 2);
+    	const ValueType& P3 = this->GetValueAt(BaseIndex + 3);
+
+    	// Setup coefficients in standard Bezier polynomial form relative to test point
+    	Coeffs[0] = P0 - Point;                    // constant term
+    	Coeffs[1] = 3 * (P1 - P0);                 // linear coefficient
+    	Coeffs[2] = 3 * (P2 - 2*P1 + P0);         // quadratic coefficient  
+    	Coeffs[3] = P3 - 3*P2 + 3*P1 - P0;        // cubic coefficient
+
+    	return Math::FindNearestPoint_Cubic(Coeffs, 0.0f, 1.0f, OutSquaredDistance);
     }
 
     /** Gets the endpoint of a segment */
diff --git a/Engine/Source/Runtime/GeometryCore/Public/Splines/Spline.h b/Engine/Source/Runtime/GeometryCore/Public/Splines/Spline.h
--- a/Engine/Source/Runtime/GeometryCore/Public/Splines/Spline.h
+++ b/Engine/Source/Runtime/GeometryCore/Public/Splines/Spline.h
@@ -227,6 +227,16 @@ public:
 		GetEvaluator().ExpandRangeToInclude(Parameter);
 		return Values.AddValue(Value, Parameter);		
 	}
+	
+	/**
+	 * Removes value from the spline's main storage channel.
+	 * @param Index - Index of the value to remove
+	 * @return true if the value was removed
+	 */
+	virtual bool RemoveValue(int32 Index)
+	{
+		return Values.RemoveValue(Index);		
+	}
 
 	/**
      * Updates an existing value at the specified index
@@ -265,7 +275,7 @@ public:
 	template<int32 Order>
 	ValueType EvaluateDerivative(float Parameter) const
 	{
-		return Values.GetEvaluator().template EvaluateDerivative<Order>(
+		return GetEvaluator().template EvaluateDerivative<Order>(
 			Values.GetStorage().GetEntries(), Parameter);
 	}
 
@@ -462,6 +472,28 @@ protected:
 	/** Hook for derived classes to copy their specific data after base data is copied */
 	virtual void OnClone(TSharedPtr<TSplineBase> NewSpline, ESplineCloneFlags CloneFlags) const {}
 
+	/** 
+	 * Get the number of segments in the spline. Derived classes must implement this.
+	 * @return Number of segments
+	 */
+	virtual int32 GetNumSegments() const = 0;
+	
+	/**
+	 * Find nearest point on the spline to a target point
+	 * @param Point - Target point
+	 * @outparam OutSquaredDistance - Squared distance to the nearest point
+	 * @return Parameter value of the nearest point
+	 */
+	virtual float FindNearest(const ValueType& Point, float& OutSquaredDistance) const = 0;
+
+	/**
+	 * Find nearest point on a specific segment
+	 * @param Point - Target point
+	 * @param SegmentIndex - Index of segment to search
+	 * @outparam OutSquaredDistance - Squared distance to the nearest point
+	 * @return Parameter value of the nearest point
+	 */
+	virtual float FindNearestOnSegment(const ValueType& Point, int32 SegmentIndex, float& OutSquaredDistance) const = 0;
 private:
 	/**
 	 * Creates a new instance of the correct derived type for cloning.
@@ -476,7 +508,7 @@ private:
 
 		// Create new instance of correct derived type
 		TSharedPtr<TSplineBase> NewSpline = CreateNewInstance();
-        NewSpline->Values = MakeUnique<TAttributeChannel<ValueType>>(*Values);
+        NewSpline->Values = Values;
 
 		// Copy parameter range settings
 		float Start, End;
diff --git a/Engine/Source/Runtime/GeometryCore/Public/Splines/SplineMath.h b/Engine/Source/Runtime/GeometryCore/Public/Splines/SplineMath.h
--- a/Engine/Source/Runtime/GeometryCore/Public/Splines/SplineMath.h
+++ b/Engine/Source/Runtime/GeometryCore/Public/Splines/SplineMath.h
@@ -22,7 +22,10 @@ struct FSplineValidation
 		{
 			return T(ForceInit);
 		}
-		return T(); 
+		else
+		{
+			return T(); 
+		} 
 	}
 
 	template<typename T>
@@ -72,7 +75,7 @@ struct FSplineValidation
 	}
 
 	template<typename T>
-	static bool IsValidGeometry(TArrayView<const T* const> Window, float Tolerance = KINDA_SMALL_NUMBER)
+	static bool IsValidGeometry(TArrayView<const T* const> Window, float Tolerance = UE_KINDA_SMALL_NUMBER)
 	{
 		if (Window.Num() < 2) return false;
 
@@ -127,6 +130,63 @@ const TArray<TArray<double>>& GetBinomialTable()
 	return Table;
 }
 
+// compute basis functions for a given degree
+static void ComputeBSplineBasisFunctions(
+	const TArray<float>& KnotVector,
+	int32 Span,
+	float Parameter, 
+	int32 Degree,
+	TArray<float>& OutBasis)
+{
+	OutBasis.SetNum(Degree + 1);
+	TArray<float> Left;
+	TArray<float> Right;
+	Left.SetNum(Degree + 1);
+	Right.SetNum(Degree + 1);
+
+	OutBasis[0] = 1.0f;
+
+	for (int32 j = 1; j <= Degree; ++j)
+	{
+		const int32 LeftIndex = Span + 1 - j;
+		const int32 RightIndex = Span + j;
+
+		if (LeftIndex < 0 || LeftIndex >= KnotVector.Num() ||
+			RightIndex < 0 || RightIndex >= KnotVector.Num())
+		{
+			UE_LOG(LogTemp, Warning, TEXT("Invalid knot indices in ComputeBasisFunctions: LeftIndex=%d, RightIndex=%d"),
+				   LeftIndex, RightIndex);
+			// Instead of continuing, we should zero out the basis and return.
+			OutBasis.SetNumZeroed(Degree + 1);
+			return;
+		}
+
+		Left[j] = Parameter - KnotVector[LeftIndex];
+		Right[j] = KnotVector[RightIndex] - Parameter;
+
+		float Saved = 0.0f;
+
+		for (int32 r = 0; r < j; ++r)
+		{
+			float Temp;
+			const float Denominator = Right[r + 1] + Left[j - r];
+            // Correctly handle zero denominator
+            if (FMath::IsNearlyZero(Denominator))
+			{
+				Temp = 0.0f;
+			}
+			else
+			{
+				Temp = OutBasis[r] / Denominator;
+			}
+
+			OutBasis[r] = Saved + Right[r + 1] * Temp;
+			Saved = Left[j - r] * Temp;
+		}
+
+		OutBasis[j] = Saved;
+	}
+}
 /** Template-based derivative calculators for spline types */
 
 // Helper struct for shared derivative implementation
@@ -150,8 +210,8 @@ struct TGenericDerivativeHelper
 		}
 		
 		// Get pre-computed binomial coefficients
-		static const auto& BinomialTable = GetBinomialTable<20>();
-		static const auto& FactorialTable = GetFactorialTable<20>();
+		static const TArray<TArray<double>>& BinomialTable = GetBinomialTable<20>();
+		static const TArray<double>& FactorialTable = GetFactorialTable<20>();
 
 		const int32 WindowSize = Window.Num();
 		if (Order >= WindowSize || Order >= BinomialTable.Num())
@@ -207,68 +267,166 @@ struct TGenericDerivativeHelper
 	}
 };
 
-// B-spline specific calculations
-template <typename T, int32 Order>
+// Computes the nth derivative (where n == DerivOrder) of a BSpline curve.
+// Note: This implementation assumes that DerivOrder <= Degree.
+template <typename T, int32 DerivOrder>
 struct TBSplineDerivativeCalculator
 {
+	/**
+	 * Computes the DerivOrder-th derivative of a BSpline curve at a given parameter.
+	 *
+	 * @param Window     A view of pointers to control points for the current BSpline segment.
+	 *                   For a BSpline of degree p, this view should contain p+1 control points.
+	 * @param Knots      The global knot vector.
+	 * @param Span       The knot span index such that Parameter ∈ [Knots[Span], Knots[Span+1]).
+	 * @param Parameter  The parameter at which to evaluate.
+	 * @param Degree     The degree of the current BSpline (should be ≥ DerivOrder).
+	 *
+	 * @return The DerivOrder-th derivative (of type T) evaluated at Parameter.
+	 */
 	static T Compute(
 		TArrayView<const T* const> Window,
 		const TArray<float>& Knots,
 		int32 Span,
 		float Parameter,
-		int32 Degree,
-		const T& DefaultValue = FSplineValidation::GetDefaultValue<T>())
+		int32 Degree)
 	{
-		// Extended validation
-		if (!FSplineValidation::IsValidWindow(Window, Degree + 1) ||
-			!FSplineValidation::IsValidKnotVector(Knots, Degree) ||
-			!FSplineValidation::IsValidParameter(Parameter))
-		{
-			return FSplineValidation::GetDefaultValue<T>();
-		}
-
-		// Validate span range
-		if (Span < 0 || Span >= Knots.Num() - Degree - 1)
-		{
-			return FSplineValidation::GetDefaultValue<T>();
-		}
-
-		// For higher derivatives, check geometry
-		if (Order > 1 && !FSplineValidation::IsValidGeometry(Window))
-		{
-			return FSplineValidation::GetDefaultValue<T>();
-		}
-		
-		// Compute basis functions and derivatives using Oslo algorithm
-		TArray<TArray<float>> Basis;
-		ComputeBSplineBasisFunctions<Order>(
-			Span, Parameter, Degree, Knots, Basis);
+		// Basic validation
+        if (Window.Num() < Degree + 1 || DerivOrder > Degree)
+        {
+            return T();
+        }
+
+		TArray<T> DerivControlPoints;
+        DerivControlPoints.SetNum(Degree);
+        const int32 KnotStart = Span - Degree;
+
+        // Determine if knots are uniformly spaced
+        bool bIsUniform = true;
+        const float h = Knots[1] - Knots[0];
+        for (int32 i = 2; i < Knots.Num(); ++i) {
+            if (!FMath::IsNearlyEqual(Knots[i] - Knots[i-1], h, UE_KINDA_SMALL_NUMBER)) {
+                bIsUniform = false;
+                break;
+            }
+        }
+
+        // Compute derivative control points
+        for (int32 i = 0; i < Degree; ++i)
+        {
+            const int32 KnotDenom = KnotStart + i + Degree - DerivOrder + 2; //Knot index for the denominator
+            const int32 KnotNum = KnotStart + i + 1; //Knot index in the numerator
+
+            if (KnotDenom >= Knots.Num() || KnotNum >= Knots.Num() ||
+                KnotDenom < 0 || KnotNum < 0)
+            {
+                return T();
+            }
+
+            const float Denom = Knots[KnotDenom] - Knots[KnotNum];
+
+            // Handle uniform knots differently
+            if (bIsUniform)
+            {
+                // For uniform knots, use the fixed interval h instead of Denom
+                DerivControlPoints[i] = ((*Window[i + 1] - *Window[i])) * (static_cast<float>(Degree) / h); //Correct scaling for uniform knots
+            }
+            else
+            {
+                DerivControlPoints[i] = ((*Window[i + 1] - *Window[i])) * (static_cast<float>(Degree) / Denom); //Correct scaling for non-uniform knots
+            }
+        }
+
+        // Adjust scale factor for uniform knots
+        //Removing the scale factor
+        float ScaleFactor = 1.0;
+
+        // Recurse for higher derivatives with adjusted knots and parameters
+        if constexpr (DerivOrder > 1)
+        {
+           //No change here
+            TArray<const T*> DerivWindow;
+            DerivWindow.SetNum(DerivControlPoints.Num());
+            for (int32 i = 0; i < DerivControlPoints.Num(); ++i)
+            {
+                DerivWindow[i] = &DerivControlPoints[i];
+            }
+
+            // Trim the knot vector for uniform case
+            TArray<float> TrimmedKnots;
+            if (bIsUniform)
+            {
+                TrimmedKnots.Reserve(Knots.Num() - 2);
+                for (int32 i = 1; i < Knots.Num() - 1; ++i)
+                {
+                    TrimmedKnots.Add(Knots[i]);
+                }
+            }
+            else
+            {
+                TrimmedKnots = Knots;
+            }
+
+            return TBSplineDerivativeCalculator<T, DerivOrder - 1>::Compute(
+                DerivWindow, //Reduced degree control points
+                TrimmedKnots, //Remove the first and the last knot as the degree is now p-1
+                Span, // Keep Span the same for uniform knots
+                Parameter, //Same parameter
+                Degree - 1 //Reduce the degree
+            ) * ScaleFactor; //Multiply with scale factor
+        }
+        else
+        {
+            // Base case: First derivative
+            TArray<float> Basis;
+            Math::ComputeBSplineBasisFunctions(Knots, Span, Parameter, Degree - 1, Basis);
+
+            T Result = T();
+            for (int32 i = 0; i < Degree; ++i)
+            {
+                Result += DerivControlPoints[i] * Basis[i];
+            }
+            return Result * ScaleFactor;
+        }
+    }
+};
 
-		// Validate basis computation succeeded
-		if (Basis.Num() == 0 || Basis[Order].Num() == 0)
-		{
-			return FSplineValidation::GetDefaultValue<T>();
-		}
-		
-		// Apply basis to control points
-		T Result = FSplineValidation::GetDefaultValue<T>();
-		for (int32 i = 0; i <= Degree - Order; ++i)
-		{
-			Result += *Window[i] * Basis[Order][i];
-		}
+// Base case: 0th derivative is just the original BSpline evaluation.
+template <typename T>
+struct TBSplineDerivativeCalculator<T, 0>
+{
+	/**
+	 * Evaluates the BSpline curve at Parameter using de Boor's algorithm.
+	 *
+	 * @param Window     A view of pointers to the BSpline's control points (should be Degree+1 points).
+	 * @param Knots      The (possibly trimmed) knot vector.
+	 * @param Span       The knot span index.
+	 * @param Parameter  The parameter at which to evaluate.
+	 * @param Degree     The degree of the BSpline.
+	 *
+	 * @return The BSpline value (of type T) at Parameter.
+	 */
+	static T Compute(
+		TArrayView<const T* const> Window,
+		const TArray<float>& Knots,
+		int32 Span,
+		float Parameter,
+		int32 Degree)
+	{
+		// Regular evaluation using current basis functions
+		TArray<float> Basis;
+		ComputeBSplineBasisFunctions(Knots, Span, Parameter, Degree, Basis);
 
-		// Scale by factorial ratio
-		static const auto& FactorialTable = GetFactorialTable<20>();
-		const double Scale = FactorialTable[Degree] / FactorialTable[Degree - Order];
-		if (Scale > 0.0 && Scale < DBL_MAX)
+		T Result = T();
+		for (int32 i = 0; i <= Degree; ++i)
 		{
-			return Result * Scale;
+			Result += *Window[i] * Basis[i];
 		}
-		
-		return FSplineValidation::GetDefaultValue<T>();
+		return Result;
 	}
 };
 
+
 // Bezier curve specific calculations
 template <typename T, int32 Order>
 struct TBezierDerivativeCalculator
@@ -298,89 +456,26 @@ struct TBezierDerivativeCalculator
 		}
 
 		// Apply segment scaling based on derivative order
-		return Result * FMath::Pow(SegmentScale, Order);
-	}
-};
-
-// B-spline basis functions with derivatives
-template <int32 Order>
-void ComputeBSplineBasisFunctions(
-	int32 Span,
-	float Parameter,
-	int32 Degree,
-	const TArray<float>& Knots,
-	TArray<TArray<float>>& OutBasis)
-{
-	// Initialize workspace for Oslo algorithm
-	OutBasis.SetNum(Order + 1);
-	for (auto& Row : OutBasis)
-	{
-		Row.SetNum(Degree + 1);
-	}
-
-	// Allocate workspace arrays
-	TArray<float> Nders;
-	Nders.SetNum((Degree + 1) * 2);
-
-	// Initialize zeroth-degree functions
-	for (int32 j = 0; j <= Degree; ++j)
-	{
-		Nders[j] = 0.0f;
-	}
-
-	// Find index of knot interval containing Parameter
-	const float eps = 1e-6f;
-	int32 Index = Span - Degree;
-	for (int32 j = 0; j <= Degree; ++j, ++Index)
-	{
-		if (Parameter >= Knots[Index] - eps &&
-			Parameter < Knots[Index + 1] + eps)
-		{
-			Nders[j] = 1.0f;
-			break;
-		}
-	}
-
-	// Compute derivatives using optimized Oslo algorithm
-	for (int32 k = 1; k <= Order; ++k)
-	{
-		const int32 Stride = Degree + 1;
-		float* Previous = &Nders[0];
-		float* Current = &Nders[Stride];
-
-		// Clear current row
-		FMemory::Memzero(Current, Stride * sizeof(float));
-
-		for (int32 j = 0; j < Degree - k + 1; ++j)
+		// For derivatives, we need to divide by SegmentScale^Order because 
+		// derivatives scale inversely with parameter range
+		if constexpr (Order > 0)
 		{
-			const float OneOverDeltaKnot = 1.0f / (
-				Knots[Span + j + Degree - k + 1] - Knots[Span + j]);
-			Current[j] = Previous[j] * OneOverDeltaKnot;
-
-			if (j > 0)
+			// Calculate scaling factor: 1 / SegmentScale^Order
+			float ScalingFactor = 1.0f;
+			for (int32 i = 0; i < Order; ++i)
 			{
-				Current[j] -= Previous[j - 1] * OneOverDeltaKnot;
+				ScalingFactor /= SegmentScale;
 			}
+			return Result * ScalingFactor;
 		}
-
-		// Scale by degree
-		const float Scale = static_cast<const float>(Degree - k + 1);
-		for (int32 j = 0; j < Degree - k + 1; ++j)
+		else
 		{
-			Current[j] *= Scale;
+			return Result; // No scaling needed for position (Order = 0)
 		}
-
-		// Swap buffers for next iteration
-		Swap(Previous, Current);
 	}
+};
+
 
-	// Copy results to output basis
-	const float* FinalDerivs = &Nders[((Order + 1) % 2) * (Degree + 1)];
-	for (int32 i = 0; i <= Degree - Order; ++i)
-	{
-		OutBasis[Order][i] = FinalDerivs[i];
-	}
-}
 
 // Bezier basis computation
 template <int32 Order>
@@ -423,6 +518,75 @@ static void ComputeBezierBasis(float t, TArray<float>& OutBasis)
 		OutBasis[3] = 6.0f;
 	}
 }
+/** 
+ * Find nearest point on a cubic Bezier curve segment to a target point
+ * @param Coeffs - Array of 4 coefficients in Bezier form relative to target point
+ * @param StartParam - Start parameter of segment
+ * @param EndParam - End parameter of segment  
+ * @param OutSquaredDistance - Output squared distance to nearest point
+ * @return Parameter at nearest point between StartParam and EndParam
+ */
+
+template<typename ValueType>
+float FindNearestPoint_Cubic(
+	const TArray<ValueType>& Coeffs,
+	float StartParam,
+	float EndParam, 
+	float& OutSquaredDistance)
+{
+	// Test endpoints first
+	const ValueType& P0 = Coeffs[0]; // Already relative to target
+	const ValueType At1 = P0 + Coeffs[1] + Coeffs[2] + Coeffs[3]; // Curve at t=1
+
+	float BestDistSq = static_cast<float>(P0.SizeSquared());
+	float BestParam = StartParam;
+
+	const float EndDistSq = static_cast<float>(At1.SizeSquared());
+	if (EndDistSq < BestDistSq)
+	{
+		BestDistSq = EndDistSq;
+		BestParam = EndParam;
+	}
+
+	// Find roots of (C(t) - P).Dot(C'(t))
+	const ValueType& P1 = Coeffs[1];
+	const ValueType& P2 = Coeffs[2];
+	const ValueType& P3 = Coeffs[3];
+
+	TStaticArray<double, 6> RootCoeffs;
+	RootCoeffs[5] = 3.0 * P3.Dot(P3);
+	RootCoeffs[4] = 5.0 * P3.Dot(P2);
+	RootCoeffs[3] = 4.0 * P3.Dot(P1) + 2.0 * P2.Dot(P2);
+	RootCoeffs[2] = 3.0 * P2.Dot(P1) + 3.0 * P3.Dot(P0);
+	RootCoeffs[1] = P1.Dot(P1) + 2.0 * P2.Dot(P0);
+	RootCoeffs[0] = P1.Dot(P0);
+
+	// Find roots in [0,1]  
+	UE::Math::TPolynomialRootSolver<double, 5> RootSolver(RootCoeffs, 0, 1);
+
+	// Check each root
+	for (int32 RootIdx = 0; RootIdx < RootSolver.Roots.Num(); ++RootIdx)
+	{
+		const float t = static_cast<float>(RootSolver.Roots[RootIdx]);
+		const float Param = FMath::Lerp(StartParam, EndParam, t);
+
+		const ValueType Point = P0 +
+								P1 * t +
+								P2 * (t * t) +
+								P3 * (t * t * t);
+
+		const float DistSq = static_cast<float>(Point.SizeSquared());
+		if (DistSq < BestDistSq)
+		{
+			BestDistSq = DistSq;
+			BestParam = Param;
+		}
+	}
+
+	OutSquaredDistance = BestDistSq;
+	return BestParam;
+}
+	
 } // end namespace UE::Geometry::Spline::Math
 } // end namespace UE::Geometry::Spline
 } // end namespace UE::Geometry
diff --git a/Engine/Source/Runtime/GeometryCore/Public/Splines/TangentBezierSpline.h b/Engine/Source/Runtime/GeometryCore/Public/Splines/TangentBezierSpline.h
new file mode 100644
--- /dev/null
+++ b/Engine/Source/Runtime/GeometryCore/Public/Splines/TangentBezierSpline.h
@@ -0,0 +1,604 @@
+// Copyright Epic Games, Inc. All Rights Reserved.
+
+#pragma once
+
+#include "Splines/PolyBezierSpline.h"
+
+using namespace UE::Geometry::Spline;
+
+/**
+ * A spline that provides tangent-based control over curve shape while using 
+ * piecewise Bezier curves internally for evaluation. Supports both manual
+ * tangent control and automatic tangent computation.
+ */
+struct FTangentBezierSpline
+{
+    /** Tension parameter for auto-computed tangents [0,1] */
+    float Tension;
+
+    /** Whether tangents should be computed automatically for each point */
+    TArray<bool> AutoTangents;
+
+    FTangentBezierSpline()
+	    : Tension(0.0f)
+		
+    {
+    }
+
+    /**  Copy constructor */
+    FTangentBezierSpline(const FTangentBezierSpline& Other)
+	    : Tension(Other.Tension)
+	    , AutoTangents(Other.AutoTangents)
+	    , InternalSpline(Other.InternalSpline)
+    {
+    }
+
+    /**  Copy assignment */
+    FTangentBezierSpline& operator=(const FTangentBezierSpline& Other)
+    {
+        if (this != &Other)
+        {
+            Tension = Other.Tension;
+            AutoTangents = Other.AutoTangents;
+            InternalSpline = Other.InternalSpline;
+        }
+        return *this;
+    }
+
+    bool Serialize(FArchive& Ar)
+    {
+        Ar << Tension;
+        Ar << AutoTangents;
+        Ar << InternalSpline;
+                
+        return true;
+    }
+
+    friend FArchive& operator<<(FArchive& Ar, FTangentBezierSpline& Spline)
+    {
+        Spline.Serialize(Ar);
+        return Ar;
+    }
+
+    void SetClosedLoop(bool bInClosedLoop)
+    {
+        // Skip if state isn't changing
+	    if (bInClosedLoop == InternalSpline.IsClosedLoop())
+	        return;
+	        
+	    if (bInClosedLoop)
+	    {
+	    	// First set the flag on the internal spline
+	    	InternalSpline.SetClosedLoop(bInClosedLoop);
+	    	
+	        const int32 NumPoints = GetNumPoints();
+	        if (NumPoints >= 2)
+	        {
+	            // Get first and last points
+	            const FVector FirstPos = GetPosition(0);
+	            const FVector LastPos = GetPosition(NumPoints - 1);
+	            
+	            // Compute tangents for the closure segment
+	            FVector LastOutTangent, FirstInTangent;
+	            
+	            if (IsAutoTangent(NumPoints - 1) && IsAutoTangent(0))
+	            {
+	                // Auto tangents for closure
+	                FVector ClosingTangent = (1.0f - Tension) * (FirstPos - LastPos);
+	                LastOutTangent = ClosingTangent;
+	                FirstInTangent = -ClosingTangent;
+	            }
+	            else
+	            {
+	                // Use existing tangents
+	                LastOutTangent = GetTangentOut(NumPoints - 1);
+	                FirstInTangent = GetTangentIn(0);
+	            }
+	            
+	            // Calculate control points for the closing segment
+	            FVector P0 = LastPos;
+	            FVector P3 = FirstPos;
+	            FVector P1 = P0 + (LastOutTangent / 3.0f);
+	            FVector P2 = P3 - (FirstInTangent / 3.0f);
+	            
+	            // Add the closing segment to the internal spline
+	            InternalSpline.AddBezierSegment(P0, P1, P2, P3, false);
+	        }
+	    }
+	    else
+	    {
+	    	// For an open spline, remove the closing segment first
+	    	// then update the internal flag
+	    	const int32 NumSegments = InternalSpline.GetNumSegments();
+	    	if (NumSegments > 0)
+	    	{
+	    		// Remove the last segment
+	    		InternalSpline.RemoveSegment(NumSegments - 1);
+	    	}
+        
+	    	// Now update the flag on the internal spline
+	    	InternalSpline.SetClosedLoop(bInClosedLoop);
+	    }
+
+	    
+	    // Ensure auto tangents are updated after changing loop state
+	    UpdateAutoTangents();
+    }
+    
+    /** 
+     * Adds a new point to the spline
+     * @param Position - Location of the new point
+     * @param InTangent - Direction and magnitude of the arriving tangent
+     * @param OutTangent - Direction and magnitude of the leaving tangent
+     * @param bAutoComputeTangents - Whether to automatically compute tangents based on surrounding points
+     */
+    void AddPoint(const FVector& Position, 
+                 const FVector& InTangent = FVector::ZeroVector,
+                 const FVector& OutTangent = FVector::ZeroVector,
+                 bool bAutoComputeTangents = true)
+    {
+        const int32 NewIndex = GetNumPoints();
+        
+        // Store auto tangent flag
+        if (AutoTangents.Num() <= NewIndex)
+        {
+            AutoTangents.AddZeroed(NewIndex - AutoTangents.Num() + 1);
+        }
+        AutoTangents[NewIndex] = bAutoComputeTangents;
+        
+    	if (InternalSpline.GetNumKeys() == 0)
+	    {
+	        // First point - just store it directly
+	        InternalSpline.AddValue(FVector(Position), 0.0f);
+	        return;
+	    }
+	    
+	    if (InternalSpline.GetNumKeys() == 1)
+	    {
+	        // We have exactly one point stored - this is our second point
+	        // Get the first point, then clear to start building segments
+	        FVector P0 = FVector(InternalSpline.GetValueAt(0));
+	        FVector P3 = Position;
+	        
+	        InternalSpline.Clear();
+	        
+	        // Calculate control points
+	        FVector P1, P2;
+	        if (bAutoComputeTangents)
+	        {
+	            // For first segment with auto tangents, use simple interpolation
+	            FVector AutoTangent = (1.0f - Tension) * (P3 - P0);
+	            P1 = P0 + (AutoTangent / 3.0f);
+	            P2 = P3 - (AutoTangent / 3.0f);
+	        }
+	        else
+	        {
+	            // Use provided tangents
+	            P1 = P0 + (OutTangent / 3.0f); // Need to save first point's OutTangent
+	            P2 = P3 - (InTangent / 3.0f);
+	        }
+	        
+	        // Add the first segment
+	        InternalSpline.AddBezierSegment(P0, P1, P2, P3);
+	    }
+	    else
+	    {
+	        // Normal case - add segment to existing spline
+	        FVector P0 = GetPosition(NewIndex - 1);
+	        FVector P3 = Position;
+	        
+	        FVector P1, P2;
+	        if (bAutoComputeTangents)
+	        {
+	            // Calculate auto tangent
+	            FVector AutoTangent;
+	            if (NewIndex == 1)
+	            {
+	                // Simple case for first segment
+	                AutoTangent = (1.0f - Tension) * (P3 - P0);
+	            }
+	            else
+	            {
+	                // Use previous point for better tangent
+	                FVector PrevPos = GetPosition(NewIndex - 2);
+	                AutoTangent = (1.0f - Tension) * (P3 - PrevPos);
+	            }
+	            
+	            // Update previous point's out tangent if auto
+	            if (IsAutoTangent(NewIndex - 1))
+	            {
+	                SetTangentOut(NewIndex - 1, AutoTangent);
+	            }
+	            
+	            P1 = P0 + (GetTangentOut(NewIndex - 1) / 3.0f);
+	            P2 = P3 - (AutoTangent / 3.0f);
+	        }
+	        else
+	        {
+	            // Use provided tangents
+	            P1 = P0 + (GetTangentOut(NewIndex - 1) / 3.0f);
+	            P2 = P3 - (InTangent / 3.0f);
+	        }
+	        
+	        // Add the segment
+	        InternalSpline.AddBezierSegment(P0, P1, P2, P3);
+	    }
+	    
+	    // Update auto tangents for consistency
+	    if (bAutoComputeTangents)
+	    {
+	        UpdateAutoTangents();
+	    }
+    }
+
+    /** Removes a point from the spline */
+    void RemovePoint(int32 Index)
+    {
+        const int32 NumPoints = GetNumPoints();
+        if (Index < 0 || Index >= NumPoints)
+        {
+            return;
+        }
+        
+        // Remove associated auto tangent flag
+        if (AutoTangents.IsValidIndex(Index))
+        {
+            AutoTangents.RemoveAt(Index);
+        }
+        InternalSpline.RemoveSegment(Index);
+        UpdateAutoTangents();
+    }
+
+    /** 
+     * Modifies an existing point
+     * @param Index - Index of the point to modify
+     * @param NewPosition - New position for the point
+     * @param NewTangentIn - New arriving tangent
+     * @param NewTangentOut - New leaving tangent
+     */
+    void ModifyPoint(int32 Index, 
+                    const FVector& NewPosition,
+                    const FVector& NewTangentIn = FVector::ZeroVector,
+                    const FVector& NewTangentOut = FVector::ZeroVector)
+    {
+        const int32 NumPoints = GetNumPoints();
+        if (Index < 0 || Index >= NumPoints)
+        {
+            return;
+        }
+        
+        const bool bAutoTangent = IsAutoTangent(Index);
+        
+        // Update position and tangents
+        if (Index == 0)
+        {
+            // First point
+            UpdateBezierControlPoint(Index * 4, NewPosition);
+            
+            if (NumPoints > 1 && !bAutoTangent)
+            {
+                // Update outgoing control point
+                FVector P1 = NewPosition + (NewTangentOut / 3.0f);
+                UpdateBezierControlPoint(Index * 4 + 1, P1);
+            }
+        }
+        else if (Index == NumPoints - 1)
+        {
+            // Last point
+            UpdateBezierControlPoint(Index * 4 - 1, NewPosition);
+            
+            if (!bAutoTangent)
+            {
+                // Update incoming control point
+                FVector P2 = NewPosition - (NewTangentIn / 3.0f);
+                UpdateBezierControlPoint(Index * 4 - 2, P2);
+            }
+        }
+        else
+        {
+            // Middle point - affects two segments
+            
+            // Update the point in the previous segment
+            UpdateBezierControlPoint((Index - 1) * 4 + 3, NewPosition);
+            
+            // Update the point in the next segment
+            UpdateBezierControlPoint(Index * 4, NewPosition);
+            
+            if (!bAutoTangent)
+            {
+                // Update control points for tangents
+                FVector P2 = NewPosition - (NewTangentIn / 3.0f);
+                FVector P1 = NewPosition + (NewTangentOut / 3.0f);
+                
+                UpdateBezierControlPoint((Index - 1) * 4 + 2, P2);
+                UpdateBezierControlPoint(Index * 4 + 1, P1);
+            }
+        }
+        
+        // Update auto tangents if needed
+        if (bAutoTangent)
+        {
+            UpdateAutoTangents();
+        }
+    }
+
+    /** Gets position at parameter [0,1] */
+    FVector GetPosition(float Parameter) const
+    {
+        return FVector(InternalSpline.Evaluate(Parameter));
+    }
+    
+    /** Gets position of the specified point */
+    FVector GetPosition(int32 Index) const
+    {
+    	const int32 NumKeys = InternalSpline.GetNumKeys();
+
+    	if (NumKeys == 0 || Index < 0 || Index >= GetNumPoints())
+    		return FVector::ZeroVector;
+    
+    	if (NumKeys == 1)
+    	{
+    		// Special case when we only have one point
+    		if (Index == 0)
+    			return FVector(InternalSpline.GetValueAt(0));
+    		return FVector::ZeroVector;
+    	}
+
+    	// Regular case - convert from point index to bezier control point index
+    	if (Index == 0)
+    	{
+    		// First point is P0 of first segment
+    		return FVector(InternalSpline.GetValueAt(0));
+    	}
+    	else
+    	{
+    		// Other points are P3 of previous segments
+    		const int32 ControlPointIndex = (Index - 1) * 4 + 3;
+    		if (ControlPointIndex < NumKeys)
+    			return FVector(InternalSpline.GetValueAt(ControlPointIndex));
+    	}
+
+    	return FVector::ZeroVector;
+    }
+
+    /** Gets tangent at parameter [0,1] */
+    FVector GetTangent(float Parameter) const
+    {
+        return FVector(InternalSpline.EvaluateDerivative<1>(Parameter));
+    }
+    
+    /** Gets incoming tangent for the specified point */
+    FVector GetTangentIn(int32 Index) const
+    {
+        if (Index == 0)
+        {
+            // Special case for first point
+            if (InternalSpline.IsClosedLoop() && InternalSpline.GetNumSegments() > 0)
+            {
+                // For closed loop, look at last segment
+                const int32 LastSegment = InternalSpline.GetNumSegments() - 1;
+                FVector P3 = FVector(InternalSpline.GetValueAt(LastSegment * 4 + 3));
+                FVector P2 = FVector(InternalSpline.GetValueAt(LastSegment * 4 + 2));
+                return 3.0f * (P3 - P2);
+            }
+            return FVector::ZeroVector;
+        }
+        
+        // Normal case - get from previous segment's control points
+        const int32 BezierIndex = (Index - 1) * 4;
+        if (BezierIndex + 3 < InternalSpline.GetNumKeys())
+        {
+            FVector P3 = FVector(InternalSpline.GetValueAt(BezierIndex + 3));
+            FVector P2 = FVector(InternalSpline.GetValueAt(BezierIndex + 2));
+            return 3.0f * (P3 - P2);
+        }
+        return FVector::ZeroVector;
+    }
+    
+    /** Gets outgoing tangent for the specified point */
+    FVector GetTangentOut(int32 Index) const
+    {
+        const int32 NumPoints = GetNumPoints();
+        if (Index == NumPoints - 1)
+        {
+            // Special case for last point
+            if (InternalSpline.IsClosedLoop() && InternalSpline.GetNumSegments() > 0)
+            {
+                // For closed loop, look at first segment
+                FVector P0 = FVector(InternalSpline.GetValueAt(0));
+                FVector P1 = FVector(InternalSpline.GetValueAt(1));
+                return 3.0f * (P1 - P0);
+            }
+            return FVector::ZeroVector;
+        }
+        
+        // Normal case - get from next segment's control points
+        const int32 BezierIndex = Index * 4;
+        if (BezierIndex + 1 < InternalSpline.GetNumKeys())
+        {
+            FVector P0 = FVector(InternalSpline.GetValueAt(BezierIndex));
+            FVector P1 = FVector(InternalSpline.GetValueAt(BezierIndex + 1));
+            return 3.0f * (P1 - P0);
+        }
+        return FVector::ZeroVector;
+    }
+
+    /** Sets tangent in for the specified point */
+    void SetTangentIn(int32 Index, const FVector& NewTangent)
+    {
+        if (Index == 0)
+        {
+            // Special case for first point in closed loop
+            if (InternalSpline.IsClosedLoop() && InternalSpline.GetNumSegments() > 0)
+            {
+                const int32 LastSegment = InternalSpline.GetNumSegments() - 1;
+                FVector P3 = FVector(InternalSpline.GetValueAt(LastSegment * 4 + 3));
+                FVector P2 = P3 - (NewTangent / 3.0f);
+                UpdateBezierControlPoint(LastSegment * 4 + 2, P2);
+            }
+            return;
+        }
+        
+        // Normal case - update previous segment's P2 control point
+        const int32 BezierIndex = (Index - 1) * 4 + 2;
+        if (BezierIndex < InternalSpline.GetNumKeys())
+        {
+            FVector P3 = FVector(InternalSpline.GetValueAt(BezierIndex + 1));
+            FVector P2 = P3 - (NewTangent / 3.0f);
+            UpdateBezierControlPoint(BezierIndex, P2);
+        }
+    }
+    
+    /** Sets tangent out for the specified point */
+    void SetTangentOut(int32 Index, const FVector& NewTangent)
+    {
+        const int32 NumPoints = GetNumPoints();
+        if (Index == NumPoints - 1)
+        {
+            // Special case for last point in closed loop
+            if (InternalSpline.IsClosedLoop() && InternalSpline.GetNumSegments() > 0)
+            {
+                FVector P0 = FVector(InternalSpline.GetValueAt(0));
+                FVector P1 = P0 + (NewTangent / 3.0f);
+                UpdateBezierControlPoint(1, P1);
+            }
+            return;
+        }
+        
+        // Normal case - update next segment's P1 control point
+        const int32 BezierIndex = Index * 4 + 1;
+        if (BezierIndex < InternalSpline.GetNumKeys())
+        {
+            FVector P0 = FVector(InternalSpline.GetValueAt(BezierIndex - 1));
+            FVector P1 = P0 + (NewTangent / 3.0f);
+            UpdateBezierControlPoint(BezierIndex, P1);
+        }
+    }
+
+    /** Gets the number of points in the spline */
+    int32 GetNumPoints() const
+    {
+    	const int32 NumKeys = InternalSpline.GetNumKeys();
+    
+    	if (NumKeys == 0)
+    		return 0;
+    	else if (NumKeys == 1)
+    		return 1;  // Special case when we just have the first point
+    	else
+    		return InternalSpline.GetNumSegments() + 1;  // Regular case: segments + 1
+    }
+
+    /** Checks if a point's tangents are automatically computed */
+    bool IsAutoTangent(int32 Index) const
+    {
+        return (AutoTangents.IsValidIndex(Index) && AutoTangents[Index]);
+    }
+    
+    /** Sets whether a point's tangents should be automatically computed */
+    void SetPointAutoTangents(int32 Index, bool bAuto)
+    {
+        if (Index < 0 || Index >= GetNumPoints())
+        {
+            return;
+        }
+        
+        // Ensure array is big enough
+        if (AutoTangents.Num() <= Index)
+        {
+            AutoTangents.AddZeroed(Index - AutoTangents.Num() + 1);
+        }
+        
+        AutoTangents[Index] = bAuto;
+        
+        // Update tangents if needed
+        if (bAuto)
+        {
+            UpdateAutoTangents();
+        }
+    }
+
+    template <typename AttrType>
+    void CreateAttributeChannel(const FName& Name)
+    {
+        InternalSpline.CreateAttributeChannel<AttrType>(Name);
+    }
+    
+    template <typename AttrType>
+    AttrType EvaluateAttribute(FName Name, float Parameter) const
+    {
+        return InternalSpline.EvaluateAttribute<AttrType>(Name, Parameter);
+    }
+
+    template <typename AttrType>
+    int32 AddAttributeValue(FName Name, const AttrType& Value, float Parameter)
+    {
+        return InternalSpline.AddAttributeValue<AttrType>(Name, Value, Parameter);
+    }
+
+    /** Updates auto-computed tangents based on surrounding points */
+    void UpdateAutoTangents()
+    {
+        const int32 NumPoints = GetNumPoints();
+        if (NumPoints < 2)
+        {
+            return;
+        }
+
+        // Update auto tangents for each point
+        for (int32 i = 0; i < NumPoints; ++i)
+        {
+            if (!IsAutoTangent(i))
+            {
+                continue;
+            }
+
+            // Get previous and next points (handling closed loop case)
+            const FVector& Prev = (i > 0) ? GetPosition(i-1) : 
+                (InternalSpline.IsClosedLoop() ? GetPosition(NumPoints-1) : GetPosition(i));
+            
+            const FVector& Next = (i < NumPoints-1) ? GetPosition(i+1) :
+                (InternalSpline.IsClosedLoop() ? GetPosition(0) : GetPosition(i));
+
+            // Current point position
+            const FVector& Current = GetPosition(i);
+
+            // Compute tangents based on surrounding points
+            const FVector Tangent = (1.0f - Tension) * (Next - Prev);
+            
+            // Update tangents
+            SetTangentIn(i, Tangent);
+            SetTangentOut(i, Tangent);
+        }
+    }
+    
+    float FindNearest(const FVector& Point, float& OutSquaredDistance) const
+    {
+        // Delegate to internal spline
+        return InternalSpline.FindNearest(Point, OutSquaredDistance);
+    }
+
+    float FindNearestOnSegment(const FVector& Point, int32 SegmentIndex, float& OutSquaredDistance) const
+    {
+        // Delegate to internal spline
+        return InternalSpline.FindNearestOnSegment(Point, SegmentIndex, OutSquaredDistance);
+    }
+    
+    /** Getter for internal spline */
+    const FPolyBezierSpline3d& GetInternalSpline() const
+    {
+        return InternalSpline;
+    }
+
+private:
+    /** Internal PolyBezier implementation */
+    FPolyBezierSpline3d InternalSpline;
+	
+    /** Updates a Bezier control point in the internal spline */
+    void UpdateBezierControlPoint(int32 Index, const FVector& NewValue)
+    {
+        if (Index < InternalSpline.GetNumKeys())
+        {
+            InternalSpline.UpdateValue(Index, NewValue);
+        }
+    }
+};
\ No newline at end of file
