Splines: horde compilation issues

diff --git a/Engine/Source/Runtime/GeometryCore/Public/Splines/PolyBezierSpline.h b/Engine/Source/Runtime/GeometryCore/Public/Splines/PolyBezierSpline.h
--- a/Engine/Source/Runtime/GeometryCore/Public/Splines/PolyBezierSpline.h
+++ b/Engine/Source/Runtime/GeometryCore/Public/Splines/PolyBezierSpline.h
@@ -165,7 +165,7 @@ public:
         const float ParentParam = this->MapToParentSpace(Parameter);
         
         // If we're at the last point exactly, just return it
-        if (ParentParam >= NumSegments)
+        if (ParentParam >= static_cast<float>(NumSegments))
         {
             return *Window[3];  // Last control point (P3) of final Bezier segment
         }
@@ -193,7 +193,7 @@ public:
 
     /** 
      * Evaluate nth derivative at parameter
-     * @param Order - The derivative order (0 = position, 1 = first derivative, etc.)
+     * @tparam Order - The derivative order (0 = position, 1 = first derivative, etc.)
      * @param Entries - Spline control points
      * @param Parameter - Parameter in normalized space [0,1]
      * @return nth derivative vector
@@ -455,7 +455,6 @@ public:
     	// Can use same structure as TBSpline since we're cubic
     	float BestDistSq = TNumericLimits<float>::Max();
     	float BestParam = 0.0f;
-    	int32 BestSegment = 0;
 
     	const int32 NumSegments = GetNumSegments();
     	for (int32 i = 0; i < NumSegments; ++i)
@@ -463,12 +462,11 @@ public:
     		float SegmentDistSq;
     		const float LocalT = FindNearestOnSegment(Point, i, SegmentDistSq);
     		// Convert local [0,1] parameter to normalized spline parameter
-    		const float SplineParam = (i + LocalT) / static_cast<float>(NumSegments);
+    		const float SplineParam = (static_cast<float>(i) + LocalT) / static_cast<float>(NumSegments);
     		if (SegmentDistSq < BestDistSq)
     		{
     			BestDistSq = SegmentDistSq;
     			BestParam = SplineParam;
-    			BestSegment = i;
     		}
     	}
 
