[AutoRTFM] Enable GeometryCore in AutoRTFM.

This is reachable from Verse when movie playback is triggered from a script.

diff --git a/Engine/Source/Programs/AutoRTFMEngineTests/AutoRTFMEngineTests.Build.cs b/Engine/Source/Programs/AutoRTFMEngineTests/AutoRTFMEngineTests.Build.cs
--- a/Engine/Source/Programs/AutoRTFMEngineTests/AutoRTFMEngineTests.Build.cs
+++ b/Engine/Source/Programs/AutoRTFMEngineTests/AutoRTFMEngineTests.Build.cs
@@ -17,6 +17,7 @@ public class AutoRTFMEngineTests : ModuleRules
 				"AutomationWorker",
 				"Core",
 				"Engine",
+				"GeometryCore",
 				"HeadMountedDisplay",
 				"HTTP",
 				"InstallBundleManager",
diff --git a/Engine/Source/Programs/AutoRTFMEngineTests/Private/AutoRTFMEngineTests.cpp b/Engine/Source/Programs/AutoRTFMEngineTests/Private/AutoRTFMEngineTests.cpp
--- a/Engine/Source/Programs/AutoRTFMEngineTests/Private/AutoRTFMEngineTests.cpp
+++ b/Engine/Source/Programs/AutoRTFMEngineTests/Private/AutoRTFMEngineTests.cpp
@@ -280,6 +280,9 @@ static void LoadModules()
 	FModuleManager::Get().LoadModule(TEXT("ChaosSolvers"));
 	ProcessNewlyLoadedUObjects();
 
+	// The ConvexHull solver requires initializing the Shewchuk Exact Predicates, which is part of the GeometryCore init.
+	FModuleManager::Get().LoadModule(TEXT("GeometryCore"));
+
 	//FModuleManager::Get().LoadModule(TEXT("IrisCore"));
 }
 
diff --git a/Engine/Source/Runtime/GeometryCore/GeometryCore.Build.cs b/Engine/Source/Runtime/GeometryCore/GeometryCore.Build.cs
--- a/Engine/Source/Runtime/GeometryCore/GeometryCore.Build.cs
+++ b/Engine/Source/Runtime/GeometryCore/GeometryCore.Build.cs
@@ -22,7 +22,5 @@ public class GeometryCore : ModuleRules
 		);
 
 		CppCompileWarningSettings.UnsafeTypeCastWarningLevel = WarningLevel.Error;
-
-		bDisableAutoRTFMInstrumentation = true;
 	}
 }
diff --git a/Engine/Source/Runtime/GeometryCore/Private/Tests/AutoRTFMConvexHullTests.cpp b/Engine/Source/Runtime/GeometryCore/Private/Tests/AutoRTFMConvexHullTests.cpp
new file mode 100644
--- /dev/null
+++ b/Engine/Source/Runtime/GeometryCore/Private/Tests/AutoRTFMConvexHullTests.cpp
@@ -0,0 +1,43 @@
+// Copyright Epic Games, Inc. All Rights Reserved.
+
+#include "Misc/AutomationTest.h"
+#include "AutoRTFM/AutoRTFM.h"
+#include "CompGeom/ConvexHull3.h"
+#include "Math/UnrealMathUtility.h"
+#include "Math/Vector.h"
+
+#if WITH_DEV_AUTOMATION_TESTS
+
+IMPLEMENT_SIMPLE_AUTOMATION_TEST(FAutoRTFMConvexHullTests, "AutoRTFM + TConvexHull3", EAutomationTestFlags::EngineFilter | EAutomationTestFlags::ClientContext | EAutomationTestFlags::ServerContext | EAutomationTestFlags::CommandletContext)
+
+bool FAutoRTFMConvexHullTests::RunTest(const FString& Parameters)
+{
+	if (!AutoRTFM::ForTheRuntime::IsAutoRTFMRuntimeEnabled())
+	{
+		ExecutionInfo.AddEvent(FAutomationEvent(EAutomationEventType::Info, TEXT("SKIPPED 'FAutoRTFMConvexHullTests' test. AutoRTFM disabled.")));
+		return true;
+	}
+
+	using Vec3 = UE::Math::TVector<float>;
+
+	AutoRTFM::Commit([this]
+	{
+		static const Vec3 Pyramid[] = 
+		{
+			{ 0.f,  0.f,  0.f},
+			{20.f,  0.f,  0.f},
+			{ 0.f, 20.f,  0.f},
+			{20.f, 20.f,  0.f},
+			{10.f, 10.f,  3.f},
+		};
+
+		UE::Geometry::TConvexHull3<float> Hull;
+		double Volume = Hull.ComputeVolume(Pyramid);
+
+		TestTrueExpr(FMath::IsNearlyEqual(Volume, 400.));
+	});
+
+	return true;
+}
+
+#endif // WITH_DEV_AUTOMATION_TESTS
\ No newline at end of file
