Splines:
- Serialization and cloning issue fixes.
- InsertPoint/InsertSegment added. Added explicit Reparameterize and removed KnotMode state from BSpline evaluator, also removed other unnecessary state variables from PolyBezier evaluator.
- PolyBezier now has a constructor which takes in 4 points and a parameterization policy.
- PolyBezier and TangentBezierSpline now always have a valid segment, and the API allows to connect new segments (append, insert, prepend etc.)
- Static generator methods for common segment types (Line, Circle, Ellipse etc.)
- Various reparameterization modes (Uniform, ChordLength, Centripetal)
- Added tests for all of the above.
- Added an example structure for a custom metadata type (creation, serialization, type registration, interpolation policy etc.)

diff --git a/Engine/Source/Runtime/GeometryCore/Public/Splines/AttributeChannel.h b/Engine/Source/Runtime/GeometryCore/Public/Splines/AttributeChannel.h
--- a/Engine/Source/Runtime/GeometryCore/Public/Splines/AttributeChannel.h
+++ b/Engine/Source/Runtime/GeometryCore/Public/Splines/AttributeChannel.h
@@ -10,6 +10,38 @@ namespace Geometry
 {
 namespace Spline
 {
+// Template trait helpers for type validation
+
+/** Compile-time check for Serialize() method */
+template<typename T, typename = void>
+struct THasSerializeMethod : std::false_type {};
+
+template<typename T>
+struct THasSerializeMethod<T, 
+	std::void_t<decltype(std::declval<T>().Serialize(std::declval<FArchive&>()))>
+> : std::true_type {};
+
+/** Compile-time check for FArchive << operator */
+template<typename T, typename = void>
+struct THasArchiveOperator : std::false_type {};
+
+template<typename T>
+struct THasArchiveOperator<T, 
+	std::void_t<decltype(std::declval<FArchive&>() << std::declval<T&>())>
+> : std::true_type {};
+
+/** Helper for static_assert that depends on template parameter */
+template<typename T>
+struct TDependentFalse : std::false_type {};
+
+template<typename T>
+struct TAttributeTypeTraits
+{
+	// Static assert ensures custom types define their name
+	static_assert(TDependentFalse<T>::value, 
+		"Attribute types must specialize TAttributeTypeTraits<T> with a static 'Name' member.");
+};
+
 /** Base interface for attribute channels **/
 class IAttributeChannelBase
 {
@@ -31,6 +63,40 @@ public:
 	 * @param Ar - Archive to serialize to/from
 	 */
 	virtual void Serialize(FArchive& Ar) {};
+	
+	/**
+	 * Gets the type name for this channel
+	 * @return The registered type name
+	 */
+	virtual FName GetTypeName() const = 0;
+
+	/**
+	 * Updates the parent evaluator reference in the channel
+	 * @param NewParentEvaluator - The new parent evaluator to reference
+	 * @return True if the parent was updated
+	 */
+	virtual bool UpdateEvaluatorParent(const void* NewParentEvaluator) { return false; }
+
+	/**
+	 * Gets the number of keys in the channel
+	 * @return Number of keys
+	 */
+	virtual int32 NumKeys() const = 0;
+
+	/**
+	 * Gets the parameter value at the given index
+	 * @param Index - Index of the parameter to retrieve
+	 * @return Parameter value
+	 */
+	virtual float GetParameter(int32 Index) const = 0;
+
+	/**
+	 * Sets the parameter value at the given index
+	 * @param Index - Index of the parameter to set
+	 * @param Param - New parameter value
+	 * @return The new parameter value
+	 */
+	virtual float SetParameter(int32 Index, float Param) = 0;
 };
 
 /**
@@ -39,7 +105,13 @@ public:
 template<typename ValueType>
 class TAttributeChannel : public IAttributeChannelBase
 {
+	/** The registered type name for this channel type */
+	static FName TypeName;
+
+	/** Storage for channel values */
     TAttributeStorage<ValueType> Storage;
+
+	/** Evaluator for interpolation */
     TUniquePtr<TBaseEvaluator<ValueType>> Evaluator;
 
     // Optional overrides:
@@ -47,22 +119,27 @@ class TAttributeChannel : public IAttributeChannelBase
     TFunction<ValueType(TArrayView<const ValueType* const>, float)> InterpolateOverride;
 
 public:
+	/** Default constructor */
 	TAttributeChannel() 
 		: Evaluator(MakeUnique<TBaseEvaluator<ValueType>>())
 	{}
-	
+
+	/** Constructor with custom evaluator */
     explicit TAttributeChannel(TUniquePtr<TBaseEvaluator<ValueType>> InEvaluator)
         : Evaluator(MoveTemp(InEvaluator))
     {
     }
 
-	// Make copying logic live in the channel class
+	/** 
+	 * Implementation of the clone method
+	 * Makes a deep copy of this channel 
+	 */
 	virtual TUniquePtr<IAttributeChannelBase> Clone() const override
 	{
     	return MakeUnique<TAttributeChannel>(*this);
     }
 	
-	// Add copy constructor to properly clone the evaluator
+	/** Copy constructor with deep evaluator clone */
 	TAttributeChannel(const TAttributeChannel& Other)
 		: Storage(Other.Storage)
 		, FindWindowOverride(Other.FindWindowOverride)
@@ -74,7 +151,7 @@ public:
     	}
     }
 
-	// Add copy assignment
+	/** Copy assignment operator with deep evaluator clone */
 	TAttributeChannel& operator=(const TAttributeChannel& Other)
     {
     	if (this != &Other)
@@ -95,36 +172,41 @@ public:
     	return *this;
     }
 
+	/**
+	 * Gets the registered type name for this channel
+	 */
+	virtual FName GetTypeName() const override
+	{
+		return TypeName;
+	}
+
+	/**
+	 * Updates the parent evaluator reference in the channel
+	 * @param NewParentEvaluator - The new parent evaluator to reference
+	 * @return True if the parent was updated
+	 */
+	virtual bool UpdateEvaluatorParent(const void* NewParentEvaluator) override
+	{
+		if (Evaluator)
+		{
+			return Evaluator->UpdateParentEvaluator(NewParentEvaluator);
+		}
+		return false;
+	}
+	
+	/**
+	 * Sets the static type name for this channel type
+	 * Called during type registration
+	 */
+	static void SetTypeName(const FName& InTypeName)
+	{
+		TypeName = InTypeName;
+	}
+
     friend FArchive& operator<<(FArchive& Ar, TAttributeChannel& Channel)
     {
-    	// Serialize storage
-    	Ar << Channel.Storage;
-        
-    	// Handle evaluator serialization
-    	bool bHasEvaluator = Channel.Evaluator != nullptr;
-    	Ar << bHasEvaluator;
-        
-    	if (Ar.IsSaving())
-    	{
-    		if (bHasEvaluator)
-    		{
-    			Ar << *Channel.Evaluator;
-    		}
-    	}
-    	else
-    	{
-    		if (bHasEvaluator)
-    		{
-    			// Create appropriate evaluator type and serialize
-    			if (!Channel.Evaluator)
-    			{
-    				Channel.Evaluator = MakeUnique<TBaseEvaluator<ValueType>>();
-    			}
-    			Ar << *Channel.Evaluator;
-    		}
-    	}
-        
-    	return Ar;
+		Channel.Serialize(Ar);
+		return Ar;
     }
 
 	/**
@@ -162,8 +244,9 @@ public:
      */
     ValueType Evaluate(float Parameter) const
     {
-        if (Storage.GetEntries().Num() == 0) return ValueType();
-        if (Storage.GetEntries().Num() == 1) return Storage.GetEntries()[0].Value;
+		int32 NumValues = Storage.GetEntries().Num();
+        if (NumValues == 0) return ValueType();
+        if (NumValues == 1) return Storage.GetEntries()[0].Value;
 
         // Get window using override or evaluator
         TArray<const ValueType*> Window;
@@ -250,11 +333,117 @@ public:
     virtual void Serialize(FArchive& Ar) override
     {
 		IAttributeChannelBase::Serialize(Ar);
-		
-		Ar << *this;
+
+		// Add version for future compatibility
+		int32 Version = 1;
+		Ar << Version;
+        
+		// If we're still here, serialize our data
+		SerializeChannelData(Ar);
     }
-};
+	
+protected:
+    /**
+     * Serializes just the internal channel data
+     * @param Ar - Archive to serialize to/from
+     */
+    void SerializeChannelData(FArchive& Ar)
+    {
+        // Serialize entry count
+        int32 NumEntries = Storage.GetEntries().Num();
+        Ar << NumEntries;
+        
+        if (Ar.IsLoading())
+        {
+            Storage.GetEntries().SetNum(NumEntries);
+        }
+        
+        // Serialize each entry
+        for (int32 i = 0; i < NumEntries; ++i)
+        {
+            // Always serialize the parameter (float)
+            Ar << Storage.GetEntries()[i].Parameter;
+            
+            // Serialize the value using appropriate method
+            SerializeValue(Ar, Storage.GetEntries()[i].Value);
+        }
+        
+        // Serialize evaluator
+        bool bHasEvaluator = Evaluator != nullptr;
+        Ar << bHasEvaluator;
+        
+        if (bHasEvaluator)
+        {
+            if (Ar.IsLoading() && !Evaluator)
+            {
+                Evaluator = MakeUnique<TBaseEvaluator<ValueType>>();
+            }
+            
+            if (Evaluator)
+            {
+                Ar << *Evaluator;
+            }
+        }
+    }
+    
+    /**
+     * Helper to serialize a value based on its type properties
+     * @param Ar - Archive to serialize to/from
+     * @param Value - The value to serialize
+     */
+    void SerializeValue(FArchive& Ar, ValueType& Value)
+    {
+        // For types with custom Serialize method
+        if constexpr (THasSerializeMethod<ValueType>::value)
+        {
+            Value.Serialize(Ar);
+        }
+        // For types supporting FArchive << operator
+        else if constexpr (THasArchiveOperator<ValueType>::value)
+        {
+            Ar << Value;
+        }
+        // For simple POD types
+        else if constexpr (TIsPODType<ValueType>::Value)
+        {
+            Ar << Value;
+        }
+        // For unsupported types, provide compile-time error
+        else
+        {
+            static_assert(TDependentFalse<ValueType>::value, 
+                "Type must support serialization via Serialize() method or operator<< with FArchive");
+        }
+    }
+
+public:
+	virtual int32 NumKeys() const override
+	{
+		return Storage.GetEntries().Num();
+	}
+
+	virtual float GetParameter(int32 Index) const override
+	{
+		if (Storage.GetEntries().IsValidIndex(Index))
+		{
+			return Storage.GetEntries()[Index].Parameter;
+		}
+		return 0.0f;
+	}
 
+	virtual float SetParameter(int32 Index, float Param) override
+	{
+		if (Storage.GetEntries().IsValidIndex(Index))
+		{
+			Storage.GetEntries()[Index].Parameter = Param;
+		}
+		return Param;
+	}
+};
+	
+	// Initialize static type name for all channel types
+	template<typename ValueType>
+	FName TAttributeChannel<ValueType>::TypeName = FName(TEXT("Unknown"));
 } // end namespace UE::Geometry::Spline
 } // end namespace UE::Geometry
 } // end namespace UE
\ No newline at end of file
diff --git a/Engine/Source/Runtime/GeometryCore/Public/Splines/BSpline.h b/Engine/Source/Runtime/GeometryCore/Public/Splines/BSpline.h
--- a/Engine/Source/Runtime/GeometryCore/Public/Splines/BSpline.h
+++ b/Engine/Source/Runtime/GeometryCore/Public/Splines/BSpline.h
@@ -11,7 +11,7 @@ namespace Geometry
 {
 namespace Spline
 {
-enum class EBSplineKnotMode 
+enum class EParameterizationPolicy 
 {
 	/** Equal spacing between knots regardless of control point positions */
 	Uniform,        
@@ -20,10 +20,7 @@ enum class EBSplineKnotMode
 	ChordLength,    
 
 	/** Uses square root of distances - helps prevent loops and cusps in the curve */
-	Centripetal,    
-
-	/** User-provided custom knot vector */
-	Custom          
+	Centripetal,
 };
 
 // Forward declaration
@@ -39,17 +36,10 @@ public:
 	using ValueType = typename TChooseVectorType<ScalarType, Dimensions>::Type;
 	using Base = TBaseEvaluator<ValueType>;
 	
-private:
-	
+protected:
 	/** Vector containing knot values that determine spline parameterization */
 	TArray<float> KnotVector;
-
-	/** Number of control points in the spline */
-	int32 NumValues = 0;
-
-	/** Method used to generate the knot vector */
-	EBSplineKnotMode KnotMode = EBSplineKnotMode::Uniform;
-
+	
 	/** Whether to clamp the endpoints for interpolation of first/last control points */
 	bool bClampEnds = true;
 
@@ -126,29 +116,154 @@ public:
         
 		// Serialize BSpline-specific data
 		Ar << KnotVector;
-		Ar << NumValues;
-        
-		uint8 KnotModeInt = static_cast<uint8>(KnotMode);
-		Ar << KnotModeInt;
+		Ar << bClampEnds;
         
-		if (Ar.IsLoading())
+	}
+	/**
+	 * Override SetParentSpaceRange to update knot vector accordingly
+	 * @param InStart - New start parameter
+	 * @param InEnd - New end parameter
+	 */
+	virtual void SetParentSpaceRange(float InStart, float InEnd) override
+	{
+		float OldStart = this->ParentStart;
+	    float OldEnd = this->ParentEnd;
+	    
+	    // Only proceed if there's an actual change
+	    if (FMath::IsNearlyEqual(OldStart, InStart) && FMath::IsNearlyEqual(OldEnd, InEnd))
+	    {
+	        return;
+	    }
+	    
+	    // Call base implementation to update the parameter range
+	    Base::SetParentSpaceRange(InStart, InEnd);
+	    
+		// Scale the knot vector regardless of mode
+		if (KnotVector.Num() >= 2)
 		{
-			KnotMode = static_cast<EBSplineKnotMode>(KnotModeInt);
-		}
+			// Get the current effective knot range
+			float KnotStart, KnotEnd;
+			GetKnotRange(KnotStart, KnotEnd);
         
-		Ar << bClampEnds;
+			const float OldRange = KnotEnd - KnotStart;
+			const float NewRange = InEnd - InStart;
         
+			if (!FMath::IsNearlyZero(OldRange) && !FMath::IsNearlyZero(NewRange))
+			{
+				const float Scale = NewRange / OldRange;
+				const float Offset = InStart - KnotStart * Scale;
+            
+				// Scale the knot vector in place
+				for (int32 i = 0; i < KnotVector.Num(); ++i)
+				{
+					KnotVector[i] = KnotVector[i] * Scale + Offset;
+				}
+            
+				UE_LOG(LogSpline, Verbose, TEXT("BSpline knot vector scaled from range [%f,%f] to [%f,%f]"), 
+					KnotStart, KnotEnd, InStart, InEnd);
+			}
+		}
+	}
+
+	/**
+	 * Inserts a knot at the specified parameter value
+	 * Note: This preserves curve shape exactly but may change parameterization properties
+	 * 
+	 * @param Parameter - Parameter value where the knot should be inserted
+	 * @param Count - Number of times to insert the knot
+	 */
+	void InsertKnot(float Parameter, int32 Count = 1)
+	{
+		int32 InsertIndex = KnotVector.Num();  // Default to appending
+
+		for (int32 i = 0; i < KnotVector.Num(); ++i)
+		{
+			if (KnotVector[i] >= Parameter)
+			{
+				// Ensure we insert *after* any existing repeated values
+				while (i < KnotVector.Num() && FMath::IsNearlyEqual(KnotVector[i], Parameter))
+				{
+					i++;
+				}
+				InsertIndex = i;
+				break;
+			}
+		}
+
+		// Insert the parameter count times at the computed index
+		for (int32 i = 0; i < Count; ++i)
+		{
+			KnotVector.Insert(Parameter, InsertIndex);
+		}
 	}
+
+	/**
+	 * Returns how many times a given knot value appears.
+	 * @param Parameter - Knot value to check
+	 * @return Number of times the knot appears
+	 */
+	int32 GetKnotMultiplicity(float Parameter) const
+	{
+	    int32 Count = 0;
+	    for (const float& Knot : KnotVector)
+	    {
+	        if (FMath::IsNearlyEqual(Knot, Parameter))
+	        {
+	            ++Count;
+	        }
+	    }
+	    return Count;
+	}
+
+	/**
+	 * Removes a single occurrence of a knot value.
+	 * For clamped endpoints, we avoid removing knots if doing so would drop the multiplicity
+	 * below the required (Degree+1) at the endpoints.
+	 * @param Parameter - Knot value to remove
+	 * @param Count - Number of times to remove the knot
+	 * @return true if the knot was removed
+	 */
+	bool RemoveKnot(float Parameter, int32 Count = 1)
+	{
+		bool bSuccess = true;
+		for (int32 i = 0; i < Count; ++i)
+		{
+			int32 Multiplicity = GetKnotMultiplicity(Parameter);
+			if (bClampEnds)
+			{
+				float KnotStart, KnotEnd;
+				GetKnotRange(KnotStart, KnotEnd);
+				// If the parameter is at an endpoint and the multiplicity is already clamped,
+				// do not remove.
+				if ((FMath::IsNearlyEqual(Parameter, KnotStart) || FMath::IsNearlyEqual(Parameter, KnotEnd)) &&
+					Multiplicity <= (GetDegree() + 1))
+				{
+					UE_LOG(LogSpline, Warning, TEXT("Cannot remove clamped endpoint knot: %f"), Parameter);
+					bSuccess &= false;
+					continue;
+				}
+			}
+
+			int32 KnotIndex = KnotVector.IndexOfByKey(Parameter);
+			if (KnotIndex == INDEX_NONE)
+			{
+				bSuccess &= false;
+				continue;
+			}
+			KnotVector.RemoveAt(KnotIndex);
+			bSuccess &= true;
+		}
+	    return bSuccess;
+	}
+	
 	/**
 	 * Sets a custom knot vector and switches to custom knot mode
 	 * @param NewKnots - Array of knot values. Must be non-decreasing and have correct length (n + p + 1)
-	 * @param InNumValues - Number of values (control points) in the spline
+	 * @param NumValues - Number of values (control points) in the spline
 	 * @return True if the knot vector was valid and set successfully
 	 */
-	bool SetCustomKnots(const TArray<float>& NewKnots, int32 InNumValues)
+	bool SetCustomKnots(const TArray<float>& NewKnots, int32 NumValues)
 	{
-		NumValues = InNumValues;
-		
 		// For closed loops, we repeat the first Degree points at the end
 		const int32 RequiredKnots = this->IsClosedLoop() ? 
 			(NumValues + 2 * Degree + 1) : // Extra knots for closure
@@ -172,7 +287,6 @@ public:
 		}
 
 		KnotVector = NewKnots;
-		KnotMode = EBSplineKnotMode::Custom;
 		
 		UE_LOG(LogSpline, Verbose, TEXT("Set Custom knot vector - Points: %d, Degree: %d, Knots: %d"),
 			   NumValues, Degree, KnotVector.Num());
@@ -194,10 +308,11 @@ public:
 	 * Updates the knot vector based on current control points and settings
 	 * 
 	 * @param Points - Array of control points and their parameters
+	 * @param ParameterizationPolicy - Policy for generating knot vector
 	 */
-	void UpdateKnotVector(const TArrayView<const TStorageEntry<ValueType>>& Points)
+	void UpdateKnotVector(const TArrayView<const TStorageEntry<ValueType>>& Points, EParameterizationPolicy ParameterizationPolicy)
 	{
-		NumValues = Points.Num();
+		int32 NumValues = Points.Num();
 		if (NumValues < Degree + 1)
 		{
 			UE_LOG(LogSpline, Warning, TEXT("Not enough points for B-spline: %d (need >= %d)"),
@@ -208,29 +323,27 @@ public:
 		const int32 KnotCount = NumValues + Degree + 1;
 		KnotVector.SetNum(KnotCount);
 
-		switch (KnotMode)
+		switch (ParameterizationPolicy)
 		{
-		case EBSplineKnotMode::Uniform:
-			GenerateUniformKnots(KnotCount);
+		case EParameterizationPolicy::Uniform:
+			GenerateUniformKnots(KnotCount, Points);
 			break;
-		case EBSplineKnotMode::ChordLength:
+		case EParameterizationPolicy::ChordLength:
 			GenerateChordLengthKnots(KnotCount, Points);
 			break;
-		case EBSplineKnotMode::Centripetal:
+		case EParameterizationPolicy::Centripetal:
 			GenerateCentripetalKnots(KnotCount, Points);
 			break;
-		case EBSplineKnotMode::Custom:
-			// Do nothing - use existing knot vector
-				break;
+		}
+
+		UE_LOG(LogSpline, Verbose, TEXT("Updated knot vector - Points: %d, Degree: %d, Knots: %d"),
+			   Points.Num(), Degree, KnotCount);
+		for (int32 i = 0; i < KnotCount; ++i)
+		{
+			UE_LOG(LogSpline, Verbose, TEXT("  Knot[%d] = %f"), i, KnotVector[i]);
 		}
 	}
 	
-	/** Changes the method used to generate knot vectors */
-	void SetKnotGenerationMode(EBSplineKnotMode InMode)
-	{
-		KnotMode = InMode;
-	}
-
 	/** 
 	 * Sets whether the spline should be clamped at endpoints
 	 * When true, the curve will interpolate the first and last control points
@@ -244,11 +357,6 @@ public:
 	{
 		return bClampEnds;
 	}
-	
-	EBSplineKnotMode GetKnotMode() const
-	{
-		return KnotMode;
-	}
 
 	// Get the valid parameter range in knot space
 	void GetKnotRange(float& OutStart, float& OutEnd) const
@@ -259,7 +367,6 @@ public:
 			OutEnd = 1.0f;
 			return;
 		}
-
 		if (bClampEnds)
 		{
 			// For clamped B-splines, the valid range is [Degree, NumKnots - Degree - 1]
@@ -309,6 +416,26 @@ public:
 	}
 	
 protected:
+	int32 GetNumValuesFromKnots() const
+	{
+		// For closed loops, expect at least 2*Degree+1 knots;
+		// for open splines, expect at least Degree+1 knots.
+		const int32 MinKnots = this->IsClosedLoop() ? (2 * Degree + 1) : Degree + 1;
+		if (KnotVector.Num() < MinKnots)
+		{
+			UE_LOG(LogSpline, Warning, TEXT("GetNumValuesFromKnots: Insufficient knots. Expected at least %d but got %d."),
+				   MinKnots, KnotVector.Num());
+			return 0;
+		}
+
+		return this->IsClosedLoop() ? (KnotVector.Num() - (2 * Degree) - 1) : (KnotVector.Num() - Degree - 1);
+	}
+	int32 GetExpectedNumKnots(int32 NumValues) const
+	{
+		return this->IsClosedLoop()
+				   ? (NumValues + (2 * Degree) + 1) // closed loop: n + 2p + 1
+				   : (NumValues + Degree + 1);
+	}
 	virtual float GetParentSpaceMapping(float NormalizedParam) const override
 	{
 		if (KnotVector.Num() < 2) return NormalizedParam;
@@ -325,6 +452,7 @@ protected:
 
 		// Normalize the parameter
 		const float Range = End - Start;
+				
 		if (FMath::IsNearlyZero(Range))
 		{
 			return 0.0f;
@@ -336,14 +464,18 @@ private:
 	/** Find knot span using binary search */
 	int32 FindKnotSpan(float Parameter) const
 	{
-		// Validate we have enough knots
-		const int32 MinKnots = NumValues + Degree + 1;
-		if (KnotVector.Num() < MinKnots)
+		// First, check that the knot vector is populated.
+		if (KnotVector.Num() == 0)
 		{
-			UE_LOG(LogSpline, Error, TEXT("Not enough knots for FindKnotSpan"));
-			return 0;
+			return 0;  // Return a safe fallback.
 		}
 
+		const int32 NumValues = GetNumValuesFromKnots();
+		if (NumValues <= 0)
+		{
+			UE_LOG(LogSpline, Error, TEXT("FindKnotSpan: Not enough knots to define a valid span."));
+			return 0; // Safe fallback.
+		}
 		// Handle numerical precision at endpoints
 		const float Epsilon = 1e-6f;
 
@@ -481,7 +613,7 @@ private:
      * Generates a uniform knot vector with equal spacing
      * For clamped ends, multiplicity of degree+1 is used at endpoints
      */
-	void GenerateUniformKnots(int32 KnotCount)
+	void GenerateUniformKnots(int32 KnotCount, const TArrayView<const TStorageEntry<ValueType>>& Points)
 	{
 		if (this->IsClosedLoop())
 		{
@@ -489,7 +621,7 @@ private:
 			// - Need additional d points wrapping from start
 			// - Need n+2d+1 knots total
 			// - Knot spacing should be uniform to maintain periodicity
-			const int32 n = NumValues;
+			const int32 n = Points.Num();
 			const int32 d = Degree; 
 			const int32 NumKnots = n + 2*d + 1;
         
@@ -533,12 +665,6 @@ private:
 			}
 		}
 
-		UE_LOG(LogSpline, Verbose, TEXT("Updated knot vector - Points: %d, Degree: %d, Knots: %d"),
-			   NumValues, Degree, KnotCount);
-		for (int32 i = 0; i < KnotCount; ++i)
-		{
-			UE_LOG(LogSpline, Verbose, TEXT("  Knot[%d] = %f"), i, KnotVector[i]);
-		}
 	}
 
 	/**
@@ -546,83 +672,110 @@ private:
      * This gives better parameterization when control points are unevenly spaced
      */
     void GenerateChordLengthKnots(int32 KnotCount, const TArrayView<const TStorageEntry<ValueType>>& Points)
+{
+    if (Points.Num() < 2)
+    {
+        GenerateUniformKnots(KnotCount, Points);
+        return;
+    }
+    
+    const int32 n = Points.Num();
+    const int32 d = Degree;
+    const int32 NumKnots = this->IsClosedLoop() ? (n + 2*d + 1) : KnotCount;
+    KnotVector.SetNum(NumKnots);
+
+    // Calculate chord lengths between consecutive points
+    TArray<float> Chords;
+    Chords.Reserve(Points.Num() - 1);
+    float TotalLength = 0.0f;
+    
+    for (int32 i = 1; i < Points.Num(); ++i)
+    {
+        const float Length = static_cast<float>((Points[i].Value - Points[i-1].Value).Size());
+        Chords.Add(Length);
+        TotalLength += Length;
+    }
+
+    // For closed loop, add the closing segment length
+    if (this->IsClosedLoop() && Points.Num() > 0)
+    {
+        const float ClosureLength = static_cast<float>((Points[0].Value - Points.Last().Value).Size());
+        Chords.Add(ClosureLength);
+        TotalLength += ClosureLength;
+    }
+
+    if (this->IsClosedLoop())
     {
-        if (Points.Num() < 2)
+        // For closed loop, generate periodic knot vector based on chord lengths
+        for (int32 i = 0; i < NumKnots; ++i)
         {
-            GenerateUniformKnots(KnotCount);
-            return;
+            KnotVector[i] = static_cast<float>(i) * TotalLength / (n + d);
         }
-		
-		const int32 n = NumValues;
-		const int32 d = Degree;
-		const int32 NumKnots = this->IsClosedLoop() ? (n + 2*d + 1) : KnotCount;
-		KnotVector.SetNum(NumKnots);
-
-        // Calculate chord lengths between consecutive points
-        TArray<float> Chords;
-        Chords.Reserve(Points.Num() - 1);
-        float TotalLength = 0.0f;
+    }
+    else if (bClampEnds)
+    {
+        // For clamped ends, we need:
+        // - First (Degree+1) knots equal to 0.0
+        // - Last (Degree+1) knots equal to 1.0
+        // - Middle knots based on chord lengths
         
-        for (int32 i = 1; i < Points.Num(); ++i)
+        // First set all knots to 0
+        for (int32 i = 0; i < KnotCount; ++i)
         {
-           const float Length = static_cast<float>((Points[i].Value - Points[i-1].Value).Size());
-            Chords.Add(Length);
-            TotalLength += Length;
+            KnotVector[i] = 0.0f;
+        }
+        
+        // Set the last (Degree+1) knots to 1.0
+        for (int32 i = KnotCount - Degree - 1; i < KnotCount; ++i)
+        {
+            KnotVector[i] = 1.0f;
+        }
+        
+        // Calculate internal knots based on accumulated chord lengths
+        // We have (KnotCount - 2*(Degree+1)) internal knots to set
+        const int32 NumInternalKnots = KnotCount - 2 * (Degree + 1);
+        if (NumInternalKnots > 0)
+        {
+            // Compute normalized parameters at each control point
+            TArray<float> Params;
+            Params.SetNum(Points.Num());
+            Params[0] = 0.0f;
+            
+            float AccumulatedLength = 0.0f;
+            for (int32 i = 1; i < Points.Num(); ++i)
+            {
+                AccumulatedLength += Chords[i-1];
+                Params[i] = (TotalLength > UE_SMALL_NUMBER) ? 
+                    (AccumulatedLength / TotalLength) : 
+                    (static_cast<float>(i) / static_cast<float>(Points.Num() - 1));
+            }
+            
+            // Map these parameters to internal knots
+            for (int32 i = 0; i < NumInternalKnots; ++i)
+            {
+                // Map i to appropriate parameter range
+                const float t = static_cast<float>(i + 1) / static_cast<float>(NumInternalKnots + 1);
+                // Map t to control point indices
+                const int32 LowIndex = FMath::FloorToInt(t * (Points.Num() - 1));
+                const int32 HighIndex = FMath::Min(LowIndex + 1, Points.Num() - 1);
+                const float Alpha = t * (Points.Num() - 1) - LowIndex;
+                
+                // Interpolate parameters
+                const float KnotValue = FMath::Lerp(Params[LowIndex], Params[HighIndex], Alpha);
+                KnotVector[Degree + 1 + i] = KnotValue;
+            }
         }
-
-		// For closed loop, add the closing segment length
-		if (this->IsClosedLoop() && Points.Num() > 0)
-		{
-			const float ClosureLength = static_cast<float>((Points[0].Value - Points.Last().Value).Size());
-			Chords.Add(ClosureLength);
-			TotalLength += ClosureLength;
-		}
-
-		if (this->IsClosedLoop())
-		{
-			// For closed loop, generate periodic knot vector based on chord lengths
-			for (int32 i = 0; i < NumKnots; ++i)
-			{
-				KnotVector[i] = static_cast<float>(i) * TotalLength / (n + d);
-			}
-		}
-		else
-		{
-			// Generate knot vector based on accumulated chord lengths
-			for (int32 i = 0; i < KnotCount; ++i)
-			{
-				if (bClampEnds)
-				{
-					if (i <= Degree)
-					{
-						KnotVector[i] = 0.0f;
-					}
-					else if (i >= KnotCount - Degree - 1)
-					{
-						KnotVector[i] = 1.0f;
-					}
-					else
-					{
-						// Accumulate normalized chord lengths
-						const int32 Segment = (i - Degree);
-						float Param = 0.0f;
-						for (int32 j = 0; j < Segment; ++j)
-						{
-							Param += (TotalLength > UE_SMALL_NUMBER) ? 
-								(Chords[j] / TotalLength) : 
-								(1.0f / (float)(Points.Num() - 1));
-						}
-						KnotVector[i] = Param;
-					}
-				}
-				else
-				{
-					const float Param = static_cast<float>(i) / static_cast<float>(KnotCount - 1);
-					KnotVector[i] = Param * TotalLength;
-				}
-			}
-		}
     }
+    else
+    {
+        // Unclamped knots - simple uniform distribution based on chord lengths
+        for (int32 i = 0; i < KnotCount; ++i)
+        {
+            const float Param = static_cast<float>(i) / static_cast<float>(KnotCount - 1);
+            KnotVector[i] = Param * TotalLength;
+        }
+    }
+}
 
     /**
      * Generates a knot vector using centripetal parameterization
@@ -631,77 +784,104 @@ private:
      */
     void GenerateCentripetalKnots(int32 KnotCount, const TArrayView<const TStorageEntry<ValueType>>& Points)
 	{
-		if (Points.Num() < 2)
-		{
-			GenerateUniformKnots(KnotCount);
-			return;
-		}
+	    if (Points.Num() < 2)
+	    {
+	        GenerateUniformKnots(KnotCount, Points);
+	        return;
+	    }
 
-		const int32 n = NumValues;
-		const int32 d = Degree;
-		const int32 NumKnots = this->IsClosedLoop() ? (n + 2*d + 1) : KnotCount;
-		KnotVector.SetNum(NumKnots);
-		
-		TArray<float> Lengths;
-		Lengths.Reserve(Points.Num() - 1);
-		float TotalLength = 0.0f;
-        
-		for (int32 i = 1; i < Points.Num(); ++i)
-		{
-			const float Length = FMath::Sqrt(static_cast<float>((Points[i].Value - Points[i-1].Value).Size()));
-			Lengths.Add(Length);
-			TotalLength += Length;
-		}
-		// For closed loop, add the closing segment
-		if (this->IsClosedLoop() && Points.Num() > 0)
-		{
-			const float ClosureLength = FMath::Sqrt(static_cast<float>((Points[0].Value - Points.Last().Value).Size()));
-			Lengths.Add(ClosureLength);
-			TotalLength += ClosureLength;
-		}
+	    const int32 n = Points.Num();
+	    const int32 d = Degree;
+	    const int32 NumKnots = this->IsClosedLoop() ? (n + 2*d + 1) : KnotCount;
+	    KnotVector.SetNum(NumKnots);
+	    
+	    TArray<float> Lengths;
+	    Lengths.Reserve(Points.Num() - 1);
+	    float TotalLength = 0.0f;
+	    
+	    for (int32 i = 1; i < Points.Num(); ++i)
+	    {
+	        // Square root of distance for centripetal parameterization
+	        const float Length = FMath::Sqrt(static_cast<float>((Points[i].Value - Points[i-1].Value).Size()));
+	        Lengths.Add(Length);
+	        TotalLength += Length;
+	    }
+	    
+	    // For closed loop, add the closing segment
+	    if (this->IsClosedLoop() && Points.Num() > 0)
+	    {
+	        const float ClosureLength = FMath::Sqrt(static_cast<float>((Points[0].Value - Points.Last().Value).Size()));
+	        Lengths.Add(ClosureLength);
+	        TotalLength += ClosureLength;
+	    }
 
-		if (this->IsClosedLoop())
-		{
-			// Generate periodic knot vector using centripetal parameterization
-			for (int32 i = 0; i < NumKnots; ++i)
-			{
-				KnotVector[i] = static_cast<float>(i) * TotalLength / (n + d);
-			}
-		}
-		else
-		{
-			for (int32 i = 0; i < KnotCount; ++i)
-			{
-				if (bClampEnds)
-				{
-					if (i <= Degree)
-					{
-						KnotVector[i] = 0.0f;
-					}
-					else if (i >= KnotCount - Degree - 1)
-					{
-						KnotVector[i] = 1.0f;
-					}
-					else
-					{
-						const int32 Segment = (i - Degree);
-						float Param = 0.0f;
-						for (int32 j = 0; j < Segment; ++j)
-						{
-							Param += (TotalLength > UE_SMALL_NUMBER) ? 
-								(Lengths[j] / TotalLength) : 
-								(1.0f / (float)(Points.Num() - 1));
-						}
-						KnotVector[i] = Param;
-					}
-				}
-				else
-				{
-					const float Param = static_cast<float>(i) / static_cast<float>(KnotCount - 1);
-					KnotVector[i] = Param * TotalLength;
-				}
-			}
-		}
+	    if (this->IsClosedLoop())
+	    {
+	        // Generate periodic knot vector using centripetal parameterization
+	        for (int32 i = 0; i < NumKnots; ++i)
+	        {
+	            KnotVector[i] = static_cast<float>(i) * TotalLength / (n + d);
+	        }
+	    }
+	    else if (bClampEnds)
+	    {
+	        // For clamped ends, same approach as chord length but with square root distances
+	        
+	        // First set all knots to 0
+	        for (int32 i = 0; i < KnotCount; ++i)
+	        {
+	            KnotVector[i] = 0.0f;
+	        }
+	        
+	        // Set the last (Degree+1) knots to 1.0
+	        for (int32 i = KnotCount - Degree - 1; i < KnotCount; ++i)
+	        {
+	            KnotVector[i] = 1.0f;
+	        }
+	        
+	        // Calculate internal knots based on accumulated centripetal parameters
+	        const int32 NumInternalKnots = KnotCount - 2 * (Degree + 1);
+	        if (NumInternalKnots > 0)
+	        {
+	            // Compute normalized parameters at each control point
+	            TArray<float> Params;
+	            Params.SetNum(Points.Num());
+	            Params[0] = 0.0f;
+	            
+	            float AccumulatedLength = 0.0f;
+	            for (int32 i = 1; i < Points.Num(); ++i)
+	            {
+	                AccumulatedLength += Lengths[i-1];
+	                Params[i] = (TotalLength > UE_SMALL_NUMBER) ? 
+	                    (AccumulatedLength / TotalLength) : 
+	                    (static_cast<float>(i) / static_cast<float>(Points.Num() - 1));
+	            }
+	            
+	            // Map these parameters to internal knots
+	            for (int32 i = 0; i < NumInternalKnots; ++i)
+	            {
+	                // Map i to appropriate parameter range
+	                const float t = static_cast<float>(i + 1) / static_cast<float>(NumInternalKnots + 1);
+	                // Map t to control point indices
+	                const int32 LowIndex = FMath::FloorToInt(t * (Points.Num() - 1));
+	                const int32 HighIndex = FMath::Min(LowIndex + 1, Points.Num() - 1);
+	                const float Alpha = t * (Points.Num() - 1) - LowIndex;
+	                
+	                // Interpolate parameters
+	                const float KnotValue = FMath::Lerp(Params[LowIndex], Params[HighIndex], Alpha);
+	                KnotVector[Degree + 1 + i] = KnotValue;
+	            }
+	        }
+	    }
+	    else
+	    {
+	        // Unclamped knots
+	        for (int32 i = 0; i < KnotCount; ++i)
+	        {
+	            const float Param = static_cast<float>(i) / static_cast<float>(KnotCount - 1);
+	            KnotVector[i] = Param * TotalLength;
+	        }
+	    }
 	}
 
 };
@@ -745,44 +925,34 @@ public:
 	}
 
 	/**
-	 * Adds a new value to the spline and updates knot vector accordingly
-	 * @param Value - The control point value to add
-	 * @param Parameter - Parameter value for the point
-	 * @return Index of the added value
+	 * Reparameterizes the spline using a new knot generation mode
+	 * @param Mode - New knot generation mode
 	 */
-	virtual int32 AddValue(const ValueType& Value, float Parameter) override
+	void Reparameterize(EParameterizationPolicy Mode = EParameterizationPolicy::Uniform)
 	{
-		UE_LOG(LogSpline, Verbose, TEXT("TBSpline::AddValue - Evaluator: %p, KnotMode: %d"), 
-			&this->GetEvaluator(), static_cast<int32>(this->GetEvaluator().GetKnotMode()));
-		
-		const int32 Index = Base::AddValue(Value, Parameter);
-		// Update knot vector unless using custom knots
-		if (Base::GetEvaluator().GetKnotMode() != EBSplineKnotMode::Custom)
-		{
-			Base::GetEvaluator().UpdateKnotVector(Base::Values.GetStorage().GetEntries());
-		}
-		return Index;
-	}
+		Base::GetEvaluator().UpdateKnotVector(Base::Values.GetStorage().GetEntries(), Mode);
 
-	/**
-	 * Updates an existing value at the specified index and updates knot vector accordingly
-	 * @param Index - Index of the value to update
-	 * @param NewValue - New value to set
-	 * @return true if the index was valid and value was updated
-	 */
-	virtual bool UpdateValue(int32 Index, const ValueType& NewValue) override
-	{
-		if (!Base::UpdateValue(Index, NewValue))
-		{
-			return false;
-		}
+		// Update the parameter range to match the new knot range
+    	float KnotStart, KnotEnd;
+    	this->GetEvaluator().GetKnotRange(KnotStart, KnotEnd);
     
-		// Update knot vector if needed
-		if (Base::GetEvaluator().GetKnotMode() != EBSplineKnotMode::Custom)
-		{
-			Base::GetEvaluator().UpdateKnotVector(Base::Values.GetStorage().GetEntries());
-		}
-		return true;
+    	// Only update if the knot range is valid
+    	if (KnotEnd > KnotStart)
+    	{
+    		float CurrentStart, CurrentEnd;
+    		this->GetEvaluator().GetParentSpaceRange(CurrentStart, CurrentEnd);
+        
+    		// Check if the knot range is significantly different from parameter range
+    		const float Tolerance = 0.01f; // Small tolerance to avoid frequent small updates
+    		if (FMath::Abs(KnotStart - CurrentStart) > Tolerance || 
+    			FMath::Abs(KnotEnd - CurrentEnd) > Tolerance)
+    		{
+    			// Update parameter range to match knot range
+    			this->GetEvaluator().SetParentSpaceRange(KnotStart, KnotEnd);
+    			UE_LOG(LogSpline, Verbose, TEXT("Updated parameter range to [%f,%f] to match knot range"),
+    				KnotStart, KnotEnd);
+    		}
+    	}
 	}
 
 	virtual float FindNearest(const ValueType& Point, float& OutSquaredDistance) const override
@@ -833,15 +1003,6 @@ public:
 			return NumPoints - Degree;
 		}
 	}
-	/**
-	 * Changes how the knot vector is generated
-	 * @param Mode - The new knot generation mode
-	 */
-	void SetKnotGenerationMode(EBSplineKnotMode Mode)
-	{
-		Base::GetEvaluator().SetKnotGenerationMode(Mode);
-		Base::GetEvaluator().UpdateKnotVector(Base::Values.GetStorage().GetEntries());
-	}
 
 	/**
 	 * Sets whether the spline should interpolate its endpoints
@@ -850,16 +1011,6 @@ public:
 	void SetClampedEnds(bool bInClampEnds)
 	{
 		Base::GetEvaluator().SetClampedEnds(bInClampEnds);
-		Base::GetEvaluator().UpdateKnotVector(Base::Values.GetStorage().GetEntries());
-	}
-
-	virtual void SetClosedLoop(bool bInClosedLoop) override
-	{
-		if (this->IsClosedLoop() != bInClosedLoop)
-		{
-			this->GetEvaluator().SetClosedLoop(bInClosedLoop);
-			this->GetEvaluator().UpdateKnotVector(Base::Values.GetStorage().GetEntries());
-		}
 	}
 };
 
diff --git a/Engine/Source/Runtime/GeometryCore/Public/Splines/BaseEvaluator.h b/Engine/Source/Runtime/GeometryCore/Public/Splines/BaseEvaluator.h
--- a/Engine/Source/Runtime/GeometryCore/Public/Splines/BaseEvaluator.h
+++ b/Engine/Source/Runtime/GeometryCore/Public/Splines/BaseEvaluator.h
@@ -191,11 +191,29 @@ public:
 		return Result;
 	}
 
-	virtual void SetClosedLoop(bool bInClosedLoop){ bClosedLoop = bInClosedLoop; };
+	// Virtual function to update parent reference if this is an attribute evaluator
+	virtual bool UpdateParentEvaluator(const void* NewParentEvaluator) { return false; }
+	
+	virtual void SetClosedLoop(bool bInClosedLoop)
+	{
+		// Only update if state is changing
+		if (bClosedLoop != bInClosedLoop)
+		{
+			bClosedLoop = bInClosedLoop; 
+        
+			// Update infinity modes for closed loops to ensure proper parameter cycling
+			if (bClosedLoop)
+			{
+				PreInfinityMode = EOutOfBoundsHandlingMode::Cycle;
+				PostInfinityMode = EOutOfBoundsHandlingMode::Cycle;
+			}
+			// Don't revert modes when opening the loop - let user control them explicitly if needed
+		}
+	}
 	bool IsClosedLoop() const { return bClosedLoop; }
 
 	/** Set the mapping range in parent space */
-	void SetParentSpaceRange(float InStart, float InEnd)
+	virtual void SetParentSpaceRange(float InStart, float InEnd)
 	{
 		ParentStart = InStart;
 		ParentEnd = InEnd;
diff --git a/Engine/Source/Runtime/GeometryCore/Public/Splines/BuiltInAttributeTypes.h b/Engine/Source/Runtime/GeometryCore/Public/Splines/BuiltInAttributeTypes.h
new file mode 100644
--- /dev/null
+++ b/Engine/Source/Runtime/GeometryCore/Public/Splines/BuiltInAttributeTypes.h
@@ -0,0 +1,100 @@
+// Copyright Epic Games, Inc. All Rights Reserved.
+
+#pragma once
+
+#include "CoreMinimal.h"
+#include "Templates/Function.h"
+#include "AttributeChannel.h"
+
+namespace UE
+{
+namespace Geometry
+{
+namespace Spline
+{
+
+// Specialize TAttributeTypeTraits for all built-in types
+//-------------------------------------------------------
+
+// Basic types
+template<> struct TAttributeTypeTraits<float> 
+{ 
+	static inline const FName Name = TEXT("Float");
+};
+
+template<> struct TAttributeTypeTraits<double> 
+{ 
+	static inline const FName Name = TEXT("Double");
+};
+
+template<> struct TAttributeTypeTraits<int32> 
+{ 
+	static inline const FName Name = TEXT("Int32");
+};
+
+template<> struct TAttributeTypeTraits<bool> 
+{ 
+	static inline const FName Name = TEXT("Bool");
+};
+
+// Math types
+template<> struct TAttributeTypeTraits<FVector> 
+{ 
+	static inline const FName Name = TEXT("Vector");
+};
+
+template<> struct TAttributeTypeTraits<FVector2D> 
+{ 
+	static inline const FName Name = TEXT("Vector2D");
+};
+
+template<> struct TAttributeTypeTraits<FVector4> 
+{ 
+	static inline const FName Name = TEXT("Vector4");
+};
+
+template<> struct TAttributeTypeTraits<FQuat> 
+{ 
+	static inline const FName Name = TEXT("Quat");
+};
+
+template<> struct TAttributeTypeTraits<FRotator> 
+{ 
+	static inline const FName Name = TEXT("Rotator");
+};
+
+template<> struct TAttributeTypeTraits<FTransform> 
+{ 
+	static inline const FName Name = TEXT("Transform");
+};
+
+// Color types
+template<> struct TAttributeTypeTraits<FLinearColor> 
+{ 
+	static inline const FName Name = TEXT("LinearColor");
+};
+
+template<> struct TAttributeTypeTraits<FColor> 
+{ 
+	static inline const FName Name = TEXT("Color");
+};
+
+// String types
+template<> struct TAttributeTypeTraits<FName> 
+{ 
+	static inline const FName Name = TEXT("Name");
+};
+
+template<> struct TAttributeTypeTraits<FString> 
+{ 
+	static inline const FName Name = TEXT("String");
+};
+
+template<> struct TAttributeTypeTraits<FText> 
+{ 
+	static inline const FName Name = TEXT("Text");
+};
+
+} // namespace Spline
+} // namespace Geometry
+} // namesp
\ No newline at end of file
diff --git a/Engine/Source/Runtime/GeometryCore/Public/Splines/ChannelTypeRegistry.h b/Engine/Source/Runtime/GeometryCore/Public/Splines/ChannelTypeRegistry.h
new file mode 100644
--- /dev/null
+++ b/Engine/Source/Runtime/GeometryCore/Public/Splines/ChannelTypeRegistry.h
@@ -0,0 +1,114 @@
+// Copyright Epic Games, Inc. All Rights Reserved.
+
+#pragma once
+
+#include "CoreMinimal.h"
+#include "BuiltInAttributeTypes.h"
+
+namespace UE
+{
+namespace Geometry
+{
+namespace Spline
+{
+
+// Forward declarations
+template<typename ValueType> class TAttributeChannel;
+class IAttributeChannelBase;
+
+/**
+ * Central registry for spline attribute channel types
+ */
+class FChannelTypeRegistry
+{
+protected:
+	/** Map of registered type names to creation factories */
+	static inline TMap<FName, TFunction<TUniquePtr<IAttributeChannelBase>()>> TypeFactories;
+    
+public:
+	/**
+	 * Registers a new attribute channel type
+	 * @tparam AttributeType The value type of the attribute channel
+	 * @param TypeName Unique name for this type
+	 */
+	template<typename AttributeType>
+	static void RegisterType()
+	{
+		const FName TypeName = TAttributeTypeTraits<AttributeType>::Name;
+        
+		// Register the creation factory
+		TypeFactories.Add(TypeName, []() { 
+			return MakeUnique<TAttributeChannel<AttributeType>>();
+		});
+
+		// Set the static type name for this channel type
+		TAttributeChannel<AttributeType>::SetTypeName(TypeName);
+	}
+
+	/**
+	 * Unregisters an attribute channel type
+	 * @param TypeName The name of the type to unregister
+	 */
+	static void UnregisterType(const FName& TypeName)
+	{
+		TypeFactories.Remove(TypeName);
+	}
+    
+	/**
+	 * Creates a new channel of the specified type
+	 * @param TypeName Type name to create
+	 * @return New channel instance, or nullptr if type not found
+	 */
+	static TUniquePtr<IAttributeChannelBase> CreateChannel(const FName& TypeName)
+	{
+		if (TFunction<TUniquePtr<IAttributeChannelBase>()>* Factory = TypeFactories.Find(TypeName))
+		{
+			return (*Factory)();
+		}
+		return nullptr;
+	}
+    
+	/**
+	 * Checks if a type is registered
+	 * @param TypeName Type to check
+	 * @return True if the type is registered
+	 */
+	static bool IsTypeRegistered(const FName& TypeName)
+	{
+		return TypeFactories.Contains(TypeName);
+	}
+        
+    /**
+     * Initializes all built-in types
+     * Call this during module/library initialization
+     */
+    static void InitializeBuiltInTypes()
+    {
+    	// Basic types
+    	RegisterType<float>();
+    	RegisterType<double>();
+    	RegisterType<int32>();
+    	RegisterType<bool>();
+    
+    	// Math types
+    	RegisterType<FVector2D>();
+    	RegisterType<FVector>();
+    	RegisterType<FVector4>();
+    	RegisterType<FQuat>();
+    	RegisterType<FRotator>();
+    	RegisterType<FTransform>();
+    
+    	// Color types
+    	RegisterType<FLinearColor>();
+    	RegisterType<FColor>();
+    
+    	// String types
+    	RegisterType<FName>();
+    	RegisterType<FString>();
+    	RegisterType<FText>();
+    }
+};
+
+} // namespace Spline
+} // namespace Geometry
+} // namespace UE
\ No newline at end of file
diff --git a/Engine/Source/Runtime/GeometryCore/Public/Splines/CustomAttributeExample.h b/Engine/Source/Runtime/GeometryCore/Public/Splines/CustomAttributeExample.h
new file mode 100644
--- /dev/null
+++ b/Engine/Source/Runtime/GeometryCore/Public/Splines/CustomAttributeExample.h
@@ -0,0 +1,155 @@
+// Copyright Epic Games, Inc. All Rights Reserved.
+
+#pragma once
+
+#include "CoreMinimal.h"
+#include "Splines/InterpolationPolicies.h"
+#include "Splines/ChannelTypeRegistry.h"
+#include "Splines/BuiltInAttributeTypes.h"
+#include "Splines/PolyBezierSpline.h"
+
+namespace UE
+{
+namespace Geometry
+{
+namespace Spline
+{
+
+/**
+ * Simple custom attribute type for splines
+ * Demonstrates the minimum requirements for a custom attribute type
+ */
+struct FFoo
+{
+    // A simple float value attribute
+    float Value = 0.0f;
+    
+    // Constructor
+    FFoo() = default;
+    
+    // Constructor with value
+    explicit FFoo(float InValue) : Value(InValue) {}
+    
+    // Serialization support
+    bool Serialize(FArchive& Ar)
+    {
+        Ar << Value;
+        return true;
+    }
+        
+    friend FArchive& operator<<(FArchive& Ar, FFoo& Foo)
+    {
+        Foo.Serialize(Ar);
+        return Ar;
+    }
+    
+};
+
+// Register the type with the attribute system by defining a specialization
+// of TAttributeTypeTraits with a Name field
+template<> struct TAttributeTypeTraits<FFoo> 
+{ 
+    static inline const FName Name = TEXT("Foo");
+};
+
+// Define a specialization of TSplineInterpolationPolicy to handle
+// interpolation for this type
+template<>
+class TSplineInterpolationPolicy<FFoo>
+{
+public:
+    // Basic parameter-based interpolation
+    static FFoo Interpolate(TArrayView<const FFoo* const> Window, float Parameter)
+    {
+        if (Window.Num() < 2)
+        {
+            return Window.Num() == 1 ? *Window[0] : FFoo();
+        }
+        
+        // Linear interpolation
+        return FFoo(FMath::Lerp(Window[0]->Value, Window[1]->Value, Parameter));
+    }
+    
+    // Window-based interpolation for spline implementations
+    static FFoo InterpolateWithBasis(TArrayView<const FFoo* const> Window, TArrayView<const float> Basis)
+    {
+        if (Window.Num() == 0 || Basis.Num() == 0)
+        {
+            return FFoo();
+        }
+        
+        // Weighted sum of values
+        float Result = 0.0f;
+        
+        for (int32 i = 0; i < FMath::Min(Window.Num(), Basis.Num()); ++i)
+        {
+            Result += Window[i]->Value * Basis[i];
+        }
+        
+        return FFoo(Result);
+    }
+    
+    // Template-based n-th derivative calculation
+    template<int32 Order>
+    static FFoo EvaluateDerivative(TArrayView<const FFoo* const> Window, float Parameter)
+    {
+        if (Order == 0)
+        {
+            return Interpolate(Window, Parameter);
+        }
+        
+        // For first derivative
+        if (Order == 1 && Window.Num() >= 2)
+        {
+            // Basic first derivative for linear interpolation
+            return FFoo((Window[1]->Value - Window[0]->Value));
+        }
+        
+        // Higher order derivatives (could be more sophisticated)
+        return FFoo(0.0f);
+    }
+};
+
+// Example function demonstrating the FFoo type with a PolyBezier spline (this would typically go in a .cpp file)
+void CreatePolyBezierWithFooAttributes()
+{
+	// This is optional though, as the type will be registered when it is first used via Spline.CreateAttributeChannel<FFoo>()
+	FChannelTypeRegistry::RegisterType<FFoo>();
+
+	// Create a PolyBezier spline with a single segment
+	FPolyBezierSpline3d Spline = FPolyBezierSpline3d::CreateLine(
+		FVector(0, 0, 0),
+		FVector(100, 100, 0)
+	);
+
+	// Create an attribute channel of our custom FFoo type
+	Spline.CreateAttributeChannel<FFoo>("FooValues");
+
+	// Add some FFoo values at specific parameters
+	Spline.AddAttributeValue<FFoo>("FooValues", FFoo(0.0f), 0.0f);   // Start point
+	Spline.AddAttributeValue<FFoo>("FooValues", FFoo(1.0f), 1.0f);   // End point
+
+	// Add another segment
+	Spline.AppendBezierSegment(
+		FVector(150, 50, 0),   // First control point
+		FVector(150, -50, 0),  // Second control point 
+		FVector(200, 0, 0)     // End point
+	);
+
+	// Add FFoo value for the end of the second segment
+	Spline.AddAttributeValue<FFoo>("FooValues", FFoo(0.5f), 2.0f);
+
+	// Evaluate at various points
+	for (float t = 0.0f; t <= 1.0f; t += 0.1f)
+	{
+		FVector Position = Spline.Evaluate(t);
+		FFoo Foo = Spline.EvaluateAttribute<FFoo>("FooValues", t);
+    
+		UE_LOG(LogTemp, Display, TEXT("At t=%.1f: Position=(%.1f, %.1f, %.1f), Foo=%.2f"),
+			t, Position.X, Position.Y, Position.Z, Foo.Value);
+	}
+}
+
+} // end namespace UE::Geometry::Spline
+} // end namespace UE::Geometry
+} // end namespace UE
diff --git a/Engine/Source/Runtime/GeometryCore/Public/Splines/InterpolationPolicies.h b/Engine/Source/Runtime/GeometryCore/Public/Splines/InterpolationPolicies.h
--- a/Engine/Source/Runtime/GeometryCore/Public/Splines/InterpolationPolicies.h
+++ b/Engine/Source/Runtime/GeometryCore/Public/Splines/InterpolationPolicies.h
@@ -11,6 +11,17 @@ namespace Geometry
 {
 namespace Spline
 {
+	
+// Trait to detect if type T supports math operations needed for interpolation.
+template<typename T, typename = void>
+struct TIsMathInterpolable : std::false_type {};
+
+template<typename T>
+struct TIsMathInterpolable<T, std::void_t<
+	decltype(FMath::Lerp(std::declval<T>(), std::declval<T>(), 0.5f)),
+	decltype(std::declval<T>() + std::declval<T>()),
+	decltype(std::declval<T>() * std::declval<float>())
+	>> : std::true_type {};
 /**
  * Base spline interpolation policy that works across all spline types
  * Extends the existing parameter-based interpolation with window-based methods
@@ -26,18 +37,36 @@ public:
 		{
 			return Window.Num() == 1 ? *Window[0] : T();
 		}
-		return FMath::Lerp(*Window[0], *Window[1], Parameter);
+		if constexpr (TIsMathInterpolable<T>::value)
+		{
+			return FMath::Lerp(*Window[0], *Window[1], Parameter);
+		}
+		else
+		{
+			// Fallback for non-math types:
+			return Parameter < 0.5f ? *Window[0] : *Window[1];
+		}
 	}
 
 	// Add window-based interpolation for spline implementations
 	static T InterpolateWithBasis(TArrayView<const T* const> Window, TArrayView<const float> Basis)
 	{
-		T Result = T();
-		for (int32 i = 0; i < FMath::Min(Window.Num(), Basis.Num()); ++i)
+		if constexpr (TIsMathInterpolable<T>::value)
+		{
+			T Result = T();  // Assumes T() gives a neutral value (e.g. zero vector)
+			int32 Count = FMath::Min(Window.Num(), Basis.Num());
+			for (int32 i = 0; i < Count; ++i)
+			{
+				// Assumes T supports addition and multiplication by float.
+				Result = Result + (*Window[i] * Basis[i]);
+			}
+			return Result;
+		}
+		else
 		{
-			Result += *Window[i] * Basis[i];
+			// Fallback: simply return the first element if available.
+			return Window.Num() > 0 ? *Window[0] : T();
 		}
-		return Result;
 	}
 
 	// Template-based n-th derivative calculation
@@ -48,10 +77,67 @@ public:
 		{
 			return Interpolate(Window, Parameter);
 		}
-		return Math::TGenericDerivativeHelper<T, Order>::Compute(Window, Parameter);
+		if constexpr (TIsMathInterpolable<T>::value)
+		{
+			// For math types, use the generic derivative helper.
+			return Math::TGenericDerivativeHelper<T, Order>::Compute(Window, Parameter);
+		}
+		else
+		{
+			// For non-math types, derivatives are not meaningful.
+			// Return a default constructed value.
+			return T();
+		}
 	}
 };
 
+// Specialization only for int32
+template<>
+class TSplineInterpolationPolicy<int32>
+{
+public:
+	static int32 Interpolate(TArrayView<const int32* const> Window, float Parameter)
+    {
+        if (Window.Num() < 2)
+        {
+            return Window.Num() == 1 ? *Window[0] : 0;
+        }
+        
+        // Handle boundary cases
+        if (Parameter <= 0.0f) return *Window[0];
+        if (Parameter >= 1.0f) return *Window[1];
+        
+        // For intermediate values, explicitly convert to avoid implicit conversion warnings
+        const float Start = static_cast<float>(*Window[0]);
+        const float End = static_cast<float>(*Window[1]);
+        const float Result = Start + Parameter * (End - Start);
+        
+        // Round to nearest integer
+        return FMath::RoundToInt(Result);
+    }
+    
+    static int32 InterpolateWithBasis(TArrayView<const int32* const> Window, TArrayView<const float> Basis)
+    {
+        float Result = 0.0f;
+        for (int32 i = 0; i < FMath::Min(Window.Num(), Basis.Num()); ++i)
+        {
+            Result += static_cast<float>(*Window[i]) * Basis[i];
+        }
+        return FMath::RoundToInt(Result);
+    }
+
+	// Template-based n-th derivative calculation
+	template<int32 Order>
+	static int32 EvaluateDerivative(TArrayView<const int32* const> Window, float Parameter)
+	{
+		if (Order == 0)
+		{
+			return Interpolate(Window, Parameter);
+		}
+		return Math::TGenericDerivativeHelper<int32, Order>::Compute(Window, Parameter);
+	}
+};
+	
 // FLinearColor interpolation policy
 template<>
 class TSplineInterpolationPolicy<FLinearColor>
diff --git a/Engine/Source/Runtime/GeometryCore/Public/Splines/PolyBezierSpline.h b/Engine/Source/Runtime/GeometryCore/Public/Splines/PolyBezierSpline.h
--- a/Engine/Source/Runtime/GeometryCore/Public/Splines/PolyBezierSpline.h
+++ b/Engine/Source/Runtime/GeometryCore/Public/Splines/PolyBezierSpline.h
@@ -15,69 +15,53 @@ namespace Spline
  */
 template<typename T, int32 D>
 class TPolyBezierEvaluator : public TBSplineEvaluator<T, D, 3>
-{
+{	
 	using ValueType = typename TChooseVectorType<T, D>::Type;
 	using Base = TBSplineEvaluator<T, D, 3>;
 
-private:
-    int32 NumSegments = 0;
     
 public:
     TPolyBezierEvaluator():
         Base()
     {
-        this->SetKnotGenerationMode(EBSplineKnotMode::Custom);
         this->SetClampedEnds(true);
     }
 
-    bool UpdateSegmentStructure(int32 InNumSegments, int32 InNumValues)
+	/**
+	 * Reparameterizes the spline based on the provided points and mode
+	 * @param Mode - New parameterization mode
+	 * @param Points - Array of control points to use for reparameterization
+	 * @return 
+	 */
+	bool Reparameterize(EParameterizationPolicy Mode, TArrayView<const TStorageEntry<ValueType>> Points)
     {
-        NumSegments = InNumSegments;
-
-        // For Bezier curves:
-        // - Each segment has degree 3
-        // - For closed loop, we need n + 2*p + 1 knots where:
-        //   n = number of control points
-        //   p = degree (3 for cubic Bezier)
-        const int32 NumKnots = this->IsClosedLoop() ? 
-            (InNumValues + (2 * 3) + 1) :  // closed loop: n + 2p + 1
-            (InNumValues + 4);         // normal case: n + p + 1 where p = 3
-        
-        // Generate appropriate knot vector for PolyBezier
-        TArray<float> BezierKnots;
-        BezierKnots.SetNum(NumKnots);
-
-        // Set up knot vector with multiple knots at segment boundaries
-        int32 CurrentKnot = 0;
-
-        // Start knots with multiplicity 4
-        for (int32 i = 0; i < 4; ++i)
-        {
-            BezierKnots[CurrentKnot++] = 0.0f;
-        }
-
-        // For closed loop, add an extra segment to connect back to start
-        const int32 NumSegmentsToProcess = this->IsClosedLoop() ? (NumSegments + 1) : NumSegments;
-
-        // Interior knots at segment boundaries with multiplicity 3
-        for (int32 i = 1; i < NumSegmentsToProcess; ++i)
-        {
-            const float t = static_cast<float>(i);
-            for (int32 j = 0; j < 3; ++j)
-            {
-                BezierKnots[CurrentKnot++] = t;
-            }
-        }
+        const int32 NumValues = Points.Num();
 
-        // End knots with multiplicity 4
-        const float EndKnot = static_cast<float>(NumSegmentsToProcess);
-        while (CurrentKnot < NumKnots)
+        const int32 NumKnots = this->GetExpectedNumKnots(NumValues);
+                
+        switch (Mode)
         {
-            BezierKnots[CurrentKnot++] = EndKnot;
+            case EParameterizationPolicy::ChordLength:
+            case EParameterizationPolicy::Centripetal:
+        	    GenerateDistanceKnotsForBezier(Mode, Points, NumKnots);
+                break;
+            
+            case EParameterizationPolicy::Uniform:
+			default:
+                GenerateBezierKnots(NumValues);
+                break;
+                
         }
-
-        return Base::SetCustomKnots(BezierKnots, InNumValues);
+    	
+    	UE_LOG(LogSpline, Verbose, TEXT("Set Custom knot vector - Points: %d, Knots: %d, Mode: %d"),
+					   NumValues,  this->KnotVector.Num(), static_cast<int32>(Mode));
+    	for (int32 i = 0; i < this->KnotVector.Num(); ++i)
+    	{
+    		UE_LOG(LogSpline, Verbose, TEXT("  Knot[%d] = %f"), i, this->KnotVector[i]);
+    	}
+        return true;
     }
+    
     /**
      * Serializes evaluator data to/from an archive
      * @param Ar - Archive to serialize to/from
@@ -86,9 +70,6 @@ public:
     {
         // Serialize base evaluator first
         Base::Serialize(Ar);
-        
-        // Serialize PolyBezier-specific data
-        Ar << NumSegments;
     }
     
     virtual TUniquePtr<TBaseEvaluator<ValueType>> Clone() const override
@@ -110,6 +91,7 @@ public:
         const float ParentParam = this->MapToParentSpace(Parameter);
         
         // Find which segment we're in
+    	const int32 NumSegments = Entries.Num() / 4;
         const float ClampedParam = FMath::Clamp(ParentParam, 0.0f, NumSegments);
         int32 SegmentIndex = FMath::Min(FMath::FloorToInt(ClampedParam), static_cast<int32>(NumSegments) - 1);
 
@@ -163,7 +145,9 @@ public:
 
         // Map to parent space (segment space)
         const float ParentParam = this->MapToParentSpace(Parameter);
-        
+
+    	const int32 NumValues = this->GetNumValuesFromKnots();
+    	int32 NumSegments = NumValues / 4;
         // If we're at the last point exactly, just return it
         if (ParentParam >= static_cast<float>(NumSegments))
         {
@@ -185,7 +169,7 @@ public:
             OneMinusT * OneMinusT2,     // (1-t)³
             3.0f * t * OneMinusT2,      // 3t(1-t)²
             3.0f * t2 * OneMinusT,      // 3t²(1-t)
-            t * t2                       // t³
+            t * t2                      // t³
         };
 
         return TSplineInterpolationPolicy<ValueType>::InterpolateWithBasis(Window, Basis);
@@ -218,7 +202,8 @@ public:
 	    {
 	        return ValueType();
 	    }
-	    
+
+    	int32 NumSegments = Entries.Num() / 4;
 	    // Find segment index and local parameter
 	    const float ClampedParam = FMath::Clamp(MappedParam, 0.0f, static_cast<float>(NumSegments));
 	    int32 SegmentIndex;
@@ -264,6 +249,7 @@ public:
     	const float MappedParam = Base::MapToParentSpace(Parameter);
     
     	// Find segment index and local parameter
+    	int32 NumSegments = Entries.Num() / 4;
     	const float ClampedParam = FMath::Clamp(MappedParam, 0.0f, static_cast<float>(NumSegments));
     	const int32 SegmentIndex = FMath::Min(FMath::FloorToInt(ClampedParam), NumSegments - 1);
     	const float LocalT = ClampedParam - static_cast<float>(SegmentIndex); // Parameter [0,1] within segment
@@ -288,17 +274,215 @@ public:
     	// Compute the derivative using B-spline methods
     	return Math::TBSplineDerivativeCalculator<ValueType, Order>::Compute(
 			Window, KnotVector, Span, LocalT, 3);
+    }
+	/**
+	 * Override SetParentSpaceRange to update segment structure accordingly
+	 * @param InStart - New start parameter
+	 * @param InEnd - New end parameter
+	 */
+	virtual void SetParentSpaceRange(float InStart, float InEnd) override
+    {
+    	float OldStart = this->ParentStart;
+    	float OldEnd = this->ParentEnd;
+    
+    	// Only proceed if there's an actual change
+    	if (FMath::IsNearlyEqual(OldStart, InStart) && FMath::IsNearlyEqual(OldEnd, InEnd))
+    	{
+    		return;
+    	}
+    
+    	// Call base implementation to update the parameter range
+    	Base::SetParentSpaceRange(InStart, InEnd);
+    
     }
 	
     virtual float GetParentSpaceMapping(float NormalizedParam) const override
     {
+    	int32 NumValues = this->GetNumValuesFromKnots();
+        int32 NumSegments = NumValues / 4;
         return NormalizedParam * static_cast<float>(NumSegments);
     }
 
     virtual float GetNormalizedSpaceMapping(float ParentParam) const override
     {
+    	int32 NumValues = this->GetNumValuesFromKnots();
+    	int32 NumSegments = NumValues / 4;
         return NumSegments > 0 ? ParentParam / static_cast<float>(NumSegments) : 0.0f;
     }
+
+	/**
+	 * Helper to insert local knots for a segment boundary.
+	 * For clamped splines, internal boundaries typically should have a multiplicity equal
+	 * to the degree rather than the clamped endpoint multiplicity.
+	 * @param StartParam - Start parameter of the segment
+	 * @param EndParam - End parameter of the segment
+	 */
+	void InsertLocalKnots(float StartParam, float EndParam)
+	{
+		if (this->bClampEnds)
+		{
+			// Desired multiplicity for an internal knot might be Degree
+			const int32 DesiredInternalMultiplicity = GetDegree();
+
+			int32 StartMult = this->GetKnotMultiplicity(StartParam);
+			if (StartMult < DesiredInternalMultiplicity)
+			{
+				this->InsertKnot(StartParam, DesiredInternalMultiplicity - StartMult);
+			}
+
+			int32 EndMult = this->GetKnotMultiplicity(EndParam);
+			if (EndMult < DesiredInternalMultiplicity)
+			{
+				this->InsertKnot(EndParam, DesiredInternalMultiplicity - EndMult);
+			}
+		}
+		else
+		{
+			// For unclamped splines, simply insert one knot at each boundary.
+			this->InsertKnot(StartParam, 1);
+			this->InsertKnot(EndParam, 1);
+		}
+	}
+private:
+	
+    /**
+     * Computes chord or centripetal distances between segment endpoints (P0 to P3).
+     * Uses those distances to build a proportional knot vector with the proper multiplicities for each Bézier segment.
+     * @param Mode - Knot generation mode to use
+     * @param Points - Control points for distance-based generation
+     * @param NumKnots - Total number of knots needed
+     */
+    void GenerateDistanceKnotsForBezier(EParameterizationPolicy Mode, TArrayView<const TStorageEntry<ValueType>> Points, int32 NumKnots)
+    {
+        // For a Bezier spline, the control points don't all lie on the curve
+        // So we need to compute chord lengths between segment endpoints only
+    	int NumSegments = Points.Num() / 4;
+    	
+        TArray<float> SegmentLengths;
+        SegmentLengths.Reserve(NumSegments);
+        float TotalLength = 0.0f;
+        
+        for (int32 i = 0; i < NumSegments; ++i)
+        {
+            const int32 StartIndex = i * 4;
+            const int32 EndIndex = StartIndex + 3;
+            
+            if (EndIndex < Points.Num())
+            {
+                float Length;
+                
+                // Use appropriate distance calculation based on knot mode
+                if (Mode == EParameterizationPolicy::Centripetal)
+                {
+                    // Square root of distance for centripetal mode
+                    Length = FMath::Sqrt(static_cast<float>((Points[EndIndex].Value - Points[StartIndex].Value).Size()));
+                }
+                else // Chord length mode
+                {
+                    // Actual distance for chord length mode
+                    Length = static_cast<float>((Points[EndIndex].Value - Points[StartIndex].Value).Size());
+                }
+                
+                SegmentLengths.Add(Length);
+                TotalLength += Length;
+            }
+        }
+        
+        // Now create a knot vector that respects:
+        // 1. Bezier segment structure (knot multiplicity)
+        // 2. Proportional chord/centripetal lengths
+        
+        TArray<float> DistanceKnots;
+        DistanceKnots.SetNum(NumKnots);
+        
+        int32 CurrentKnot = 0;
+
+        // Start knots with multiplicity 4
+        for (int32 i = 0; i < 4; ++i)
+        {
+            DistanceKnots[CurrentKnot++] = 0.0f;
+        }
+        
+        // Calculate accumulated distances for internal knots
+        float AccumLength = 0.0f;
+        for (int32 i = 0; i < NumSegments - 1; ++i) // -1 because we don't need knots at the very end
+        {
+            if (i < SegmentLengths.Num())
+            {
+                AccumLength += SegmentLengths[i];
+                
+                // Use normalized accumulated length as the knot value
+                const float t = TotalLength > UE_SMALL_NUMBER ? 
+                    (AccumLength / TotalLength) * static_cast<float>(NumSegments) : // Scale to match segment count
+                    static_cast<float>(i + 1);										// Fallback to uniform if total length ~ 0
+                
+                // Interior knots at segment boundaries with multiplicity 3
+                for (int32 j = 0; j < 3; ++j)
+                {
+                    DistanceKnots[CurrentKnot++] = t;
+                }
+            }
+        }
+
+        // End knots with multiplicity 4
+        const float EndKnot = static_cast<float>(NumSegments);
+        while (CurrentKnot < NumKnots)
+        {
+            DistanceKnots[CurrentKnot++] = EndKnot;
+        }
+
+        // Set the knot vector
+        this->KnotVector = DistanceKnots;
+    }
+
+    /**
+     * Generates a standard Bezier-compatible knot vector
+     */
+    void GenerateBezierKnots(int32 NumValues)
+    {
+    	int32 NumSegments = NumValues / 4;
+        // For Bezier curves:
+        // - Each segment has degree 3
+        // - For closed loop, we need n + 2*p + 1 knots where:
+        //   n = number of control points
+        //   p = degree (3 for cubic Bezier)
+        const int32 NumKnots = this->GetExpectedNumKnots(NumValues);             // normal case: n + p + 1 where p = 3
+        
+        // Set up knot vector with multiple knots at segment boundaries
+        TArray<float> BezierKnots;
+        BezierKnots.SetNum(NumKnots);
+        
+        int32 CurrentKnot = 0;
+
+        // Start knots with multiplicity 4
+        for (int32 i = 0; i < 4; ++i)
+        {
+            BezierKnots[CurrentKnot++] = 0.0f;
+        }
+
+        // For closed loop, add an extra segment to connect back to start
+        const int32 NumSegmentsToProcess = this->IsClosedLoop() ? (NumSegments + 1) : NumSegments;
+
+        // Interior knots at segment boundaries with multiplicity 3
+        for (int32 i = 1; i < NumSegmentsToProcess; ++i)
+        {
+            const float t = static_cast<float>(i);
+            for (int32 j = 0; j < 3; ++j)
+            {
+                BezierKnots[CurrentKnot++] = t;
+            }
+        }
+
+        // End knots with multiplicity 4
+        const float EndKnot = static_cast<float>(NumSegmentsToProcess);
+        while (CurrentKnot < NumKnots)
+        {
+            BezierKnots[CurrentKnot++] = EndKnot;
+        }
+
+        // Set the knot vector
+        this->KnotVector = BezierKnots;
+    }
 };
 
 /**
@@ -313,86 +497,682 @@ public:
     using FEvaluator = TPolyBezierEvaluator<ScalarType, Dimensions>;
     using Base = TSplineBase<ScalarType, Dimensions, TPolyBezierEvaluator<ScalarType, Dimensions>>;
 
-    TPolyBezierSpline()
-        : Base(MakeUnique<TPolyBezierEvaluator<ScalarType, Dimensions>>())
-    {
-    }
+	/**
+	 * Primary constructor that requires a cubic Bezier segment 
+	 */
+	TPolyBezierSpline(
+		const ValueType& P0, 
+		const ValueType& P1, 
+		const ValueType& P2, 
+		const ValueType& P3,
+		EParameterizationPolicy Parameterization = EParameterizationPolicy::Uniform)
+	: Base(MakeUnique<TPolyBezierEvaluator<ScalarType, Dimensions>>())
+	{
+		// Add the initial segment
+		AddBezierSegmentInternal(P0, P1, P2, P3, 0, Parameterization);
+		Reparameterize(Parameterization);
+	}
     
     // Virtual destructor
     virtual ~TPolyBezierSpline() = default;
 
-    virtual void SetClosedLoop(bool bInClosedLoop) override
+	/**
+	 * Factory method to create an empty spline with a default segment 
+	 */
+	static TPolyBezierSpline CreateDefault()
+	{
+		// Create a minimal valid default segment for serialization
+		ValueType P0(0, 0, 0);
+		ValueType P1(0.33f, 0, 0);
+		ValueType P2(0.66f, 0, 0);
+		ValueType P3(1, 0, 0);
+		
+		return TPolyBezierSpline(P0, P1, P2, P3, EParameterizationPolicy::Uniform);
+	}
+
+	// Static factory methods for common shapes    
+    // Create a straight line between two points
+    static TPolyBezierSpline<ScalarType, Dimensions> CreateLine(
+        const ValueType& Start, 
+        const ValueType& End)
     {
-        if (this->IsClosedLoop() != bInClosedLoop)
-        {
-            this->GetEvaluator().SetClosedLoop(bInClosedLoop);
-            UpdateSegmentStructure();
-        }
+        // Calculate control points for a straight line
+        // For a straight line, the control points should be evenly spaced
+        ValueType P0 = Start;
+        ValueType P3 = End;
+        ValueType Direction = (End - Start);
+        ValueType P1 = Start + Direction * 0.33f;
+        ValueType P2 = Start + Direction * 0.66f;
+        
+        TPolyBezierSpline<ScalarType, Dimensions> Result(P0, P1, P2, P3);
+        Result.Reparameterize(EParameterizationPolicy::Uniform);
+        return Result;
+    }
+    
+    // Create a circular arc with given center, radius, and angle range
+    static TPolyBezierSpline<ScalarType, Dimensions> CreateCircleArc(
+        const ValueType& Center,
+        float Radius,
+        float StartAngle,
+        float EndAngle,
+        int32 NumSegments = 4)
+    {
+           // Ensure angles are in the right order
+           if (EndAngle < StartAngle)
+           {
+               float Temp = StartAngle;
+               StartAngle = EndAngle;
+               EndAngle = Temp;
+           }
+           
+           // Ensure reasonable segment count
+           NumSegments = FMath::Max(1, NumSegments);
+           
+           // Calculate angle per segment
+           float AngleSpan = EndAngle - StartAngle;
+           float AnglePerSegment = AngleSpan / static_cast<float>(NumSegments);
+           
+           // For large angles per segment, increase segment count for better approximation
+           if (AnglePerSegment > PI/4)
+           {
+               NumSegments = FMath::CeilToInt(AngleSpan / (PI/4));
+               AnglePerSegment = AngleSpan / static_cast<float>(NumSegments);
+           }
+           
+           // Pre-compute all points along the arc for exact positioning
+           TArray<ValueType> Points;
+           Points.SetNum(NumSegments + 1);
+           
+           for (int32 i = 0; i <= NumSegments; i++)
+           {
+               ScalarType Angle = StartAngle + (static_cast<float>(i) * AnglePerSegment);
+               Points[i] = Center + ValueType(
+                   Radius * FMath::Cos(Angle),
+                   Radius * FMath::Sin(Angle),
+                   0);
+           }
+           
+           // Calculate first segment with precise tangents
+           ValueType P0 = Points[0];
+           ValueType P3 = Points[1];
+           
+           // Get normalized direction vectors from center to points
+           ValueType Dir0 = (P0 - Center).GetSafeNormal();
+           ValueType Dir3 = (P3 - Center).GetSafeNormal();
+           
+           // Calculate perpendicular vectors (normalized)
+           ValueType Perp0 = ValueType(-Dir0.Y, Dir0.X, 0);
+           ValueType Perp3 = ValueType(-Dir3.Y, Dir3.X, 0);
+           
+           // Calculate precise tangent scale
+           ScalarType TangentScale = Radius * (ScalarType(4.0)/ScalarType(3.0)) * 
+                                    FMath::Tan(AnglePerSegment / ScalarType(4.0));
+           
+           // Create control points with exact perpendicular tangents
+           ValueType P1 = P0 + Perp0 * TangentScale;
+           ValueType P2 = P3 - Perp3 * TangentScale;
+           
+           // Create spline with first segment
+           TPolyBezierSpline<ScalarType, Dimensions> Result(P0, P1, P2, P3);
+           
+           // Add remaining segments with carefully calculated control points
+           for (int32 i = 1; i < NumSegments; i++)
+           {
+               P0 = Points[i];
+               P3 = Points[i+1];
+               
+               // Calculate exact perpendicular vectors for this segment
+               Dir0 = (P0 - Center).GetSafeNormal();
+               Dir3 = (P3 - Center).GetSafeNormal();
+               Perp0 = ValueType(-Dir0.Y, Dir0.X, 0);
+               Perp3 = ValueType(-Dir3.Y, Dir3.X, 0);
+               
+               // Create control points with exact perpendicular tangents
+               P1 = P0 + Perp0 * TangentScale;
+               P2 = P3 - Perp3 * TangentScale;
+               
+               Result.AppendBezierSegment(P1, P2, P3);
+           }
+           
+           return Result;
+    }
+    
+    // Create a complete circle
+    static TPolyBezierSpline<ScalarType, Dimensions> CreateCircle(
+        const ValueType& Center,
+        float Radius,
+        int32 NumSegments = 4)
+    {
+		// Ensure reasonable segment count
+		NumSegments = FMath::Max(4, NumSegments);
+    
+		// Create a full circle arc (0 to 2π)
+		TPolyBezierSpline<ScalarType, Dimensions> Result = 
+			CreateCircleArc(Center, Radius, 0.0f, 2.0f * PI, NumSegments);
+    
+		// Set as closed loop
+		Result.SetClosedLoop(true);
+    
+		return Result;
     }
     
+    // Create an ellipse with X and Y radii
+    static TPolyBezierSpline<ScalarType, Dimensions> CreateEllipse(
+        const ValueType& Center,
+        float RadiusX,
+        float RadiusY,
+        int32 NumSegments = 4)
+    {
+        // Ensure at least 4 segments for good quality
+	    if (NumSegments < 4) NumSegments = 4;
+	    
+	    // Calculate angle per segment
+	    const float AnglePerSegment = 2.0f * PI / static_cast<float>(NumSegments);
+	    
+	    // Start with first point at angle 0
+	    float CurrentAngle = 0.0f;
+	    ValueType P0 = Center + ValueType(RadiusX * FMath::Cos(CurrentAngle), 
+	                                      RadiusY * FMath::Sin(CurrentAngle), 0);
+	    
+	    // Calculate next point
+	    CurrentAngle += AnglePerSegment;
+	    ValueType P3 = Center + ValueType(RadiusX * FMath::Cos(CurrentAngle), 
+	                                      RadiusY * FMath::Sin(CurrentAngle), 0);
+	    
+	    // Calculate correct tangent scale for this angle
+	    const float TangentScale = (4.0f/3.0f) * FMath::Tan(AnglePerSegment / 4.0f);
+	    
+	    // Calculate tangent vectors (scaled appropriately for ellipse)
+	    ValueType Tangent0 = ValueType(-RadiusX * FMath::Sin(0.0f), 
+	                                  RadiusY * FMath::Cos(0.0f), 0) * TangentScale;
+	    ValueType Tangent3 = ValueType(-RadiusX * FMath::Sin(CurrentAngle), 
+	                                  RadiusY * FMath::Cos(CurrentAngle), 0) * TangentScale;
+	    
+	    ValueType P1 = P0 + Tangent0;
+	    ValueType P2 = P3 - Tangent3;
+	    
+	    // Create the spline with the first segment
+	    TPolyBezierSpline<ScalarType, Dimensions> Result(P0, P1, P2, P3);
+	    
+	    // Add remaining segments
+	    for (int32 i = 1; i < NumSegments; ++i)
+	    {
+	        P0 = P3; // Start from previous endpoint
+	        Tangent0 = Tangent3; // Reuse previous tangent
+	        
+	        // Calculate next endpoint
+	        CurrentAngle += AnglePerSegment;
+	        P3 = Center + ValueType(RadiusX * FMath::Cos(CurrentAngle), 
+	                               RadiusY * FMath::Sin(CurrentAngle), 0);
+	        
+	        // Calculate tangent at next point
+	        Tangent3 = ValueType(-RadiusX * FMath::Sin(CurrentAngle), 
+	                            RadiusY * FMath::Cos(CurrentAngle), 0) * TangentScale;
+	        
+	        // Calculate control points
+	        P1 = P0 + Tangent0;
+	        P2 = P3 - Tangent3;
+	        
+	        // Add the segment
+	        Result.AppendBezierSegment(P1, P2, P3);
+	    }
+	    
+	    // Set as closed loop
+	    Result.SetClosedLoop(true);
+	            
+        // Apply consistent parameterization
+        Result.Reparameterize(EParameterizationPolicy::Centripetal);
+        
+        return Result;
+    }
+
+	/**
+	 * Replaces all control points in the spline with the provided points.
+	 * For a cubic Bezier spline, the number of points must be a multiple of 4 (or 4N-1 for open splines).
+	 * 
+	 * @param Points - Array of new control points
+	 * @param ParameterizationPolicy - How to distribute parameters for the new segments
+	 * @return True if the points were successfully set
+	 */
+	bool SetControlPoints(const TArray<ValueType>& Points, EParameterizationPolicy ParameterizationPolicy)
+	{
+	    // For a cubic Bezier spline, we need either:
+	    // - 4N points for closed loop (4 points per segment)
+	    // - 4N-1 points for open spline (last point is shared between segments)
+	    const int32 NumPoints = Points.Num();
+	    	    
+	    // For open splines with 4N-1 points, the last control point of one segment 
+	    // and the first control point of the next segment must be identical
+	    const bool bIsClosedLoop = this->IsClosedLoop();
+
+		// Need at least 2 points to add a valid segment
+		if (NumPoints < 4)
+		{
+			UE_LOG(LogSpline, Warning, TEXT("SetControlPoints requires at least 4 points to add a valid segment. Got %d points."), NumPoints);
+			return false;
+		}
+		
+	    // Check if points array is valid for the spline type
+	    if (bIsClosedLoop)
+	    {
+	        // For closed loops, we need a multiple of 4 points (complete segments)
+	        if (NumPoints % 4 != 0)
+	        {
+	            UE_LOG(LogSpline, Warning, TEXT("Invalid point count for closed loop PolyBezier. Expected multiple of 4, got %d."), NumPoints);
+	            return false;
+	        }
+	    }
+	    else
+	    {
+	        // For open splines, we need 4N-1 points: 4 for first segment, 3 for each additional segment
+	        // Equivalently: 1 + 3*NumSegments
+	        // So NumSegments = (NumPoints - 1) / 3
+	        // This means NumPoints - 1 must be divisible by 3
+	        if ((NumPoints - 1) % 3 != 0 && NumPoints > 1)
+	        {
+	            UE_LOG(LogSpline, Warning, TEXT("Invalid point count for open PolyBezier. Expected 4N-1 (1 + 3*NumSegments), got %d."), NumPoints);
+	            return false;
+	        }
+	    }
+	    
+	    // Clear existing points
+	    this->Clear();
+	    
+	    // Calculate number of segments
+	    int32 NumSegments;
+	    if (bIsClosedLoop)
+	    {
+	        NumSegments = NumPoints / 4;
+	    }
+	    else
+	    {
+	        // For open splines: NumPoints = 1 + 3*NumSegments, so NumSegments = (NumPoints - 1) / 3
+	        NumSegments = NumPoints > 0 ? (NumPoints - 1) / 3 + 1 : 0;
+	    }
+	    
+	    // Add all segments
+	    for (int32 SegmentIndex = 0; SegmentIndex < NumSegments; ++SegmentIndex)
+	    {
+	        const int32 BaseIndex = SegmentIndex * 4;
+	        
+	        // For each segment, we need 4 control points
+	        ValueType P0, P1, P2, P3;
+	        
+	        if (SegmentIndex == 0 || bIsClosedLoop)
+	        {
+	            // For first segment or any segment in closed loop, get all 4 points
+	            P0 = Points[BaseIndex % NumPoints];
+	            P1 = Points[(BaseIndex + 1) % NumPoints];
+	            P2 = Points[(BaseIndex + 2) % NumPoints];
+	            P3 = Points[(BaseIndex + 3) % NumPoints];
+	        }
+	        else
+	        {
+	            // For subsequent segments in open spline, P0 is shared with the last P3
+	            // and we only need 3 more points
+	            P0 = this->GetValueAt(this->GetNumKeys() - 1); // Last point of previous segment
+	            
+	            // Get next 3 control points (P1, P2, P3) from the input array
+	            // For segment i (i > 0), the points start at 1 + 3*(i-1)
+	            const int32 InputStartIndex = 1 + 3 * (SegmentIndex - 1);
+	            P1 = Points[InputStartIndex];
+	            P2 = Points[InputStartIndex + 1];
+	            P3 = Points[InputStartIndex + 2];
+	        }
+	        
+	        // Add the segment
+	        if (SegmentIndex == 0)
+	        {
+	            // For first segment, add all 4 points
+	            this->AddBezierSegmentInternal(P0, P1, P2, P3, 0, ParameterizationPolicy);
+	        }
+	        else
+	        {
+	            // For subsequent segments, use AppendBezierSegment which only needs P1, P2, P3
+	            this->AppendBezierSegment(P1, P2, P3, ParameterizationPolicy);
+	        }
+	    }
+	    
+	    // Set closed loop state and reparameterize
+	    this->SetClosedLoop(bIsClosedLoop);
+	    this->Reparameterize(ParameterizationPolicy);
+	    
+	    return true;
+	}
+
+	/**
+	 * Adds multiple Bezier segments to the spline.
+	 * 
+	 * @param Points - Array of control points defining the new segments
+	 * @param bAppend - If true, append to the end of spline; if false, prepend to the beginning
+	 * @param ParameterizationPolicy - How to distribute parameters for the new segments
+	 * @return True if the segments were successfully added
+	 */
+	bool AddBezierSegments(const TArray<ValueType>& Points, bool bAppend, EParameterizationPolicy ParameterizationPolicy)
+	{
+	    const int32 NumPoints = Points.Num();
+	    
+		// Need at least 2 points to add a valid segment
+		if (NumPoints < 4)
+		{
+			UE_LOG(LogSpline, Warning, TEXT("AddBezierSegments requires at least 4 points to add a valid segment. Got %d points."), NumPoints);
+			return false;
+		}
+	    
+	    // If the spline is currently empty (even though it shouldn't be according to design), this is equivalent to SetControlPoints
+	    if (this->GetNumKeys() == 0)
+	    {
+	        return SetControlPoints(Points, ParameterizationPolicy);
+	    }
+	    
+	    // Check points format based on append/prepend direction
+	    if (bAppend)
+	    {
+	        // For appending, we expect 3N+1 points: P0 (shared with last existing point), then (P1,P2,P3) per segment
+	        if (NumPoints % 3 != 1)
+	        {
+	            UE_LOG(LogSpline, Warning, TEXT("Invalid point count for appending segments. Expected 3N+1, got %d."), NumPoints);
+	            return false;
+	        }
+	    }
+	    else
+	    {
+	        // For prepending, we expect 3N+1 points: (P0,P1,P2) per segment, then P3 (shared with first existing point)
+	        if (NumPoints % 3 != 1)
+	        {
+	            UE_LOG(LogSpline, Warning, TEXT("Invalid point count for prepending segments. Expected 3N+1, got %d."), NumPoints);
+	            return false;
+	        }
+	    }
+	    
+	    // Check connection point
+	    if (bAppend)
+	    {
+	        // First point of new segments should match the last point of existing spline
+	        const ValueType& FirstNewPoint = Points[0];
+	        const ValueType& LastExistingPoint = this->GetValueAt(this->GetNumKeys() - 1);
+	        
+	        // Check if the connection point is reasonably close
+	        constexpr double ConnectionTolerance = 1e-4;
+	        const double DistanceSquared = (LastExistingPoint - FirstNewPoint).SizeSquared();
+	        if (DistanceSquared > ConnectionTolerance)
+	        {
+	            UE_LOG(LogSpline, Warning, TEXT("Connection point mismatch when appending segments. Distance: %f"),
+	                FMath::Sqrt(static_cast<float>(DistanceSquared)));
+	            // Continue anyway but log warning
+	        }
+	    }
+	    else
+	    {
+	        // Last point of new segments should match the first point of existing spline
+	        const ValueType& LastNewPoint = Points.Last();
+	        const ValueType& FirstExistingPoint = this->GetValueAt(0);
+	        
+	        // Check if the connection point is reasonably close
+	        constexpr double ConnectionTolerance = 1e-4;
+	        const double DistanceSquared = (LastNewPoint - FirstExistingPoint).SizeSquared();
+	        if (DistanceSquared > ConnectionTolerance)
+	        {
+	            UE_LOG(LogSpline, Warning, TEXT("Connection point mismatch when prepending segments. Distance: %f"),
+	                FMath::Sqrt(static_cast<float>(DistanceSquared)));
+	            // Continue anyway but log warning
+	        }
+	    }
+	    
+	    if (bAppend)
+	    {
+	        // Append segments to the end of the spline
+	        const int32 NumSegmentsToAdd = (NumPoints - 1) / 3;
+	        
+	        // Add each segment
+	        for (int32 SegmentIndex = 0; SegmentIndex < NumSegmentsToAdd; ++SegmentIndex)
+	        {
+	            const int32 BaseIndex = 1 + SegmentIndex * 3; // Skip the first point (P0) for first segment
+	            
+	            // For each segment, we need P1, P2, P3 (P0 is the last point of existing spline)
+	            const ValueType& P1 = Points[BaseIndex];
+	            const ValueType& P2 = Points[BaseIndex + 1];
+	            const ValueType& P3 = Points[BaseIndex + 2];
+	            
+	            // Use the existing AppendBezierSegment method
+	            this->AppendBezierSegment(P1, P2, P3, ParameterizationPolicy);
+	        }
+	    }
+	    else
+	    {
+	        // Prepend segments to the beginning of the spline
+	        const int32 NumSegmentsToPrepend = (NumPoints - 1) / 3;
+	        
+	        // Prepend segments in order starting from the first segment
+	        // For prepending, we need to work backwards since each PrependBezierSegment
+	        // adds to the beginning
+	        for (int32 i = NumSegmentsToPrepend - 1; i >= 0; --i)
+	        {
+	            const int32 BaseIndex = i * 3;
+	            
+	            // For each segment, we need P0, P1, P2 (P3 will connect to the existing spline)
+	            const ValueType& P0 = Points[BaseIndex];
+	            const ValueType& P1 = Points[BaseIndex + 1];
+	            const ValueType& P2 = Points[BaseIndex + 2];
+	            
+	            // Use the existing PrependBezierSegment method
+	            this->PrependBezierSegment(P0, P1, P2, ParameterizationPolicy);
+	        }
+	    }
+	    
+	    // Reparameterize the spline to ensure consistent parameterization
+	    this->Reparameterize(ParameterizationPolicy);
+	    
+	    return true;
+	}
     /**
-     * Adds a new Bezier segment to the spline
-     * @param P0 Start point
+     * Appends a Bezier segment using only 3 points (P1, P2, P3)
+     * The start point (P0) is automatically inferred from the last point of the previous segment.
+     * 
      * @param P1 First control point
      * @param P2 Second control point
      * @param P3 End point
-     * @param bEnforceContinuity If true, enforces that the segment connects to the previous one
+     * @param ParameterizationPolicy How to distribute parameters
+     * @return Index of the new segment
+     */
+    int32 AppendBezierSegment(
+        const ValueType& P1, 
+        const ValueType& P2, 
+        const ValueType& P3, 
+        EParameterizationPolicy ParameterizationPolicy = EParameterizationPolicy::Centripetal)
+    {
+		// Get P0 from the last point of the previous segment
+		ValueType P0 = Base::GetValueAt(Base::GetNumKeys() - 1);
+    
+		// Get current parameter range end
+		float CurrentStart, CurrentEnd;
+		this->GetEvaluator().GetParentSpaceRange(CurrentStart, CurrentEnd);
+    
+		// Use a parameter just past the current end to trigger append behavior
+		float AppendParam = CurrentEnd + 0.1f;
+    
+		// Add segment with append behavior
+		return AddBezierSegmentInternal(P0, P1, P2, P3, AppendParam, ParameterizationPolicy);
+    }
+
+    /**
+     * Prepends a Bezier segment using only 3 points (P0, P1, P2)
+     * The end point (P3) is automatically inferred from the first point of the existing spline.
+     * 
+     * @param P0 First control point
+     * @param P1 Second control point
+     * @param P2 Third control point
+     * @param ParameterizationPolicy How to distribute parameters
      * @return Index of the new segment
      */
-    int32 AddBezierSegment(const ValueType& P0, const ValueType& P1, const ValueType& P2, const ValueType& P3, bool bEnforceContinuity = true)
+    int32 PrependBezierSegment(
+		const ValueType& P0, 
+		const ValueType& P1, 
+		const ValueType& P2, 
+		EParameterizationPolicy ParameterizationPolicy = EParameterizationPolicy::Centripetal)
+	{
+		// Get P3 from the first point of the existing spline
+		ValueType P3 = Base::GetValueAt(0);
+    
+		// Get current parameter range start
+		float CurrentStart, CurrentEnd;
+		this->GetEvaluator().GetParentSpaceRange(CurrentStart, CurrentEnd);
+    
+		// Use a parameter just before the current start to trigger prepend behavior
+		float PrependParam = CurrentStart - 0.1f;
+    
+		// Add segment with prepend behavior
+		return AddBezierSegmentInternal(P0, P1, P2, P3, PrependParam, ParameterizationPolicy);
+	}
+
+    /**
+     * Inserts a new point at the specified parameter along the spline.
+     * Uses the specified parameterization policy to preserve attribute positions.
+     * 
+     * @param Parameter Parameter value [0,1] where to insert the point
+     * @param Position Position for the new point
+     * @param ParameterizationPolicy How to distribute parameters
+     * @return The index of the newly created point
+     */
+    int32 InsertPoint(
+        float Parameter, 
+        const ValueType& Position,
+        EParameterizationPolicy ParameterizationPolicy = EParameterizationPolicy::Centripetal)
     {
-        const int32 NumPoints = Base::GetNumKeys();
-        ValueType AdjustedP0 = P0;
+        // Find segment to split
+        const int32 NumSegments = GetNumSegments();
+        if (NumSegments == 0)
+        {
+            // If no segments, just add this as the first point
+            Base::AddValue(Position, 0.0f);
+            return 0;
+        }
+        
+        // Ensure parameter is in valid range
+        Parameter = FMath::Clamp(Parameter, 0.0f, 1.0f);
+            	
+        float MappedParam = this->GetEvaluator().MapToParentSpace(Parameter);
+    	const float ClampedParam = FMath::Clamp(MappedParam, 0.0f, static_cast<float>(NumSegments));
+    	const int32 SegmentIndex = FMath::Min(FMath::FloorToInt(ClampedParam), NumSegments - 1);
+    	const float LocalT = ClampedParam - static_cast<float>(SegmentIndex); // Parameter [0,1] within segment
     	
-    	// Check continuity with previous segment if this isn't the first one
-    	if (NumPoints > 0 && bEnforceContinuity)
-    	{
-    		const ValueType& LastPoint = Base::GetValueAt(NumPoints - 1);
-    		if (!LastPoint.Equals(P0, UE_KINDA_SMALL_NUMBER))
-    		{
-    			UE_LOG(LogSpline, Warning, TEXT("PolyBezier segment start point doesn't match previous end point. Adjusting for continuity."));
-    			AdjustedP0 = LastPoint; // Adjust the start point for continuity
-    		}
-    	}
+        // Get the control points of the segment to split
+        const int32 BaseIndex = SegmentIndex * 4;
+        ValueType P0 = Base::GetValueAt(BaseIndex);
+        ValueType P1 = Base::GetValueAt(BaseIndex + 1);
+        ValueType P2 = Base::GetValueAt(BaseIndex + 2);
+        ValueType P3 = Base::GetValueAt(BaseIndex + 3);
+
+    	float Param0 = Base::GetParameterAt(BaseIndex);
+    	float Param3 = Base::GetParameterAt(BaseIndex + 3);
+    	float SplitParam = Param0 + (Param3 - Param0) * LocalT;
+    	
+        // Use De Casteljau algorithm to split the curve
+        const float t = LocalT;
+        const float mt = 1.0f - t;
+        
+        // Calculate split control points
+        ValueType Q0 = P0;
+        ValueType Q1 = P0 * mt + P1 * t;
+        ValueType Q2 = (P0 * mt + P1 * t) * mt + (P1 * mt + P2 * t) * t;
+        ValueType Q3 = Position;  // Use provided position for the split point
+        
+        ValueType R0 = Position;  // Use provided position for the split point
+        ValueType R1 = (P1 * mt + P2 * t) * mt + (P2 * mt + P3 * t) * t;
+        ValueType R2 = P2 * mt + P3 * t;
+        ValueType R3 = P3;
+        
+        // Store indices for return value calculation
+        const int32 InsertedPointIndex = BaseIndex + 3;
+        
+        
+		// Remove the original segment
+		RemoveSegmentInternal(SegmentIndex);
+		
+        // Add the two new segments using the policy
+        AddBezierSegmentInternal(Q0, Q1, Q2, Q3, Param0, ParameterizationPolicy);
+        AddBezierSegmentInternal(R0, R1, R2, R3, SplitParam, ParameterizationPolicy);
+        
+		
+    	// Insert the knot directly into the knot vector
+    	this->GetEvaluator().InsertKnot(SplitParam);
 
-    	// Add control points using base storage
-    	const float SegmentT = static_cast<float>(NumPoints / 4);
-    	Base::AddValue(AdjustedP0, SegmentT);
-    	Base::AddValue(P1, SegmentT);
-    	Base::AddValue(P2, SegmentT);
-    	Base::AddValue(P3, SegmentT);
+        return InsertedPointIndex;
+    }
 
-        // Update segment structure
-        UpdateSegmentStructure();
+    /**
+     * Inserts a Bezier segment at the specified parameter
+     * 
+     * @param Parameter Parameter value [0,1] where to insert the segment
+     * @param P1 First control point
+     * @param P2 Second control point
+     * @param P3 End point
+     * @param ParameterizationPolicy How to distribute parameters
+     * @return Index of the new segment
+     */
+    int32 InsertBezierSegment(
+        float Parameter, 
+        const ValueType& P1, 
+        const ValueType& P2, 
+        const ValueType& P3, 
+        EParameterizationPolicy ParameterizationPolicy = EParameterizationPolicy::Centripetal)
+    {
+        // Handle edge cases
+        if (Parameter >= 1.0f)
+        {
+            // Add new segment at the end
+            return AppendBezierSegment(P1, P2, P3, ParameterizationPolicy);
+        }
         
-        return NumPoints / 4;
+        if (Parameter <= 0.0f)
+        {
+			// we treat P1, P2, P3 as the starting three control points, and connect to existing spline at the start
+            return PrependBezierSegment(P1, P2, P3, ParameterizationPolicy);
+        }
+        
+        // Normal case - insert at parameter
+        const int32 NumSegments = GetNumSegments();
+        
+        // Find which segment contains this parameter
+        const float SegmentLength = 1.0f / NumSegments;
+        const int32 SegmentIndex = FMath::Min(FMath::FloorToInt(Parameter / SegmentLength), NumSegments - 1);
+        const float LocalParam = (Parameter - SegmentIndex * SegmentLength) / SegmentLength;
+        
+        // Evaluate position at the parameter
+        ValueType PositionAtParam = this->Evaluate(Parameter);
+        
+        // Use the InsertPoint method to split the curve at the parameter
+        int32 PointIndex = InsertPoint(Parameter, PositionAtParam, ParameterizationPolicy);
+        
+        // Add the new segment from the inserted point
+        int32 Result = AddBezierSegmentInternal(PositionAtParam, P1, P2, P3, Parameter, ParameterizationPolicy);
+
+    	// Update the knot vector by inserting the new boundary.
+    	this->GetEvaluator().InsertLocalKnots(Parameter, Parameter);
+    	
+        return Result;
     }
 
+    
+
     /**
      * Removes a Bezier segment from the spline
-     * @param SegmentIndex Segment index to remove
+     * Prevents removing the last segment
+     * @param SegmentIndex - Segment index to remove
      * @return true if successfully removed
      */
     bool RemoveSegment(const int32 SegmentIndex)
     {
-    	const int32 NumSegments = GetNumSegments();
-    	if (SegmentIndex < 0 || SegmentIndex >= NumSegments)
-    	{
-    		return false;
-    	}
-
-    	const int32 BaseIndex = SegmentIndex * 4;
-    	bool bSuccess = true;
-    	bSuccess &= Base::RemoveValue(BaseIndex + 3);
-    	bSuccess &= Base::RemoveValue(BaseIndex + 2);
-    	bSuccess &= Base::RemoveValue(BaseIndex + 1);
-    	bSuccess &= Base::RemoveValue(BaseIndex);
-
-    	// Update segment structure
-    	UpdateSegmentStructure();
-    	
-    	return bSuccess;
+		// Prevent removing the last segment
+		if (GetNumSegments() <= 1)
+		{
+			UE_LOG(LogSpline, Warning, TEXT("Cannot remove the last segment from a spline"));
+			return false;
+		}
+		
+    	return RemoveSegmentInternal(SegmentIndex);
     }
+
     /**
 	 * Updates a single control point within a Bezier segment
 	 * @param SegmentIndex - Which cubic Bezier segment (each has 4 points)
@@ -411,7 +1191,7 @@ public:
 
         // Convert segment+point index to global index in spline values
         const int32 GlobalIndex = SegmentIndex * 4 + PointIndex;
-        return UpdateValue(GlobalIndex, NewValue);
+        return Base::UpdateValue(GlobalIndex, NewValue);
     }
 
     /**
@@ -442,7 +1222,6 @@ public:
         return bSuccess;
     }
 
-
     /** Returns the number of Bezier segments in the spline */
     virtual int32 GetNumSegments() const override
     {
@@ -450,6 +1229,12 @@ public:
         return NumPoints/4;
     }
 
+    /**
+     * Finds the closest point on the spline to the given point
+     * @param Point - Target point to find nearest point for
+     * @param OutSquaredDistance - Output squared distance to the nearest point
+     * @return Parameter value of the nearest point [0,1]
+     */
 	virtual float FindNearest(const ValueType& Point, float& OutSquaredDistance) const override
     {
     	// Can use same structure as TBSpline since we're cubic
@@ -475,6 +1260,13 @@ public:
     	return BestParam;
     }
 	
+    /**
+     * Finds the closest point on a specific segment to the given point
+     * @param Point - Target point to find nearest point for
+     * @param SegmentIndex - Index of the segment to search
+     * @param OutSquaredDistance - Output squared distance to the nearest point
+     * @return Parameter value within the segment [0,1]
+     */
 	virtual float FindNearestOnSegment(const ValueType& Point, int32 SegmentIndex, float& OutSquaredDistance) const override
     {
     	if (!IsValidSegmentIndex(SegmentIndex))
@@ -486,7 +1278,7 @@ public:
     	// Get Bezier control points for this segment
     	TArray<ValueType> Coeffs;
     	Coeffs.SetNum(4);
-    	const int32 BaseIndex = SegmentIndex * 4;  // Changed from *3 to *4
+    	const int32 BaseIndex = SegmentIndex * 4;
         
     	// Get the control points
     	const ValueType& P0 = this->GetValueAt(BaseIndex);
@@ -497,13 +1289,17 @@ public:
     	// Setup coefficients in standard Bezier polynomial form relative to test point
     	Coeffs[0] = P0 - Point;                    // constant term
     	Coeffs[1] = 3 * (P1 - P0);                 // linear coefficient
-    	Coeffs[2] = 3 * (P2 - 2*P1 + P0);         // quadratic coefficient  
-    	Coeffs[3] = P3 - 3*P2 + 3*P1 - P0;        // cubic coefficient
+    	Coeffs[2] = 3 * (P2 - 2*P1 + P0);          // quadratic coefficient  
+    	Coeffs[3] = P3 - 3*P2 + 3*P1 - P0;         // cubic coefficient
 
     	return Math::FindNearestPoint_Cubic(Coeffs, 0.0f, 1.0f, OutSquaredDistance);
     }
 
-    /** Gets the endpoint of a segment */
+    /**
+     * Gets the endpoint of a segment
+     * @param SegmentIndex - Index of the segment
+     * @return The end point position
+     */
     ValueType GetSegmentEndPoint(int32 SegmentIndex) const
     {
         if (!IsValidSegmentIndex(SegmentIndex))
@@ -512,25 +1308,252 @@ public:
         }
         return Base::GetValueAt(SegmentIndex * 4 + 3);
     }
+    
+    /**
+     * Checks if a segment index is valid
+     * @param Index - Index to check
+     * @return true if the index is valid
+     */
     bool IsValidSegmentIndex(int32 Index) const
     {
         return Index >= 0 && Index < GetNumSegments();
     }
+
+	/**
+	 * Reparameterizes the spline based on a new policy
+	 */
+	void Reparameterize(EParameterizationPolicy Mode = EParameterizationPolicy::Centripetal)
+    {
+    	// Pass the control points to the evaluator when updating segment structure
+    	// This allows distance-based knot generation to work properly
+    	this->GetEvaluator().Reparameterize(Mode, this->Values.GetStorage().GetEntries());
+    }
     
 protected:
     virtual TSharedPtr<Base> CreateNewInstance() const override
     {
-        return MakeShared<TPolyBezierSpline>();
+    	return MakeShared<TPolyBezierSpline>(CreateDefault());
     }
+
 private:
-    void UpdateSegmentStructure()
+	// Default constructor used for serialization
+    TPolyBezierSpline()
+		: Base(MakeUnique<TPolyBezierEvaluator<ScalarType, Dimensions>>())
+	{
+    	// Create a minimal valid default segment for serialization
+    	ValueType P0(0, 0, 0);
+    	ValueType P1(0.33f, 0, 0);
+    	ValueType P2(0.66f, 0, 0);
+    	ValueType P3(1, 0, 0);
+    	AddBezierSegmentInternal(P0, P1, P2, P3, 0, EParameterizationPolicy::Uniform);
+	}
+	
+	/**
+	 * Computes an appropriate parameter for a new segment based on insertion position and parameterization policy
+	 * @param P0 Start point of the segment
+	 * @param P1 First control point
+	 * @param P2 Second control point
+	 * @param P3 End point of the segment
+	 * @param RequestedParam The requested parameter (used to determine prepend vs append)
+	 * @param ParameterizationPolicy The parameterization policy to use
+	 * @return A properly computed parameter value for the segment
+	 */
+	float ComputeNewSegmentParameter(
+	    const ValueType& P0,
+	    const ValueType& P1, 
+	    const ValueType& P2, 
+	    const ValueType& P3,
+	    float RequestedParam,
+	    EParameterizationPolicy ParameterizationPolicy)
+	{
+	    // Get current parameter range
+	    float CurrentStart, CurrentEnd;
+	    this->GetEvaluator().GetParentSpaceRange(CurrentStart, CurrentEnd);
+	    
+	    // Handle special case for empty splines
+	    bool bIsEmpty = (this->GetNumSegments() == 0);
+	    
+	    // If the spline is empty, just use the requested parameter (or 0 if near zero)
+	    if (bIsEmpty)
+	    {
+	        return FMath::IsNearlyZero(RequestedParam) ? 0.0f : RequestedParam;
+	    }
+	    
+	    // Check if parameter is within current range
+	    if (RequestedParam >= CurrentStart && RequestedParam <= CurrentEnd)
+	    {
+	        // For insertion within the curve, use the requested parameter directly
+	        return RequestedParam;
+	    }
+	    
+	    // Handle out-of-bounds cases
+	    float NewParam;
+	    float SegmentLength;
+	    
+	    // Calculate segment length based on parameterization policy
+	    switch (ParameterizationPolicy)
+	    {
+	    case EParameterizationPolicy::ChordLength:
+	        SegmentLength = static_cast<float>((P3 - P0).Size());
+	        break;
+	    case EParameterizationPolicy::Centripetal:
+	        SegmentLength = FMath::Sqrt(static_cast<float>((P3 - P0).Size()));
+	        break;
+	    case EParameterizationPolicy::Uniform:
+	    default:
+	        SegmentLength = 1.0f;
+	        break;
+	    }
+	    
+	    // Determine prepend or append
+	    if (RequestedParam < CurrentStart)
+	    {
+	        // Prepend case - parameter before current start
+	        NewParam = CurrentStart - SegmentLength;
+	    }
+	    else
+	    {
+	        // Append case - parameter after current end
+	        NewParam = CurrentEnd + SegmentLength;
+	        
+	        // For uniform parameterization, use integer values for append
+	        if (ParameterizationPolicy == EParameterizationPolicy::Uniform)
+	        {
+	            const int32 NumSegments = GetNumSegments();
+	            NewParam = static_cast<float>(NumSegments + 1);
+	        }
+	    }
+	    
+	    return NewParam;
+	}
+
+	/**
+	 * Computes the end parameter for a segment given the start parameter and policy
+	 * @param StartParam - Start parameter for the segment
+	 * @param StartPoint - Start point of the segment
+	 * @param EndPoint - End point of the segment
+	 * @param ParameterizationPolicy - Policy to use for parameterization
+	 * @return End parameter for the segment
+	 */
+	float ComputeEndParameter(
+		float StartParam,
+		const ValueType& StartPoint,
+		const ValueType& EndPoint,
+		EParameterizationPolicy ParameterizationPolicy)
     {
-        const int32 NumPoints = this->GetNumKeys();
+    	float EndParam;
+    
+    	switch (ParameterizationPolicy)
+    	{
+    	case EParameterizationPolicy::ChordLength:
+    		{
+    			// Use actual distance between endpoints
+    			const float Distance = static_cast<float>((EndPoint - StartPoint).Size());
+    			EndParam = StartParam + Distance;
+    		}
+    		break;
+        
+    	case EParameterizationPolicy::Centripetal:
+    		{
+    			// Use square root of distance for better curvature distribution
+    			const float Distance = FMath::Sqrt(static_cast<float>((EndPoint - StartPoint).Size()));
+    			EndParam = StartParam + Distance;
+    		}
+    		break;
         
-        this->GetEvaluator().UpdateSegmentStructure(GetNumSegments(), NumPoints);
+    	case EParameterizationPolicy::Uniform:
+    	default:
+    		EndParam = StartParam + 1.0f;  // Default for uniform
+			break;
+    	}
+    
+    	return EndParam;
     }
-};
+	
+    /**
+     * Internal implementation to add a Bezier segment with all 4 control points
+     * @param P0 Start point
+     * @param P1 First control point
+     * @param P2 Second control point
+     * @param P3 End point
+     * @param StartParam Start parameter for the segment
+     * @param ParameterizationPolicy How to distribute parameters
+     * @return Index of the new segment
+     */
+	int32 AddBezierSegmentInternal(
+		const ValueType& P0,
+		const ValueType& P1, 
+		const ValueType& P2, 
+		const ValueType& P3,
+		float StartParam,
+		EParameterizationPolicy ParameterizationPolicy = EParameterizationPolicy::Centripetal)
+    {
+    	// Compute the appropriate start parameter based on requested value and policy
+    	StartParam = ComputeNewSegmentParameter(P0, P1, P2, P3, StartParam, ParameterizationPolicy);
+    
+    	// Calculate endpoint parameter based on policy
+    	float EndParam = ComputeEndParameter(StartParam, P0, P3, ParameterizationPolicy);
 
+    	// Calculate intermediate control point parameters
+    	float IntermediateParams[2] = {
+    		StartParam + (EndParam - StartParam) * 0.33f,
+			StartParam + (EndParam - StartParam) * 0.67f
+		};
+
+    	// Add the control points
+    	int32 P0Index = Base::AddValue(P0, StartParam);
+    	Base::AddValue(P1, IntermediateParams[0]);
+    	Base::AddValue(P2, IntermediateParams[1]);
+    	Base::AddValue(P3, EndParam);
+    
+    	this->GetEvaluator().InsertLocalKnots(StartParam, EndParam);
+
+    	// Update parameter range
+    	float CurrentStart, CurrentEnd;
+    	this->GetEvaluator().GetParentSpaceRange(CurrentStart, CurrentEnd);
+    
+    	this->GetEvaluator().SetParentSpaceRange(
+			FMath::Min(CurrentStart, StartParam),
+			FMath::Max(CurrentEnd, EndParam)
+		);
+    
+    	return P0Index / 4;
+    }
+
+    /**
+     * Internal implementation to remove a Bezier segment
+     * @param SegmentIndex - Index of the segment to remove
+     * @return true if successfully removed
+     */
+    bool RemoveSegmentInternal(const int32 SegmentIndex)
+    {
+    	if (SegmentIndex < 0 || SegmentIndex >= GetNumSegments())
+    	{
+    		return false;
+    	}
+
+    	const int32 BaseIndex = SegmentIndex * 4;
+    	// Get the parameters for the removed segment.
+    	float StartParam = Base::GetParameterAt(BaseIndex);
+    	float EndParam = Base::GetParameterAt(BaseIndex + 3);
+    	
+    	bool bSuccess = true;
+    	bSuccess &= Base::RemoveValue(BaseIndex + 3);
+    	bSuccess &= Base::RemoveValue(BaseIndex + 2);
+    	bSuccess &= Base::RemoveValue(BaseIndex + 1);
+    	bSuccess &= Base::RemoveValue(BaseIndex);
+
+    	if (bSuccess)
+    	{
+    		// Remove the corresponding knots.
+    		// Here we remove one occurrence at each boundary.
+    		this->GetEvaluator().RemoveKnot(StartParam);
+    		this->GetEvaluator().RemoveKnot(EndParam);
+    	}
+    	return bSuccess;
+    }
+};
+	
 // Common type definitions
 using FPolyBezierSpline2f = TPolyBezierSpline<float, 2>;
 using FPolyBezierSpline3f = TPolyBezierSpline<float, 3>;
diff --git a/Engine/Source/Runtime/GeometryCore/Public/Splines/Spline.h b/Engine/Source/Runtime/GeometryCore/Public/Splines/Spline.h
--- a/Engine/Source/Runtime/GeometryCore/Public/Splines/Spline.h
+++ b/Engine/Source/Runtime/GeometryCore/Public/Splines/Spline.h
@@ -4,6 +4,8 @@
 
 #include "AttributeChannel.h"
 #include "BaseEvaluator.h"
+#include "ChannelTypeRegistry.h"
+#include "BuiltInAttributeTypes.h"
 
 namespace UE
 {
@@ -24,42 +26,71 @@ enum class ESplineCloneFlags : uint32
 };
 ENUM_CLASS_FLAGS(ESplineCloneFlags);
 
+// Helper function to ensure initialization
+inline bool EnsureSplineTypesInitialized()
+{
+	static std::atomic<bool> bInitialized(false);
+	bool bWasInitialized = bInitialized.load(std::memory_order_acquire);
+
+	if (!bWasInitialized)
+	{
+		if (bInitialized.compare_exchange_strong(bWasInitialized, true))
+		{
+			UE::Geometry::Spline::FChannelTypeRegistry::InitializeBuiltInTypes();
+			return true;
+		}
+	}
+
+	return bWasInitialized;
+}
+	
 // Default evaluator for metadata/attributes that follows parent spline's parameter space
 template <typename ValueType, typename SplineT, int32 SplineD>
-class TAttributeChannelEvaluator final : public TBaseEvaluator<ValueType>
+class TParentBoundAttributeEvaluator final : public TBaseEvaluator<ValueType>
 {
 public:
 	// Make parent interface type accessible
 	using InterfaceType = TBaseEvaluator<ValueType>;
 
 	using SplineValueType = typename TChooseVectorType<SplineT, SplineD>::Type;
-	const TBaseEvaluator<SplineValueType>& ParentEvaluator;
+	// Change from reference to pointer
+	const TBaseEvaluator<SplineValueType>* ParentEvaluator;
 
 public:
-	explicit TAttributeChannelEvaluator(const TBaseEvaluator<SplineValueType>& InParentEvaluator)
-		: ParentEvaluator(InParentEvaluator)
+	explicit TParentBoundAttributeEvaluator(const TBaseEvaluator<SplineValueType>& InParentEvaluator)
+		: ParentEvaluator(&InParentEvaluator)
 	{
 	}
 
 	virtual TUniquePtr<TBaseEvaluator<ValueType>> Clone() const override
 	{
-		return MakeUnique<TAttributeChannelEvaluator<ValueType, SplineT, SplineD>>(ParentEvaluator);
+		return MakeUnique<TParentBoundAttributeEvaluator<ValueType, SplineT, SplineD>>(*this);
 	}
 
 	virtual int32 GetDegree() const override
 	{
-		return ParentEvaluator.GetDegree();
+		return ParentEvaluator->GetDegree();
 	}
 
+	virtual bool UpdateParentEvaluator(const void* NewParentEvaluator) override
+	{
+		if (NewParentEvaluator)
+		{
+			using ParentEvalType = TBaseEvaluator<SplineValueType>;
+			ParentEvaluator = static_cast<const ParentEvalType*>(NewParentEvaluator);
+			return true;
+		}
+		return false;
+	}
 protected:
 	virtual float GetParentSpaceMapping(float NormalizedParam) const override
 	{
-		return ParentEvaluator.MapToParentSpace(NormalizedParam);
+		return ParentEvaluator->MapToParentSpace(NormalizedParam);
 	}
 
 	virtual float GetNormalizedSpaceMapping(float ParentParam) const override
 	{
-		return ParentEvaluator.MapFromParentSpace(ParentParam);
+		return ParentEvaluator->MapFromParentSpace(ParentParam);
 	}
 
 };
@@ -81,6 +112,8 @@ public:
 	explicit TSplineBase(TUniquePtr<EvaluatorType> InEvaluator = nullptr)
 		: Values(InEvaluator ? MoveTemp(InEvaluator) : MakeUnique<EvaluatorType>())
 	{
+		// Ensure initialization before creating any attribute channel
+		EnsureSplineTypesInitialized();
 	}
 
 	virtual ~TSplineBase() = default;
@@ -115,6 +148,9 @@ public:
 
 	friend FArchive& operator<<(FArchive& Ar, TSplineBase& Spline)
 	{
+		// Ensure initialization before serialization/deserialization
+		EnsureSplineTypesInitialized();
+		
 		// Serialize main value channel
 		Ar << Spline.Values;
         
@@ -128,6 +164,12 @@ public:
 			{
 				// Save channel name and data
 				Ar << Pair.Key;
+				
+				// Save type name explicitly - critical for correct deserialization
+				FName ChannelTypeName = Pair.Value->GetTypeName();
+				Ar << ChannelTypeName;
+            
+				// Now serialize the channel data
 				Pair.Value->Serialize(Ar);
 			}
 		}
@@ -137,18 +179,43 @@ public:
             
 			for (int32 i = 0; i < NumChannels; ++i)
 			{
+				// Read channel name
 				FName ChannelName;
 				Ar << ChannelName;
-                
-				// Create new channel
-				TUniquePtr<IAttributeChannelBase> Channel = 
-					MakeUnique<TAttributeChannel<ValueType>>();
-				Channel->Serialize(Ar);
-                
-				Spline.AttributeChannels.Add(ChannelName, MoveTemp(Channel));
+            
+				// Read the channel type name
+				FName ChannelTypeName;
+				Ar << ChannelTypeName;
+            
+				// Create the appropriate channel type using the registry
+				TUniquePtr<IAttributeChannelBase> Channel;
+            
+				if (FChannelTypeRegistry::IsTypeRegistered(ChannelTypeName))
+				{
+					// Use the registry to create the correct channel type
+					Channel = FChannelTypeRegistry::CreateChannel(ChannelTypeName);
+				}
+				else
+				{
+					// Default to the spline's value type for unknown types
+					Channel = MakeUnique<TAttributeChannel<ValueType>>();
+					UE_LOG(LogSpline, Warning, TEXT("Unknown channel type '%s', using default type."), *ChannelTypeName.ToString());
+				}
+            
+				// Deserialize the channel data
+				if (Channel)
+				{
+					Channel->Serialize(Ar);
+					Spline.AttributeChannels.Add(ChannelName, MoveTemp(Channel));
+				}
 			}
+			// After loading all channels, update their parent evaluator references
+            for (auto& Pair : Spline.AttributeChannels)
+            {
+                Pair.Value->UpdateEvaluatorParent(&Spline.GetEvaluator());
+            }
 		}
-        
+		        
 		return Ar;
 	}
 
@@ -174,6 +241,16 @@ public:
 			return;
 		}
 
+		// Ensure the registry itself is initialized
+		EnsureSplineTypesInitialized();
+
+		// Ensure this specific type is registered
+		const FName TypeName = TAttributeTypeTraits<AttrType>::Name;
+		if (!FChannelTypeRegistry::IsTypeRegistered(TypeName))
+		{
+			FChannelTypeRegistry::RegisterType<AttrType>();
+		}
+
 		TUniquePtr<TAttributeChannel<AttrType>> Channel = MakeUnique<TAttributeChannel<AttrType>>(MoveTemp(AttrEvaluator));
 		AttributeChannels.Add(Name, MoveTemp(Channel));
 	}
@@ -186,7 +263,7 @@ public:
 	template <typename AttrType>
 	void CreateAttributeChannel(const FName& Name)
 	{
-		using AttributeEvaluator = TAttributeChannelEvaluator<AttrType, ScalarType, Dimensions>;
+		using AttributeEvaluator = TParentBoundAttributeEvaluator<AttrType, ScalarType, Dimensions>;
 
 		// Create the evaluator
 		TUniquePtr<AttributeEvaluator> ChannelEvaluator = MakeUnique<AttributeEvaluator>(GetEvaluator());
@@ -220,7 +297,7 @@ public:
 	 * Also expands the parameter range if necessary.
 	 * @param Value - The value to add
 	 * @param Parameter - The parameter value in parent spline space
-	 * @return 
+	 * @return Index of the added value
 	 */
 	virtual int32 AddValue(const ValueType& Value, float Parameter)
 	{
@@ -319,18 +396,6 @@ public:
 		return AttrType();
 	}
 
-	/**
-	 * Clears all values from the spline including custom attribute channels
-	 */
-	virtual void Clear()
-	{
-		Values.Clear();
-		for (TTuple<FName, TUniquePtr<IAttributeChannelBase>>& Pair : AttributeChannels)
-		{
-			static_cast<TAttributeChannel<ValueType>*>(Pair.Value.Get())->Clear();
-		}
-	}
-
 	/**
 	 * Clears all values from a specific attribute channel
 	 * @param Name - Name of the attribute channel to clear
@@ -346,19 +411,91 @@ public:
 		return false;
 	}
 
-	/** 
-	 * Auto-expands parameter range to include all control points.
-	 * Called automatically when adding points outside current range.
+	/**
+	 * Checks if the spline has an attribute channel with the given name
+	 * @param Name - Name of the attribute channel 
+	 * @return true if the channel exists, false otherwise
 	 */
-	void UpdateParameterRange()
+	bool HasAttributeChannel(const FName& Name) const
 	{
-		if (GetNumKeys() < 2) return;
+		return AttributeChannels.Find(Name) != nullptr;
+	}
 
-		const auto& Entries = Values.GetStorage().GetEntries();
-		const float NewStart = Entries[0].Parameter;
-		const float NewEnd = Entries.Last().Parameter;
+	/**
+	 * Retrieves the number of values in a specific attribute channel
+	 * @param Name - Name of the attribute channel
+	 * @return Number of values or zero if the channel does not exist
+	 */
+	int32 NumAttributeValues(const FName& Name) const
+	{
+		if (const TUniquePtr<IAttributeChannelBase>* Channel = AttributeChannels.Find(Name))
+		{
+			return (*Channel)->NumKeys();
+		}
+		return 0;
+	}
 
-		SetParameterRange(NewStart, NewEnd);
+	/**
+	 * Retrieves the parameter value of a specific attribute at the given index
+	 * @param Index - Index of the value
+	 * @param Name - Name of the attribute channel
+	 * @return Parameter value
+	 */
+	float GetAttributeParameter(int32 Index, const FName& Name) const
+	{
+		if (const TUniquePtr<IAttributeChannelBase>* Channel = AttributeChannels.Find(Name))
+		{
+			return (*Channel)->GetParameter(Index);
+		}
+		return 0.f;
+	}
+
+	/**
+	 * Sets the parameter value of a specific attribute at the given index
+	 * @param Index - Index of the value
+	 * @param NewParameter - New parameter to set
+	 * @param Name - Name of the attribute channel
+	 */
+	void SetAttributeParameter(int32 Index, float NewParameter, const FName& Name)
+	{
+		if (const TUniquePtr<IAttributeChannelBase>* Channel = AttributeChannels.Find(Name))
+		{
+			(*Channel)->SetParameter(Index, NewParameter);
+		}
+	}
+
+	/**
+	 * Retrieves the value of a specific attribute at the given index
+	 * @tparam AttrType - Custom attribute type
+	 * @param Index - Index of the value
+	 * @param Name - Name of the attribute channel
+	 * @return Attribute value
+	 */
+	template <typename AttrType>
+	AttrType GetAttributeValue(int32 Index, const FName& Name) const
+	{
+		if (const TAttributeChannel<AttrType>* Channel = GetAttributeChannel<AttrType>(Name))
+		{
+			return Channel->GetStorage().GetEntries()[Index].Value;
+		}
+
+		return AttrType();
+	}
+
+	/**
+	 * Sets the value of a specific attribute at the given index
+	 * @tparam AttrType - Custom attribute type
+	 * @param Index - Index of the value
+	 * @param Value - New value to set
+	 * @param Name - Name of the attribute channel
+	 */
+	template <typename AttrType>
+	void SetAttributeValue(int32 Index, const AttrType& Value, const FName& Name)
+	{
+		if (const TAttributeChannel<AttrType>* Channel = GetAttributeChannel<AttrType>(Name))
+		{
+			Channel->GetStorage().GetEntries()[Index].Value = Value;
+		}
 	}
 
 	/** 
@@ -408,7 +545,10 @@ public:
 
 	virtual void SetClosedLoop(bool bInClosedLoop)
 	{
-		GetEvaluator().SetClosedLoop(bInClosedLoop);
+		if (this->IsClosedLoop() != bInClosedLoop)
+		{
+			this->GetEvaluator().SetClosedLoop(bInClosedLoop);
+		}
 	}
 	
 	bool IsClosedLoop() const
@@ -454,21 +594,30 @@ public:
 		return NewSpline;
 	}
 
-protected:
 	template <typename AttrType>
-	TAttributeChannel<AttrType>* GetAttributeChannel(const FName Name)
+	const TAttributeChannel<AttrType>* GetAttributeChannel(const FName& Name) const
+	{
+		const TUniquePtr<IAttributeChannelBase>* Found = AttributeChannels.Find(Name);
+		return Found ? static_cast<const TAttributeChannel<AttrType>*>(Found->Get()) : nullptr;
+	}
+	
+	const TAttributeChannel<ValueType>& GetValuesChannel() const
+	{
+		return Values;
+	}
+	
+	template <typename AttrType>
+	TAttributeChannel<AttrType>* GetAttributeChannel(const FName& Name)
 	{
 		TUniquePtr<IAttributeChannelBase>* Found = AttributeChannels.Find(Name);
 		return Found ? static_cast<TAttributeChannel<AttrType>*>(Found->Get()) : nullptr;
 	}
-
-	template <typename AttrType>
-	const TAttributeChannel<AttrType>* GetAttributeChannel(const FName Name) const
+protected:
+	TAttributeChannel<ValueType>& GetValuesChannel()
 	{
-		const TUniquePtr<IAttributeChannelBase>* Found = AttributeChannels.Find(Name);
-		return Found ? static_cast<const TAttributeChannel<AttrType>*>(Found->Get()) : nullptr;
+		return Values;
 	}
-
+	
 	/** Hook for derived classes to copy their specific data after base data is copied */
 	virtual void OnClone(TSharedPtr<TSplineBase> NewSpline, ESplineCloneFlags CloneFlags) const {}
 
@@ -494,6 +643,18 @@ protected:
 	 * @return Parameter value of the nearest point
 	 */
 	virtual float FindNearestOnSegment(const ValueType& Point, int32 SegmentIndex, float& OutSquaredDistance) const = 0;
+
+	/**
+	 * Clears all values from the spline including custom attribute channels
+	 */
+	virtual void Clear()
+	{
+		Values.Clear();
+		for (TTuple<FName, TUniquePtr<IAttributeChannelBase>>& Pair : AttributeChannels)
+		{
+			static_cast<TAttributeChannel<ValueType>*>(Pair.Value.Get())->Clear();
+		}
+	}
 private:
 	/**
 	 * Creates a new instance of the correct derived type for cloning.
@@ -527,15 +688,30 @@ private:
 				const FName& Name = AttributePair.Key;
 				const TUniquePtr<IAttributeChannelBase>& Channel = AttributePair.Value;
 
-				// Use type erasure to copy the channel without knowing its concrete type
-				if (const TAttributeChannel<ValueType>* TypedChannel = static_cast<const TAttributeChannel<ValueType>*>(Channel.Get()))
-				{
-					NewSpline->AttributeChannels.Add(Name, MakeUnique<TAttributeChannel<ValueType>>(*TypedChannel));
-				}
+				// Clone the channel
+				TUniquePtr<IAttributeChannelBase> ClonedChannel = Channel->Clone();
+				
+				// Update the parent evaluator reference directly
+				UpdateChannelEvaluatorParent(ClonedChannel.Get(), NewSpline->GetEvaluator());
+				
+				// Add to new spline
+                NewSpline->AttributeChannels.Add(Name, MoveTemp(ClonedChannel));
 			}
 		}
 		return NewSpline;
 	}
+	
+	/**
+	 * Updates the parent evaluator reference in cloned attribute channels
+	 */
+	void UpdateChannelEvaluatorParent(IAttributeChannelBase* Channel, const EvaluatorType& NewParentEvaluator) const
+	{
+		if (Channel)
+		{
+			Channel->UpdateEvaluatorParent(&NewParentEvaluator);
+		}
+	}
+	
 };
 
 
diff --git a/Engine/Source/Runtime/GeometryCore/Public/Splines/TangentBezierSpline.h b/Engine/Source/Runtime/GeometryCore/Public/Splines/TangentBezierSpline.h
--- a/Engine/Source/Runtime/GeometryCore/Public/Splines/TangentBezierSpline.h
+++ b/Engine/Source/Runtime/GeometryCore/Public/Splines/TangentBezierSpline.h
@@ -3,8 +3,88 @@
 #pragma once
 
 #include "Splines/PolyBezierSpline.h"
+namespace UE
+{
+namespace Geometry
+{
+namespace Spline
+{
+	
+/**
+ * Enum defining how tangents are computed for a spline control point.
+ */
+enum class ETangentMode : uint8
+{
+	/** Automatically compute unclamped tangents based on surrounding points */
+	Auto,
+    
+	/** Automatically compute clamped tangents based on surrounding points */
+	AutoClamped,
+    
+	/** User-specified tangents */
+	User,
+    
+	/** Broken tangents - in and out can be different */
+	Broken,
+    
+	/** Linear tangents aligned with adjacent points */
+	Linear,
+    
+	/** Constant - no interpolation between points */
+	Constant,
+    
+	/** Unknown or invalid tangent mode */
+	Unknown
+};
+
+/**
+ * Structure representing a control point on a tangent-based spline curve
+ */
+struct FTangentBezierControlPoint 
+{
+	/** Position of the control point */
+	FVector Position;
 
-using namespace UE::Geometry::Spline;
+	/** Incoming tangent vector */
+	FVector TangentIn;
+
+	/** Outgoing tangent vector */
+	FVector TangentOut;
+
+	/** Tangent computation mode */
+	ETangentMode TangentMode;
+
+	/** Default constructor creates a point at origin with zero tangents */
+	FTangentBezierControlPoint ()
+		: Position(FVector::ZeroVector)
+		, TangentIn(FVector::ZeroVector)
+		, TangentOut(FVector::ZeroVector)
+		, TangentMode(ETangentMode::Auto)
+	{
+	}
+
+	/** Constructor with position only - defaults to auto tangents */
+	explicit FTangentBezierControlPoint (const FVector& InPosition)
+		: Position(InPosition)
+		, TangentIn(FVector::ZeroVector)
+		, TangentOut(FVector::ZeroVector)
+		, TangentMode(ETangentMode::Auto)
+	{
+	}
+
+	/** Full constructor */
+	FTangentBezierControlPoint (
+		const FVector& InPosition,
+		const FVector& InTangentIn,
+		const FVector& InTangentOut,
+		ETangentMode InTangentMode = ETangentMode::Auto)
+		: Position(InPosition)
+		, TangentIn(InTangentIn)
+		, TangentOut(InTangentOut)
+		, TangentMode(InTangentMode)
+	{
+	}
+};
 
 /**
  * A spline that provides tangent-based control over curve shape while using 
@@ -16,19 +96,52 @@ struct FTangentBezierSpline
     /** Tension parameter for auto-computed tangents [0,1] */
     float Tension;
 
-    /** Whether tangents should be computed automatically for each point */
-    TArray<bool> AutoTangents;
-
-    FTangentBezierSpline()
-	    : Tension(0.0f)
+	/** How tangents should be computed for each point */
+	TArray<ETangentMode> TangentModes;
+	
+protected:
+	/** 
+	 * Protected default constructor for serialization
+	 * Initializes with a minimal default segment
+	 */
+	FTangentBezierSpline()
+		: Tension(0.0f)
+		, InternalSpline(FPolyBezierSpline3d::CreateLine(FVector::ZeroVector, FVector(100,0,0)))
+	{
+	}
+	
+public:
+	/** Default constructor with at least one segment */
+	FTangentBezierSpline(const FVector& StartPoint, const FVector& EndPoint)
+		: Tension(0.0f)
+		, InternalSpline(FPolyBezierSpline3d::CreateLine(StartPoint, EndPoint))
+	{
 		
-    {
-    }
+	}
 
+	/**
+	 * Constructor with full control over tangents
+	 */
+	FTangentBezierSpline(
+		const FVector& StartPoint, 
+		const FVector& EndPoint,
+		const FVector& StartTangent,
+		const FVector& EndTangent,
+		bool bAutoTangents = false)
+		: Tension(0.0f)
+		, TangentModes({bAutoTangents ? ETangentMode::Auto : ETangentMode::User,
+							bAutoTangents ? ETangentMode::Auto : ETangentMode::User})
+		, InternalSpline(FPolyBezierSpline3d(StartPoint,
+											StartPoint + StartTangent / 3.0f,
+											EndPoint - (EndTangent / 3.0f),
+											EndPoint))
+	{
+	}
+   
     /**  Copy constructor */
     FTangentBezierSpline(const FTangentBezierSpline& Other)
 	    : Tension(Other.Tension)
-	    , AutoTangents(Other.AutoTangents)
+	    , TangentModes(Other.TangentModes)
 	    , InternalSpline(Other.InternalSpline)
     {
     }
@@ -39,7 +152,7 @@ struct FTangentBezierSpline
         if (this != &Other)
         {
             Tension = Other.Tension;
-            AutoTangents = Other.AutoTangents;
+            TangentModes = Other.TangentModes;
             InternalSpline = Other.InternalSpline;
         }
         return *this;
@@ -48,7 +161,25 @@ struct FTangentBezierSpline
     bool Serialize(FArchive& Ar)
     {
         Ar << Tension;
-        Ar << AutoTangents;
+		// Serialize as uint8 for compatibility
+		int32 NumTangentModes = TangentModes.Num();
+		Ar << NumTangentModes;
+        
+		if (Ar.IsLoading())
+		{
+			TangentModes.SetNum(NumTangentModes);
+		}
+        
+		for (int32 i = 0; i < NumTangentModes; ++i)
+		{
+			uint8 TangentModeValue = static_cast<uint8>(TangentModes[i]);
+			Ar << TangentModeValue;
+            
+			if (Ar.IsLoading())
+			{
+				TangentModes[i] = static_cast<ETangentMode>(TangentModeValue);
+			}
+		}
         Ar << InternalSpline;
                 
         return true;
@@ -60,6 +191,115 @@ struct FTangentBezierSpline
         return Ar;
     }
 
+	// Static shape generators
+
+    /**
+     * Creates a straight line between two points
+     * @param StartPoint - Start point of the line
+     * @param EndPoint - End point of the line
+     * @return New TangentBezierSpline initialized as a line
+     */
+    static FTangentBezierSpline CreateLine(
+        const FVector& StartPoint, 
+        const FVector& EndPoint)
+    {
+		// Create empty spline - we'll manually build it
+		FTangentBezierSpline Result;
+        
+		// Clear the default initialization
+		Result.InternalSpline = FPolyBezierSpline3d::CreateLine(StartPoint, EndPoint);
+		
+        Result.Reparameterize(EParameterizationPolicy::Uniform);
+        return Result;
+    }
+
+    /**
+     * Creates a circular arc with specified parameters
+     * @param Center - Center point of the circle
+     * @param Radius - Radius of the circle
+     * @param StartAngle - Start angle in radians
+     * @param EndAngle - End angle in radians
+     * @param NumSegments - Number of segments to use (more segments = smoother curve)
+     * @return New TangentBezierSpline initialized as an arc
+     */
+    static FTangentBezierSpline CreateCircleArc(
+        const FVector& Center,
+        float Radius,
+        float StartAngle,
+        float EndAngle,
+        int32 NumSegments = 4)
+    {
+        // Create empty spline - we'll manually build it
+        FTangentBezierSpline Result;
+        
+        // Clear the default initialization
+        Result.InternalSpline = FPolyBezierSpline3d::CreateCircleArc(
+            FVector3d(Center), Radius, StartAngle, EndAngle, NumSegments);
+        
+		// Set up tangent status for all points
+		Result.TangentModes.Init(ETangentMode::Auto, Result.GetNumPoints());
+        
+        // Apply consistent parameterization - even though the internal PolyBezier already calls
+        // Reparameterize, we need to do it again here to ensure consistency
+        Result.Reparameterize(EParameterizationPolicy::Centripetal);
+        
+        return Result;
+    }
+
+    /**
+     * Creates a complete circle
+     * @param Center - Center point of the circle
+     * @param Radius - Radius of the circle
+     * @param NumSegments - Number of segments to use (more segments = smoother curve)
+     * @return New TangentBezierSpline initialized as a circle
+     */
+    static FTangentBezierSpline CreateCircle(
+        const FVector& Center,
+        float Radius,
+        int32 NumSegments = 4)
+    {
+        // Create the circle
+        FTangentBezierSpline Result;
+        Result.InternalSpline = FPolyBezierSpline3d::CreateCircle(
+            FVector3d(Center), Radius, NumSegments);
+        
+        // Set up tangent status for all points
+		Result.TangentModes.Init(ETangentMode::Auto, Result.GetNumPoints());
+		
+		// Apply consistent parameterization
+		Result.Reparameterize(EParameterizationPolicy::Centripetal);
+        
+        return Result;
+    }
+
+    /**
+     * Creates an ellipse with specified parameters
+     * @param Center - Center point of the ellipse
+     * @param RadiusX - Radius along X axis
+     * @param RadiusY - Radius along Y axis
+     * @param NumSegments - Number of segments to use (more segments = smoother curve)
+     * @return New TangentBezierSpline initialized as an ellipse
+     */
+    static FTangentBezierSpline CreateEllipse(
+        const FVector& Center,
+        float RadiusX,
+        float RadiusY,
+        int32 NumSegments = 4)
+    {
+        // Create the ellipse
+        FTangentBezierSpline Result;
+        Result.InternalSpline = FPolyBezierSpline3d::CreateEllipse(
+            FVector3d(Center), RadiusX, RadiusY, NumSegments);
+        
+		// Set up tangent status for all points
+		Result.TangentModes.Init(ETangentMode::Auto, Result.GetNumPoints());
+        
+        // Apply consistent parameterization
+        Result.Reparameterize(EParameterizationPolicy::Centripetal);
+        
+        return Result;
+    }
+    
     void SetClosedLoop(bool bInClosedLoop)
     {
         // Skip if state isn't changing
@@ -102,7 +342,7 @@ struct FTangentBezierSpline
 	            FVector P2 = P3 - (FirstInTangent / 3.0f);
 	            
 	            // Add the closing segment to the internal spline
-	            InternalSpline.AddBezierSegment(P0, P1, P2, P3, false);
+	            InternalSpline.AppendBezierSegment(P1, P2, P3);
 	        }
 	    }
 	    else
@@ -122,203 +362,517 @@ struct FTangentBezierSpline
 
 	    
 	    // Ensure auto tangents are updated after changing loop state
-	    UpdateAutoTangents();
+	    UpdateTangents();
     }
     
-    /** 
-     * Adds a new point to the spline
-     * @param Position - Location of the new point
-     * @param InTangent - Direction and magnitude of the arriving tangent
-     * @param OutTangent - Direction and magnitude of the leaving tangent
-     * @param bAutoComputeTangents - Whether to automatically compute tangents based on surrounding points
-     */
-    void AddPoint(const FVector& Position, 
-                 const FVector& InTangent = FVector::ZeroVector,
-                 const FVector& OutTangent = FVector::ZeroVector,
-                 bool bAutoComputeTangents = true)
-    {
-        const int32 NewIndex = GetNumPoints();
-        
-        // Store auto tangent flag
-        if (AutoTangents.Num() <= NewIndex)
-        {
-            AutoTangents.AddZeroed(NewIndex - AutoTangents.Num() + 1);
-        }
-        AutoTangents[NewIndex] = bAutoComputeTangents;
-        
-    	if (InternalSpline.GetNumKeys() == 0)
+	/**
+	 * Replaces all points in the spline with the provided points.
+	 * 
+	 * @param Points - Array of control points to set
+	 * @return True if the points were successfully set
+	 */
+	bool SetControlPoints(const TArray<FTangentBezierControlPoint>& Points)
+	{
+	    const int32 NumPoints = Points.Num();
+	    
+	    // Require at least 2 points to form a valid segment
+	    if (NumPoints < 2)
 	    {
-	        // First point - just store it directly
-	        InternalSpline.AddValue(FVector(Position), 0.0f);
-	        return;
+	        UE_LOG(LogSpline, Warning, TEXT("SetControlPoints requires at least 2 points to form a valid segment. Got %d points."), NumPoints);
+	        return false;
 	    }
 	    
-	    if (InternalSpline.GetNumKeys() == 1)
+	   // Convert to cubic Bezier control points for the internal spline
+	    TArray<FVector> BezierPoints;
+	    
+	    // Reserve space for Bezier control points (3n-2 for n points)
+	    BezierPoints.Reserve(3 * NumPoints - 2);
+	    
+	    // Handle the first point
+	    const FVector& FirstPoint = Points[0].Position;
+	    FVector FirstOutTangent = Points[0].TangentOut;
+	    
+		// Apply auto tangent calculation if needed
+		if (Points[0].TangentMode == ETangentMode::Auto || Points[0].TangentMode == ETangentMode::AutoClamped)
+		{
+			FirstOutTangent = (1.0f - Tension) * (Points[1].Position - FirstPoint);
+		}
+	    
+	    // Add first point and its outgoing control point
+	    BezierPoints.Add(FirstPoint);
+	    BezierPoints.Add(FirstPoint + FirstOutTangent / 3.0f);
+	    
+	    // Handle middle points
+	    for (int32 i = 1; i < NumPoints - 1; ++i)
 	    {
-	        // We have exactly one point stored - this is our second point
-	        // Get the first point, then clear to start building segments
-	        FVector P0 = FVector(InternalSpline.GetValueAt(0));
-	        FVector P3 = Position;
+	        const FVector& CurrentPoint = Points[i].Position;
 	        
-	        InternalSpline.Clear();
+	        // Get tangents based on mode
+	        FVector InTangent = Points[i].TangentIn;
+	        FVector OutTangent = Points[i].TangentOut;
 	        
-	        // Calculate control points
-	        FVector P1, P2;
-	        if (bAutoComputeTangents)
+	        if (Points[i].TangentMode == ETangentMode::Auto || Points[i].TangentMode == ETangentMode::AutoClamped)
 	        {
-	            // For first segment with auto tangents, use simple interpolation
-	            FVector AutoTangent = (1.0f - Tension) * (P3 - P0);
-	            P1 = P0 + (AutoTangent / 3.0f);
-	            P2 = P3 - (AutoTangent / 3.0f);
+	            // Compute a smooth tangent using surrounding points
+	            FVector AutoTangent = (1.0f - Tension) * (Points[i+1].Position - Points[i-1].Position);
+	            InTangent = AutoTangent;
+	            OutTangent = AutoTangent;
 	        }
-	        else
+	        else if (Points[i].TangentMode == ETangentMode::Linear)
 	        {
-	            // Use provided tangents
-	            P1 = P0 + (OutTangent / 3.0f); // Need to save first point's OutTangent
-	            P2 = P3 - (InTangent / 3.0f);
-	        }
-	        
-	        // Add the first segment
-	        InternalSpline.AddBezierSegment(P0, P1, P2, P3);
-	    }
-	    else
-	    {
-	        // Normal case - add segment to existing spline
-	        FVector P0 = GetPosition(NewIndex - 1);
-	        FVector P3 = Position;
-	        
-	        FVector P1, P2;
-	        if (bAutoComputeTangents)
-	        {
-	            // Calculate auto tangent
-	            FVector AutoTangent;
-	            if (NewIndex == 1)
-	            {
-	                // Simple case for first segment
-	                AutoTangent = (1.0f - Tension) * (P3 - P0);
-	            }
-	            else
-	            {
-	                // Use previous point for better tangent
-	                FVector PrevPos = GetPosition(NewIndex - 2);
-	                AutoTangent = (1.0f - Tension) * (P3 - PrevPos);
-	            }
-	            
-	            // Update previous point's out tangent if auto
-	            if (IsAutoTangent(NewIndex - 1))
-	            {
-	                SetTangentOut(NewIndex - 1, AutoTangent);
-	            }
-	            
-	            P1 = P0 + (GetTangentOut(NewIndex - 1) / 3.0f);
-	            P2 = P3 - (AutoTangent / 3.0f);
+	            // Linear tangents aligned with adjacent points
+	            InTangent = CurrentPoint - Points[i-1].Position;
+	            OutTangent = Points[i+1].Position - CurrentPoint;
 	        }
-	        else
+	        else if (Points[i].TangentMode == ETangentMode::Constant)
 	        {
-	            // Use provided tangents
-	            P1 = P0 + (GetTangentOut(NewIndex - 1) / 3.0f);
-	            P2 = P3 - (InTangent / 3.0f);
+	            // Zero tangents for constant interpolation
+	            InTangent = FVector::ZeroVector;
+	            OutTangent = FVector::ZeroVector;
 	        }
 	        
-	        // Add the segment
-	        InternalSpline.AddBezierSegment(P0, P1, P2, P3);
+	        // Add the control points
+	        BezierPoints.Add(CurrentPoint - InTangent / 3.0f);
+	        BezierPoints.Add(CurrentPoint);
+	        BezierPoints.Add(CurrentPoint + OutTangent / 3.0f);
 	    }
 	    
-	    // Update auto tangents for consistency
-	    if (bAutoComputeTangents)
+	    // Handle the last point
+	    const FVector& LastPoint = Points.Last().Position;
+	    FVector LastInTangent = Points.Last().TangentIn;
+	    
+	    if (Points.Last().TangentMode == ETangentMode::Auto || Points.Last().TangentMode == ETangentMode::AutoClamped)
 	    {
-	        UpdateAutoTangents();
+	        LastInTangent = (1.0f - Tension) * (LastPoint - Points[NumPoints - 2].Position);
 	    }
-    }
+	    
+	    // Add last incoming control point and the last point
+	    BezierPoints.Add(LastPoint - LastInTangent / 3.0f);
+	    BezierPoints.Add(LastPoint);
+	    
+	    // Set the control points in the internal spline
+	    InternalSpline.SetControlPoints(BezierPoints, EParameterizationPolicy::Centripetal);
+	    
+	    // Store tangent modes for later use
+	    TangentModes.Reset(NumPoints);
+	    for (const FTangentBezierControlPoint& Point : Points)
+	    {
+	        TangentModes.Add(Point.TangentMode);
+	    }
+	    
+	    // Update auto tangents to ensure proper shape
+	    UpdateTangents();
+	    
+	    // Reparameterize the spline for better parameter distribution
+	    Reparameterize();
+	    
+	    return true;
+	}
 
-    /** Removes a point from the spline */
-    void RemovePoint(int32 Index)
+	/**
+	 * Adds multiple points to the spline.
+	 * 
+	 * @param Points - Array of control points to add
+	 * @param bAppend - If true, append to the end of spline; if false, prepend to the beginning
+	 * @return True if the points were successfully added
+	 */
+	 bool AddControlPoints(
+        const TArray<FTangentBezierControlPoint>& Points,
+        bool bAppend = true)
     {
-        const int32 NumPoints = GetNumPoints();
-        if (Index < 0 || Index >= NumPoints)
+        const int32 NumPoints = Points.Num();
+        
+        // Need at least 2 points to add a valid segment
+        if (NumPoints < 2)
         {
-            return;
+            UE_LOG(LogSpline, Warning, TEXT("AddControlPoints requires at least 2 points to add a valid segment. Got %d points."), NumPoints);
+            return false;
         }
         
-        // Remove associated auto tangent flag
-        if (AutoTangents.IsValidIndex(Index))
+        // If the spline is currently empty (though it shouldn't be based on our design),
+        // this is equivalent to SetControlPoints
+        if (GetNumPoints() == 0)
         {
-            AutoTangents.RemoveAt(Index);
+            return SetControlPoints(Points);
         }
-        InternalSpline.RemoveSegment(Index);
-        UpdateAutoTangents();
+        
+        if (bAppend)
+        {
+            // For each point, use AddPoint 
+            for (int32 i = 0; i < NumPoints; ++i)
+            {
+                // Skip the first point if it's the same as our last existing point 
+                // (to avoid duplicates at connection point)
+                if (i == 0)
+                {
+                    const int32 ExistingNumPoints = GetNumPoints();
+                    const FVector LastExistingPoint = GetPosition(ExistingNumPoints - 1);
+                    
+                    // Check if connection points are reasonably close
+                    constexpr double ConnectionTolerance = 1e-4;
+                    const double DistanceSquared = (LastExistingPoint - Points[0].Position).SizeSquared();
+                    
+                    if (DistanceSquared <= ConnectionTolerance)
+                    {
+                        // Points are close - the first existing point is the same as our first new point
+                        // So update the last point's tangent mode if needed
+                        if (Points[0].TangentMode != ETangentMode::User && Points[0].TangentMode != ETangentMode::Broken)
+                        {
+                            SetPointTangentMode(ExistingNumPoints - 1, Points[0].TangentMode);
+                        }
+                        continue; // Skip adding this point
+                    }
+                }
+                
+                // Add each point normally
+                AddPoint(Points[i]);
+            }
+        }
+        else
+        {
+            // For prepending, we need to add points in reverse order
+            for (int32 i = NumPoints - 1; i >= 0; --i)
+            {
+                // Skip the last point if it's the same as our first existing point
+                if (i == NumPoints - 1)
+                {
+                    const FVector FirstExistingPoint = GetPosition(0);
+                    
+                    // Check if connection points are reasonably close
+                    constexpr double ConnectionTolerance = 1e-4;
+                    const double DistanceSquared = (FirstExistingPoint - Points[i].Position).SizeSquared();
+                    
+                    if (DistanceSquared <= ConnectionTolerance)
+                    {
+                        // Points are close - update first point's tangent mode if needed
+                        if (Points[i].TangentMode != ETangentMode::User && Points[i].TangentMode != ETangentMode::Broken)
+                        {
+                            SetPointTangentMode(0, Points[i].TangentMode);
+                        }
+                        continue; // Skip adding this point
+                    }
+                }
+                
+                // Prepend each point
+                PrependPoint(Points[i]);
+            }
+        }
+        
+        // Reparameterize for consistent parameter distribution
+        Reparameterize();
+        
+        return true;
     }
 
+	/**
+	 * Helper method to prepend a point to the spline.
+	 * 
+	 * @param ControlPoint Control point structure to prepend
+	 */
+	void PrependPoint(const FTangentBezierControlPoint& ControlPoint)
+	{
+		const int32 NumPoints = GetNumPoints();
+		if (NumPoints == 0)
+		{
+			// If spline is empty (shouldn't happen with our design), just use regular AddPoint
+			AddPoint(ControlPoint);
+			return;
+		}
+        
+		// First ensure the tangent modes array is large enough for existing points
+		if (TangentModes.Num() < NumPoints)
+		{
+			// Fill with default values (Auto) for any missing entries
+			TangentModes.AddDefaulted(NumPoints - TangentModes.Num());
+		}
+        
+		// Now insert the new tangent mode at the beginning
+		TangentModes.Insert(ControlPoint.TangentMode, 0);
+        
+		// Calculate control points for the cubic Bezier segment
+		FVector P0 = ControlPoint.Position;
+        
+		// Handle tangent computation based on mode
+		FVector OutTangent;
+        
+		if (ControlPoint.TangentMode == ETangentMode::User || 
+			ControlPoint.TangentMode == ETangentMode::Broken)
+		{
+			// Use the provided tangent for user-specified mode
+			OutTangent = ControlPoint.TangentOut;
+		}
+		else
+		{
+			// For auto modes, use a reasonable default that will be updated later
+			OutTangent = GetPosition(0) - ControlPoint.Position;
+		}
+        
+		FVector P1 = P0 + (OutTangent / 3.0f);
+		FVector P2 = GetPosition(0) - (GetTangentIn(0) / 3.0f);
+        
+		// Let InternalSpline handle the prepending efficiently
+		InternalSpline.PrependBezierSegment(P0, P1, P2);
+        
+		// Update tangents based on tangent mode
+		UpdatePointTangents(0, ControlPoint.TangentMode);
+        
+		// Update tangents for the next point if it's auto
+		if (NumPoints > 0 && 
+			(TangentModes[1] == ETangentMode::Auto || TangentModes[1] == ETangentMode::AutoClamped))
+		{
+			UpdatePointTangents(1, TangentModes[1]);
+		}
+	}
+    
+    /**
+     * Inserts a new point at the specified parameter along the spline.
+     * Uses distance-based parameterization to preserve attribute positions.
+     * 
+     * @param Parameter - Parameter value [0,1] where to insert the point
+     * @param Position - Position for the new point
+	 * @param InTangentMode - Tangent mode for the new point
+     * @return The index of the newly created point
+     */
+    int32 InsertPoint(float Parameter, 
+                     const FVector& Position,
+                     ETangentMode InTangentMode = ETangentMode::Auto)
+    {
+               
+        // Have the internal spline handle the insertion
+        // This efficiently splits the Bezier segment while preserving attributes
+        int32 InsertedIndex = InternalSpline.InsertPoint(Parameter, Position);
+        
+        // Convert from PolyBezier control point index to TangentBezier point index
+        int32 NewPointIndex = (InsertedIndex + 1) / 4;
+        
+		// Insert tangent mode at the correct position
+		if (NewPointIndex >= TangentModes.Num())
+		{
+			TangentModes.SetNum(NewPointIndex + 1);
+		}
+		else
+		{
+			TangentModes.Insert(InTangentMode, NewPointIndex);
+		}
+		TangentModes[NewPointIndex] = InTangentMode;
+        
+		// Update tangents to ensure proper curve continuity
+		UpdatePointTangents(NewPointIndex, InTangentMode);
+        
+		// Update adjacent points if they have automatic tangents
+		if (NewPointIndex > 0 && 
+			(TangentModes[NewPointIndex-1] == ETangentMode::Auto || 
+			 TangentModes[NewPointIndex-1] == ETangentMode::AutoClamped))
+		{
+			UpdatePointTangents(NewPointIndex-1, TangentModes[NewPointIndex-1]);
+		}
+        
+		if (NewPointIndex+1 < GetNumPoints() && 
+			(TangentModes[NewPointIndex+1] == ETangentMode::Auto || 
+			 TangentModes[NewPointIndex+1] == ETangentMode::AutoClamped))
+		{
+			UpdatePointTangents(NewPointIndex+1, TangentModes[NewPointIndex+1]);
+		}
+                
+        return NewPointIndex;
+    }
+    
     /** 
-     * Modifies an existing point
-     * @param Index - Index of the point to modify
-     * @param NewPosition - New position for the point
-     * @param NewTangentIn - New arriving tangent
-     * @param NewTangentOut - New leaving tangent
+     * Adds a new point to the spline
+	 * @param ControlPoint Structure containing position and tangent information
      */
-    void ModifyPoint(int32 Index, 
-                    const FVector& NewPosition,
-                    const FVector& NewTangentIn = FVector::ZeroVector,
-                    const FVector& NewTangentOut = FVector::ZeroVector)
+    void AddPoint(const FTangentBezierControlPoint& ControlPoint)
     {
-        const int32 NumPoints = GetNumPoints();
+        const int32 NewIndex = GetNumPoints();
+        
+		// Store tangent mode
+		if (TangentModes.Num() <= NewIndex)
+		{
+			TangentModes.AddZeroed(NewIndex - TangentModes.Num() + 1);
+		}
+		TangentModes[NewIndex] = ControlPoint.TangentMode;
+     	    
+        FVector P0 = GetPosition(NewIndex - 1);
+        FVector P3 = ControlPoint.Position;
+        
+        // Update the previous point's tangents if needed
+        if (TangentModes[NewIndex - 1] == ETangentMode::Auto || 
+            TangentModes[NewIndex - 1] == ETangentMode::AutoClamped)
+        {
+            UpdatePointTangents(NewIndex - 1, TangentModes[NewIndex - 1]);
+        }
+        
+        // Get its outgoing tangent (may have been updated above)
+        FVector OutTangentPrev = GetTangentOut(NewIndex - 1);
+        
+        // Calculate control points for the new segment
+        FVector P1 = P0 + (OutTangentPrev / 3.0f);
+        
+        // For the incoming tangent of the new point
+        FVector P2;
+        
+        if (ControlPoint.TangentMode == ETangentMode::User || 
+            ControlPoint.TangentMode == ETangentMode::Broken)
+        {
+            // For user-defined tangents, use the provided tangent
+            P2 = P3 - (ControlPoint.TangentIn / 3.0f);
+        }
+        else
+        {
+            // For other modes, we'll set up a reasonable default
+            // The UpdatePointTangents call below will update it properly
+            FVector DefaultTangent = (1.0f - Tension) * (P3 - P0);
+            P2 = P3 - (DefaultTangent / 3.0f);
+        }
+        
+        // Add the segment
+        InternalSpline.AppendBezierSegment(P1, P2, P3);
+        
+        // Now update the tangents based on the tangent mode
+        UpdatePointTangents(NewIndex, ControlPoint.TangentMode);
+        
+        // If this is auto mode, update tangents for adjacent points
+        if (ControlPoint.TangentMode == ETangentMode::Auto || 
+            ControlPoint.TangentMode == ETangentMode::AutoClamped)
+        {
+            // Update previous point's tangents again to ensure continuity
+            if (TangentModes[NewIndex - 1] == ETangentMode::Auto || 
+                TangentModes[NewIndex - 1] == ETangentMode::AutoClamped)
+            {
+                UpdatePointTangents(NewIndex - 1, TangentModes[NewIndex - 1]);
+            }
+        }
+    }
+	/**
+	 * Gets the control point at the specified index
+	 * @param Index Index of the point to retrieve
+	 * @return Control point structure with position and tangent information
+	 */
+	FTangentBezierControlPoint GetControlPoint(int32 Index) const
+	{
+		const int32 NumPoints = GetNumPoints();
+		if (Index < 0 || Index >= NumPoints)
+		{
+			return FTangentBezierControlPoint();
+		}
+    
+		FTangentBezierControlPoint Result;
+		Result.Position = GetPosition(Index);
+		Result.TangentIn = GetTangentIn(Index);
+		Result.TangentOut = GetTangentOut(Index);
+    
+		// Get tangent mode
+		if (TangentModes.IsValidIndex(Index))
+		{
+			Result.TangentMode = TangentModes[Index];
+		}
+		else
+		{
+			Result.TangentMode = ETangentMode::Auto;
+		}
+    
+		return Result;
+	}
+	
+    /**
+     * Removes a point from the spline
+     * @param Index - Index of the point to remove
+     */
+    void RemovePoint(int32 Index)
+	{
+		const int32 NumPoints = GetNumPoints();
+		if (Index < 0 || Index >= NumPoints)
+		{
+			return;
+		}
+		
+		// Prevent removing if we'd end up with fewer than 2 points
+        if (NumPoints <= 2)
+        {
+            UE_LOG(LogSpline, Warning, TEXT("Cannot remove point - spline must have at least 2 points."));
+            return;
+        }
+		
+		if (InternalSpline.RemoveSegment(Index))
+		{
+			// Remove associated auto tangent flag
+			if (TangentModes.IsValidIndex(Index))
+			{
+				TangentModes.RemoveAt(Index);
+			}
+	        
+			UpdateTangents();
+		}
+    }
+
+	/** 
+	 * Modifies an existing point
+	 * @param Index - Index of the point to modify
+	 * @param ControlPoint - New control point data
+	 */
+	void ModifyPoint(int32 Index, const FTangentBezierControlPoint& ControlPoint)
+	{
+       const int32 NumPoints = GetNumPoints();
         if (Index < 0 || Index >= NumPoints)
         {
             return;
         }
+
+        const FVector OldPosition = GetPosition(Index);
+        bool bPositionChanged = (OldPosition != ControlPoint.Position);
         
-        const bool bAutoTangent = IsAutoTangent(Index);
+        // Store the tangent mode
+        if (TangentModes.Num() <= Index)
+        {
+            TangentModes.AddZeroed(Index - TangentModes.Num() + 1);
+        }
+        TangentModes[Index] = ControlPoint.TangentMode;
         
-        // Update position and tangents
+        // Update position
         if (Index == 0)
         {
             // First point
-            UpdateBezierControlPoint(Index * 4, NewPosition);
-            
-            if (NumPoints > 1 && !bAutoTangent)
-            {
-                // Update outgoing control point
-                FVector P1 = NewPosition + (NewTangentOut / 3.0f);
-                UpdateBezierControlPoint(Index * 4 + 1, P1);
-            }
+            UpdateBezierControlPoint(Index * 4, ControlPoint.Position);
         }
         else if (Index == NumPoints - 1)
         {
             // Last point
-            UpdateBezierControlPoint(Index * 4 - 1, NewPosition);
-            
-            if (!bAutoTangent)
-            {
-                // Update incoming control point
-                FVector P2 = NewPosition - (NewTangentIn / 3.0f);
-                UpdateBezierControlPoint(Index * 4 - 2, P2);
-            }
+            UpdateBezierControlPoint(Index * 4 - 1, ControlPoint.Position);
         }
         else
         {
             // Middle point - affects two segments
-            
-            // Update the point in the previous segment
-            UpdateBezierControlPoint((Index - 1) * 4 + 3, NewPosition);
-            
-            // Update the point in the next segment
-            UpdateBezierControlPoint(Index * 4, NewPosition);
-            
-            if (!bAutoTangent)
-            {
-                // Update control points for tangents
-                FVector P2 = NewPosition - (NewTangentIn / 3.0f);
-                FVector P1 = NewPosition + (NewTangentOut / 3.0f);
-                
-                UpdateBezierControlPoint((Index - 1) * 4 + 2, P2);
-                UpdateBezierControlPoint(Index * 4 + 1, P1);
-            }
+            UpdateBezierControlPoint((Index - 1) * 4 + 3, ControlPoint.Position);
+            UpdateBezierControlPoint(Index * 4, ControlPoint.Position);
+        }
+        
+        // If not auto tangents, apply the specified tangents
+        if (ControlPoint.TangentMode != ETangentMode::Auto && 
+            ControlPoint.TangentMode != ETangentMode::AutoClamped && 
+            ControlPoint.TangentMode != ETangentMode::Linear && 
+            ControlPoint.TangentMode != ETangentMode::Constant)
+        {
+            SetTangentIn(Index, ControlPoint.TangentIn);
+            SetTangentOut(Index, ControlPoint.TangentOut);
         }
         
-        // Update auto tangents if needed
-        if (bAutoTangent)
+        // Update tangents based on tangent mode
+        UpdatePointTangents(Index, ControlPoint.TangentMode);
+        
+        // Update surrounding points if position changed and they're auto tangents
+        if (bPositionChanged)
         {
-            UpdateAutoTangents();
+            if (Index > 0 && (TangentModes[Index-1] == ETangentMode::Auto || 
+                              TangentModes[Index-1] == ETangentMode::AutoClamped))
+            {
+                UpdatePointTangents(Index-1, TangentModes[Index-1]);
+            }
+            
+            if (Index < NumPoints-1 && (TangentModes[Index+1] == ETangentMode::Auto ||
+                                       TangentModes[Index+1] == ETangentMode::AutoClamped))
+            {
+                UpdatePointTangents(Index+1, TangentModes[Index+1]);
+            }
         }
     }
 
@@ -426,6 +980,11 @@ struct FTangentBezierSpline
     /** Sets tangent in for the specified point */
     void SetTangentIn(int32 Index, const FVector& NewTangent)
     {
+		if (Index < 0 || Index >= GetNumPoints())
+		{
+			return;
+		}
+		
         if (Index == 0)
         {
             // Special case for first point in closed loop
@@ -491,31 +1050,33 @@ struct FTangentBezierSpline
     /** Checks if a point's tangents are automatically computed */
     bool IsAutoTangent(int32 Index) const
     {
-        return (AutoTangents.IsValidIndex(Index) && AutoTangents[Index]);
+		return TangentModes.IsValidIndex(Index) && TangentModes[Index] == ETangentMode::Auto;
     }
     
-    /** Sets whether a point's tangents should be automatically computed */
-    void SetPointAutoTangents(int32 Index, bool bAuto)
-    {
-        if (Index < 0 || Index >= GetNumPoints())
-        {
-            return;
-        }
+	/**
+	  * Sets the tangent mode for a specific point
+	  * @param Index - Index of the point
+	  * @param Mode - New tangent mode to set
+	  */
+	void SetPointTangentMode(int32 Index, ETangentMode Mode)
+	{
+		const int32 NumPoints = GetNumPoints();
+		if (Index < 0 || Index >= NumPoints)
+		{
+			return;
+		}
         
-        // Ensure array is big enough
-        if (AutoTangents.Num() <= Index)
-        {
-            AutoTangents.AddZeroed(Index - AutoTangents.Num() + 1);
-        }
+		// Ensure array is big enough
+		if (TangentModes.Num() <= Index)
+		{
+			TangentModes.AddZeroed(Index - TangentModes.Num() + 1);
+		}
         
-        AutoTangents[Index] = bAuto;
+		TangentModes[Index] = Mode;
         
-        // Update tangents if needed
-        if (bAuto)
-        {
-            UpdateAutoTangents();
-        }
-    }
+		// Update tangents based on new mode
+		UpdatePointTangents(Index, Mode);
+	}
 
     template <typename AttrType>
     void CreateAttributeChannel(const FName& Name)
@@ -535,41 +1096,143 @@ struct FTangentBezierSpline
         return InternalSpline.AddAttributeValue<AttrType>(Name, Value, Parameter);
     }
 
-    /** Updates auto-computed tangents based on surrounding points */
-    void UpdateAutoTangents()
-    {
-        const int32 NumPoints = GetNumPoints();
-        if (NumPoints < 2)
-        {
-            return;
-        }
+    /**
+	 * Updates tangents for all control points based on their tangent modes
+	 */
+	void UpdateTangents()
+	{
+	    const int32 NumPoints = GetNumPoints();
+	    if (NumPoints < 2)
+	    {
+	        return;
+	    }
 
-        // Update auto tangents for each point
-        for (int32 i = 0; i < NumPoints; ++i)
-        {
-            if (!IsAutoTangent(i))
-            {
-                continue;
-            }
+		// Update tangents for each point according to its mode
+		for (int32 i = 0; i < NumPoints; ++i)
+		{
+			if (TangentModes.IsValidIndex(i) && 
+				(TangentModes[i] == ETangentMode::Auto || 
+				 TangentModes[i] == ETangentMode::AutoClamped ||
+				 TangentModes[i] == ETangentMode::Linear ||
+				 TangentModes[i] == ETangentMode::Constant))
+			{
+				UpdatePointTangents(i, TangentModes[i]);
+			}
+		}
+	}
 
-            // Get previous and next points (handling closed loop case)
-            const FVector& Prev = (i > 0) ? GetPosition(i-1) : 
-                (InternalSpline.IsClosedLoop() ? GetPosition(NumPoints-1) : GetPosition(i));
-            
-            const FVector& Next = (i < NumPoints-1) ? GetPosition(i+1) :
-                (InternalSpline.IsClosedLoop() ? GetPosition(0) : GetPosition(i));
+	/**
+	 * Updates tangents for a specific point based on its tangent mode
+	 * @param Index Index of the point to update
+	 * @param Mode Tangent mode to apply
+	 */
+	void UpdatePointTangents(int32 Index, ETangentMode Mode)
+	{
+	    const int32 NumPoints = GetNumPoints();
+	    if (Index < 0 || Index >= NumPoints)
+	    {
+	        return;
+	    }
 
-            // Current point position
-            const FVector& Current = GetPosition(i);
+	    // Get the position of the current point
+	    const FVector& Current = GetPosition(Index);
+	    
+	    // Get previous and next points (handling closed loop case)
+	    const FVector& Prev = (Index > 0) ? GetPosition(Index-1) : 
+	        (InternalSpline.IsClosedLoop() ? GetPosition(NumPoints-1) : Current);
+	    
+	    const FVector& Next = (Index < NumPoints-1) ? GetPosition(Index+1) :
+	        (InternalSpline.IsClosedLoop() ? GetPosition(0) : Current);
+	    
+		// For User and Broken modes, we should preserve existing tangents
+		// So we fetch the current tangents first and only overwrite if needed
+		FVector InTangent = GetTangentIn(Index);
+		FVector OutTangent = GetTangentOut(Index);
 
-            // Compute tangents based on surrounding points
-            const FVector Tangent = (1.0f - Tension) * (Next - Prev);
-            
-            // Update tangents
-            SetTangentIn(i, Tangent);
-            SetTangentOut(i, Tangent);
-        }
-    }
+	    switch (Mode)
+	    {
+	    case ETangentMode::Auto:
+	    case ETangentMode::AutoClamped:
+	        {
+	            // Compute smooth tangent based on surrounding points
+	            FVector Tangent = (1.0f - Tension) * (Next - Prev);
+	            
+	            // Apply clamping for AutoClamped mode
+	            if (Mode == ETangentMode::AutoClamped)
+	            {
+	                // Clamp tangent to prevent overshoot
+	                const double PrevDistance = (Current - Prev).Size();
+	                const double NextDistance = (Next - Current).Size();
+	                
+	                if (PrevDistance > 0.0 && NextDistance > 0.0)
+	                {
+	                    const double MaxTangentSize = FMath::Min(PrevDistance, NextDistance) * 1.0;
+	                    const double TangentSize = Tangent.Size();
+	                    
+	                    if (TangentSize > MaxTangentSize && TangentSize > UE_KINDA_SMALL_NUMBER)
+	                    {
+	                        Tangent = Tangent * (MaxTangentSize / TangentSize);
+	                    }
+	                }
+	            }
+	            
+	            InTangent = Tangent;
+	            OutTangent = Tangent;
+	        }
+	        break;
+	        
+	    case ETangentMode::Linear:
+	        {
+	            // Compute tangents aligned with adjacent points
+	            if (Index > 0 || InternalSpline.IsClosedLoop())
+	            {
+	                InTangent = Current - Prev;
+	            }
+	            else 
+	            {
+	                InTangent = FVector::ZeroVector;
+	            }
+	            
+	            if (Index < NumPoints-1 || InternalSpline.IsClosedLoop())
+	            {
+	                OutTangent = Next - Current;
+	            }
+	            else
+	            {
+	                OutTangent = FVector::ZeroVector;
+	            }
+	        }
+	        break;
+	        
+	    case ETangentMode::Constant:
+	        {
+	            // Zero tangents for constant interpolation
+	            InTangent = FVector::ZeroVector;
+	            OutTangent = FVector::ZeroVector;
+	        }
+	        break;
+	        
+	    case ETangentMode::Broken:
+	    case ETangentMode::User:
+	        // For User and Broken modes, we're keeping the existing tangents
+	        // which we already fetched above, so no action needed here
+	        break;
+	        
+	    case ETangentMode::Unknown:
+	    default:
+	        // Default to Auto mode if unknown
+	        {
+	            FVector Tangent = (1.0f - Tension) * (Next - Prev);
+	            InTangent = Tangent;
+	            OutTangent = Tangent;
+	        }
+	        break;
+	    }
+	    
+	    // Apply the calculated tangents
+	    SetTangentIn(Index, InTangent);
+	    SetTangentOut(Index, OutTangent);
+	}
     
     float FindNearest(const FVector& Point, float& OutSquaredDistance) const
     {
@@ -582,6 +1245,11 @@ struct FTangentBezierSpline
         // Delegate to internal spline
         return InternalSpline.FindNearestOnSegment(Point, SegmentIndex, OutSquaredDistance);
     }
+
+	void Reparameterize(EParameterizationPolicy Policy = EParameterizationPolicy::Centripetal)
+    {
+	    InternalSpline.Reparameterize(Policy);
+    }
     
     /** Getter for internal spline */
     const FPolyBezierSpline3d& GetInternalSpline() const
@@ -601,4 +1269,7 @@ private:
             InternalSpline.UpdateValue(Index, NewValue);
         }
     }
-};
\ No newline at end of file
+};
+} // end namespace UE::Geometry::Spline
+} // end namespace UE::Geometry
+} // end namespace UE
\ No newline at end of file
