Fix bounds for procedural noise

diff --git a/Engine/Source/Runtime/GeometryCore/Public/ProceduralNoise.h b/Engine/Source/Runtime/GeometryCore/Public/ProceduralNoise.h
--- a/Engine/Source/Runtime/GeometryCore/Public/ProceduralNoise.h
+++ b/Engine/Source/Runtime/GeometryCore/Public/ProceduralNoise.h
@@ -40,10 +40,10 @@ T FractalBrownianMotionNoise(
 {
 	T Amplitude{ 1. };
 
-	check(Smoothness > T(0.));
-	check(Gain > T(0.));
-	check(Lacunarity > T(0.));
-	check(Gamma > T(0.));
+	checkSlow(Smoothness >= T(0.));
+	checkSlow(Gain > T(0.));
+	checkSlow(Lacunarity > T(0.));
+	checkSlow(Gamma > T(0.));
 
 	const T SmoothEps = Smoothness * T(0.01);
 	FVector2d ST(Coords.X, Coords.Y);
