Splines:
- Updating storage after reparameterization.
- Fixed issues with knot insertion/deletion which caused spline to break in various situations.
- Fixed delete point and prepend point by alt+dragging breaking the spline

diff --git a/Engine/Source/Runtime/Engine/Private/Curves/Spline.cpp b/Engine/Source/Runtime/Engine/Private/Curves/Spline.cpp
--- a/Engine/Source/Runtime/Engine/Private/Curves/Spline.cpp
+++ b/Engine/Source/Runtime/Engine/Private/Curves/Spline.cpp
@@ -183,6 +183,7 @@ public:
 	template <typename AttrType> TArray<FName> GetAttributeChannelNames() const;
 	template <typename AttrType> AttrType EvaluateAttribute(FName Name, float Parameter) const;
 	int32 SetAttributeParameter(const FName& Name, int32 Index, float ParentSpaceParameter);
+	float GetAttributeParameter(const FName& Name, int32 Index) const;
 	
 	/** Updates the spline using the current configuration. */
 	void UpdateSpline();
@@ -239,18 +240,26 @@ TArray<FName> FNewSpline::GetAttributeChannelNames() const
 template <typename AttrType>
 AttrType FNewSpline::EvaluateAttribute(FName Name, float Parameter) const
 {
-	const int32 NumEffectivePoints = IsClosedLoop() ? GetNumPoints() + 1 : GetNumPoints();
+	const int32 NumEffectivePoints = GetNumPoints();
 	Parameter = GetParameterRange().Interpolate(Parameter / (NumEffectivePoints - 1));
 	return InternalSpline.EvaluateAttribute<AttrType>(Name, Parameter);
 }
 
 int32 FNewSpline::SetAttributeParameter(const FName& Name, int32 Index, float Parameter)
 {
-	const int32 NumEffectivePoints = IsClosedLoop() ? GetNumPoints() + 1 : GetNumPoints();
+	const int32 NumEffectivePoints = GetNumPoints();
    	Parameter = GetParameterRange().Interpolate(Parameter / (NumEffectivePoints - 1));
 	return InternalSpline.SetAttributeParameter(Name, Index, Parameter);
 }
 
+float FNewSpline::GetAttributeParameter(const FName& Name, int32 Index) const
+{
+	float SplineSpaceParam = InternalSpline.GetAttributeParameter(Name, Index);
+	const float NormalizedParam = this->GetParameterRange().GetT(SplineSpaceParam);
+	const float SegmentSpaceParam = NormalizedParam * static_cast<float>(GetNumPoints() - 1);
+	return SegmentSpaceParam;
+}
+
 
 /** FLegacySpline Definition */
 
@@ -1640,7 +1649,7 @@ float FNewSpline::GetParameterAtDistance(float Distance) const
 
 float FNewSpline::GetDistanceAtParameter(float Parameter) const
 {
-    const int32 NumEffectivePoints = IsClosedLoop() ? GetNumPoints() + 1 : GetNumPoints();
+    const int32 NumEffectivePoints = GetNumPoints();
     Parameter = Parameter / (NumEffectivePoints - 1);
     return EvaluateAttribute<float>(ParamToLengthAttrName, Parameter);
 }
@@ -1680,7 +1689,7 @@ float FNewSpline::FindNearestOnSegment(const FVector& InLocation, int32 SegmentI
 
 FVector FNewSpline::Evaluate(float Parameter) const
 {
-    const int32 NumEffectivePoints = IsClosedLoop() ? GetNumPoints() + 1 : GetNumPoints();
+    const int32 NumEffectivePoints = GetNumPoints();
     Parameter = GetParameterRange().Interpolate(Parameter / (NumEffectivePoints - 1));
 	
     return FVector(GetPosition(Parameter));
@@ -1688,21 +1697,21 @@ FVector FNewSpline::Evaluate(float Parameter) const
 
 FVector FNewSpline::EvaluateDerivative(float Parameter) const
 {
-    const int32 NumEffectivePoints = IsClosedLoop() ? GetNumPoints() + 1 : GetNumPoints();
+    const int32 NumEffectivePoints = GetNumPoints();
 	Parameter = GetParameterRange().Interpolate(Parameter / (NumEffectivePoints - 1));
     return FVector(GetTangent(Parameter));
 }
 
 FQuat FNewSpline::EvaluateRotation(float Parameter) const
 {
-    const int32 NumEffectivePoints = IsClosedLoop() ? GetNumPoints() + 1 : GetNumPoints();
+    const int32 NumEffectivePoints = GetNumPoints();
 	Parameter = GetParameterRange().Interpolate(Parameter / (NumEffectivePoints - 1));
     return EvaluateAttribute<FQuat>(RotationAttrName, Parameter);
 }
 
 FVector FNewSpline::EvaluateScale(float Parameter) const
 {
-    const int32 NumEffectivePoints = IsClosedLoop() ? GetNumPoints() + 1 : GetNumPoints();
+    const int32 NumEffectivePoints = GetNumPoints();
 	Parameter = GetParameterRange().Interpolate(Parameter / (NumEffectivePoints - 1));
     return EvaluateAttribute<FVector>(ScaleAttrName, Parameter);
 }
diff --git a/Engine/Source/Runtime/GeometryCore/Public/Splines/BSpline.h b/Engine/Source/Runtime/GeometryCore/Public/Splines/BSpline.h
--- a/Engine/Source/Runtime/GeometryCore/Public/Splines/BSpline.h
+++ b/Engine/Source/Runtime/GeometryCore/Public/Splines/BSpline.h
@@ -175,39 +175,77 @@ public:
 	 * For clamped endpoints, we avoid removing knots if doing so would drop the multiplicity
 	 * below the required (Degree+1) at the endpoints.
 	 * @param Parameter - Knot value to remove
+	 * @param RemoveCompletely - If true, remove all occurrences of the knot
 	 * @param Count - Number of times to remove the knot
 	 * @return true if the knot was removed
 	 */
-	bool RemoveKnot(float Parameter, int32 Count = 1)
+	bool RemoveKnot(float Parameter, bool RemoveCompletely = false, int32 Count = 1)
 	{
-		bool bSuccess = true;
-		for (int32 i = 0; i < Count; ++i)
+		TArray OrigKnotVector = KnotVector;
+		
+		// First check if this is an endpoint knot
+		FInterval1f KnotRange = GetKnotRange();
+		bool bIsEndpoint = FMath::IsNearlyEqual(Parameter, KnotRange.Min) || 
+						   FMath::IsNearlyEqual(Parameter, KnotRange.Max);
+    
+		bool bSuccess = false;
+		// Remove all occurrences of the knot without checking multiplicity
+		if (RemoveCompletely)
+		{
+			KnotVector.Remove(Parameter);
+			bSuccess = true;
+		}
+		else
 		{
-			int32 Multiplicity = GetKnotMultiplicity(Parameter);
-			if (bClampEnds)
+			// Find all knots with this parameter value
+			TArray<int32> KnotIndices;
+			for (int32 i = 0; i < KnotVector.Num(); ++i)
 			{
-				FInterval1f KnotRange = GetKnotRange();
-				// If the parameter is at an endpoint and the multiplicity is already clamped,
-				// do not remove.
-				if ((FMath::IsNearlyEqual(Parameter, KnotRange.Min) || FMath::IsNearlyEqual(Parameter, KnotRange.Max)) &&
-					Multiplicity <= (GetDegree() + 1))
+				if (FMath::IsNearlyEqual(KnotVector[i], Parameter))
 				{
-					UE_LOG(LogSpline, Warning, TEXT("Cannot remove clamped endpoint knot: %f"), Parameter);
-					bSuccess &= false;
-					continue;
+					KnotIndices.Add(i);
 				}
 			}
-
-			int32 KnotIndex = KnotVector.IndexOfByKey(Parameter);
-			if (KnotIndex == INDEX_NONE)
+			// Check current multiplicity
+			int32 CurrentMultiplicity = KnotIndices.Num();
+			if (CurrentMultiplicity == 0)
+			{
+				UE_LOG(LogSpline, Warning, TEXT("No knots found at parameter %f"), Parameter);
+				return false;
+			}
+    
+			// Calculate how many knots we can safely remove
+			int32 MinRequiredMultiplicity = bIsEndpoint && bClampEnds ? (GetDegree() + 1) : 1;
+			int32 MaxRemovableKnots = FMath::Max(0, CurrentMultiplicity - MinRequiredMultiplicity);
+			int32 KnotsToRemove = FMath::Min(Count, MaxRemovableKnots);
+    
+			if (KnotsToRemove <= 0)
+			{
+				if (bIsEndpoint && bClampEnds)
+				{
+					UE_LOG(LogSpline, Warning, TEXT("Cannot remove endpoint knot below required multiplicity: %f (current: %d, required: %d)"), 
+						Parameter, CurrentMultiplicity, MinRequiredMultiplicity);
+				}
+				return false;
+			}
+    
+			// Remove knots in reverse order to maintain valid indices
+			
+			for (int32 i = 0; i < KnotsToRemove; ++i)
 			{
-				bSuccess &= false;
-				continue;
+				// Always remove from the last found index to avoid shifting issues
+				KnotVector.RemoveAt(KnotIndices[KnotIndices.Num() - 1 - i]);
+				bSuccess = true;
 			}
-			KnotVector.RemoveAt(KnotIndex);
-			bSuccess &= true;
 		}
-	    return bSuccess;
+		// validate if removing completely adheres to knot vector requirements
+		if (!UE::Geometry::Spline::Math::FSplineValidation::IsValidKnotVector(KnotVector, Degree))
+		{
+			UE_LOG(LogSpline, Warning, TEXT("Invalid knot vector after removing knot(s) %f"), Parameter);
+			KnotVector = OrigKnotVector;
+			return false;
+		}
+		return bSuccess;
 	}
 	
 	/**
@@ -374,6 +412,40 @@ public:
 		return this->GetKnotRange().Interpolate(localT);
     	
     }
+	void FixEndKnots(int32 NumKeys)
+	{
+		// Calculate how many knots to add to reach desired multiplicity
+		const int32 ExpectedKnots = GetExpectedNumKnots(NumKeys);
+		int32 Diff = FMath::Abs(KnotVector.Num() - ExpectedKnots);
+		if (KnotVector.Num() > ExpectedKnots)
+		{
+			KnotVector.SetNum(ExpectedKnots);
+		}
+		else if (KnotVector.Num() < ExpectedKnots)
+		{
+			float ClampedKnot = KnotVector.Last();
+			for (int32 i = 0; i < Diff; ++i)
+			{
+				KnotVector.Add(ClampedKnot);
+			}
+		}
+	}
+	void FillEndKnots(float NewEndKnot, int32 NumKeys)
+	{
+		// Calculate how many knots to add to reach desired multiplicity
+		const int32 ExpectedKnots = GetExpectedNumKnots(NumKeys); 
+		int32 DesiredMultiplicity = ExpectedKnots - KnotVector.Num(); 
+		int32 KnotsToAdd = FMath::Max(0, DesiredMultiplicity);
+        
+		UE_LOG(LogSpline, Log, TEXT("Adding %d knots at new end parameter %f to reach multiplicity %d"),
+			   KnotsToAdd, NewEndKnot, DesiredMultiplicity);
+        
+		// Insert knots to reach desired multiplicity
+		for (int32 i = 0; i < KnotsToAdd; ++i)
+		{
+			KnotVector.Add(NewEndKnot);
+		}
+	}
 	
 protected:
 	int32 GetNumValuesFromKnots() const
@@ -868,28 +940,17 @@ public:
 	 */
 	void Reparameterize(EParameterizationPolicy Mode = EParameterizationPolicy::Uniform)
 	{
+		// Store current parameter range
+		FInterval1f OldRange = this->GetParameterRange();
+    
+		// Update the knot vector
 		Base::GetEvaluator().UpdateKnotVector(Base::Values.GetStorage().GetEntries(), Mode);
-
-		// Update the parameter range to match the new knot range
-    	
-    	FInterval1f KnotRange = this->GetEvaluator().GetKnotRange();
     
-    	// Only update if the knot range is valid
-    	if (!KnotRange.IsEmpty())
-    	{
-    		FInterval1f CurrentRange = this->GetEvaluator().GetParentSpaceRange();
-
-    		// Check if the knot range is significantly different from parameter range
-    		const float Tolerance = 0.01f; // Small tolerance to avoid frequent small updates
-    		if (FMath::Abs(KnotRange.Min - CurrentRange.Min) > Tolerance || 
-    			FMath::Abs(KnotRange.Max - CurrentRange.Max) > Tolerance)
-    		{
-    			// Update parameter range to match knot range
-    			this->GetEvaluator().SetParentSpaceRange(KnotRange);
-    			UE_LOG(LogSpline, Verbose, TEXT("Updated parameter range to [%f,%f] to match knot range"),
-    				KnotRange.Min, KnotRange.Max);
-    		}
-    	}
+		// Get new parameter range
+		FInterval1f NewRange = this->GetParameterRange();
+    
+		// Update parameter values in storage
+		this->ReparameterizeStorage(OldRange, NewRange);
 	}
 
 	virtual float FindNearest(const ValueType& Point, float& OutSquaredDistance) const override
diff --git a/Engine/Source/Runtime/GeometryCore/Public/Splines/PolyBezierSpline.h b/Engine/Source/Runtime/GeometryCore/Public/Splines/PolyBezierSpline.h
--- a/Engine/Source/Runtime/GeometryCore/Public/Splines/PolyBezierSpline.h
+++ b/Engine/Source/Runtime/GeometryCore/Public/Splines/PolyBezierSpline.h
@@ -89,13 +89,15 @@ public:
         {
             return Window;
         }
-
-        // Map parameter to parent space and handle bounds
-        const float ParentParam = this->MapToParentSpace(Parameter);
-        
+       
         // Find which segment we're in
     	const int32 NumSegments = Entries.Num() / 4;
-        const float ClampedParam = FMath::Clamp(ParentParam, 0.0f, NumSegments);
+
+    	// Map parameter to parent space
+        const float ParentParam = this->MapToParentSpace(Parameter);
+        const float NormalizedParam = this->GetKnotRange().GetT(ParentParam);
+        const float SegmentSpaceParam = NormalizedParam * static_cast<float>(NumSegments);
+    	const float ClampedParam = FMath::Clamp(SegmentSpaceParam, 0.0f, NumSegments);
         int32 SegmentIndex = FMath::Min(FMath::FloorToInt(ClampedParam), static_cast<int32>(NumSegments) - 1);
 
         if (this->IsClosedLoop())
@@ -146,19 +148,21 @@ public:
             return ValueType();
         }
 
-        // Map to parent space (segment space)
-        const float ParentParam = this->MapToParentSpace(Parameter);
-
-    	const int32 NumValues = this->GetNumValuesFromKnots();
+       	const int32 NumValues = this->GetNumValuesFromKnots();
     	int32 NumSegments = NumValues / 4;
+
+    	// Map to parent space
+        const float ParentParam = this->MapToParentSpace(Parameter);
+        const float NormalizedParam = this->GetKnotRange().GetT(ParentParam);
+        const float SegmentSpaceParam = NormalizedParam * static_cast<float>(NumSegments);
         // If we're at the last point exactly, just return it
-        if (ParentParam >= static_cast<float>(NumSegments))
+        if (SegmentSpaceParam >= static_cast<float>(NumSegments))
         {
             return *Window[3];  // Last control point (P3) of final Bezier segment
         }
-        
-        // Get local parameter within segment [0,1]
-        const float LocalT = FMath::Frac(ParentParam);
+    	const float ClampedParam = FMath::Clamp(SegmentSpaceParam, 0.0f, NumSegments);
+    	int32 SegmentIndex = FMath::FloorToInt(ClampedParam);
+    	const float LocalT = ClampedParam - static_cast<float>(SegmentIndex);
 
         // Calculate Bernstein basis functions
         const float t = LocalT;
@@ -195,9 +199,6 @@ public:
 	    {
 	        return InterpolateWindow(FindWindow(Entries, Parameter), Parameter);
 	    }
-
-	    // Map to parent space
-	    const float MappedParam = Base::MapToParentSpace(Parameter);
 	    
 		// Get window for the current parameter
 	    TArray<const ValueType*> Window = FindWindow(Entries, Parameter);
@@ -207,8 +208,12 @@ public:
 	    }
 
     	int32 NumSegments = Entries.Num() / 4;
+    	// Map parameter to parent space
+    	const float ParentParam = this->MapToParentSpace(Parameter);
+    	const float NormalizedParam = this->GetKnotRange().GetT(ParentParam);
+    	const float SegmentSpaceParam = NormalizedParam * static_cast<float>(NumSegments);
 	    // Find segment index and local parameter
-	    const float ClampedParam = FMath::Clamp(MappedParam, 0.0f, static_cast<float>(NumSegments));
+	    const float ClampedParam = FMath::Clamp(SegmentSpaceParam, 0.0f, static_cast<float>(NumSegments));
 	    int32 SegmentIndex;
 		float LocalT;
     	
@@ -248,12 +253,15 @@ public:
     		return InterpolateWindow(FindWindow(Entries, Parameter), Parameter);
     	}
 
-    	// Map to parent space
-    	const float MappedParam = Base::MapToParentSpace(Parameter);
-    
     	// Find segment index and local parameter
     	int32 NumSegments = Entries.Num() / 4;
-    	const float ClampedParam = FMath::Clamp(MappedParam, 0.0f, static_cast<float>(NumSegments));
+
+    	// Map parameter to parent space
+    	const float ParentParam = this->MapToParentSpace(Parameter);
+    	const float NormalizedParam = this->GetKnotRange().GetT(ParentParam);
+    	const float SegmentSpaceParam = NormalizedParam * static_cast<float>(NumSegments);
+    	
+    	const float ClampedParam = FMath::Clamp(SegmentSpaceParam, 0.0f, static_cast<float>(NumSegments));
     	const int32 SegmentIndex = FMath::Min(FMath::FloorToInt(ClampedParam), NumSegments - 1);
     	const float LocalT = ClampedParam - static_cast<float>(SegmentIndex); // Parameter [0,1] within segment
     
@@ -279,50 +287,47 @@ public:
 			Window, KnotVector, Span, LocalT, 3);
     }
 	
-    // virtual float GetParentSpaceMapping(float ChildParam) const override
-    // {
-    // 	float localT = this->GetKnotRange.GetT(ChildParam);
-    // 	return this->GetParentSpaceRange().Interpolate(localT);
-    // 	// int32 NumValues = this->GetNumValuesFromKnots();
-    //  //    int32 NumSegments = NumValues / 4;
-    //  //    return this->GetParentSpaceRange().GetT(ChildParam) * static_cast<float>(NumSegments);
-    // }
-    //
-    // virtual float GetChildSpaceMapping(float ParentParam) const override
-    // {
-    // 	float localT = this->GetParentSpaceRange().GetT(ParentParam);
-    // 	return this->GetKnotRange().Interpolate(localT);
-    // 	
-    // 	// int32 NumValues = this->GetNumValuesFromKnots();
-    // 	// int32 NumSegments = NumValues / 4;
-    //  //    float NormalizedParam = NumSegments > 0 ? ParentParam / static_cast<float>(NumSegments) : 0.0f;
-    // 	
-    // }
-
 	/**
 	 * Helper to insert local knots for a segment boundary.
 	 * For clamped splines, internal boundaries typically should have a multiplicity equal
 	 * to the degree rather than the clamped endpoint multiplicity.
 	 * @param StartParam - Start parameter of the segment
 	 * @param EndParam - End parameter of the segment
+	 * @param NumKeys - Total number of keys in the spline
 	 */
-	void InsertLocalKnots(float StartParam, float EndParam)
+	void InsertLocalKnots(float StartParam, float EndParam, int32 NumKeys)
 	{
 		if (this->bClampEnds)
 		{
-			// Desired multiplicity for an internal knot might be Degree
-			const int32 DesiredInternalMultiplicity = GetDegree();
+			// Get parameter range to determine if these are endpoints
+			FInterval1f KnotRange = this->GetKnotRange();
+        
+			// Check if StartParam is at the beginning of the knot range
+			bool bIsStartEndpoint = FMath::IsNearlyEqual(StartParam, KnotRange.Min) || StartParam < KnotRange.Min;
+        
+			// Check if EndParam is at the end of the knot range
+			bool bIsEndEndpoint = FMath::IsNearlyEqual(EndParam, KnotRange.Max) || EndParam > KnotRange.Max;
+        
+			// For endpoints, we want multiplicity of degree+1 for clamped splines
+			const int32 EndpointMultiplicity = this->GetDegree() + 1;
+        
+			// For internal knots, we typically want multiplicity of degree
+			const int32 InternalMultiplicity = this->GetDegree();
 
 			int32 StartMult = this->GetKnotMultiplicity(StartParam);
-			if (StartMult < DesiredInternalMultiplicity)
+			int32 DesiredStartMult = bIsStartEndpoint ? EndpointMultiplicity : InternalMultiplicity;
+        
+			if (StartMult < DesiredStartMult)
 			{
-				this->InsertKnot(StartParam, DesiredInternalMultiplicity - StartMult);
+				this->InsertKnot(StartParam, DesiredStartMult - StartMult);
 			}
 
 			int32 EndMult = this->GetKnotMultiplicity(EndParam);
-			if (EndMult < DesiredInternalMultiplicity)
+			int32 DesiredEndMult = bIsEndEndpoint ? EndpointMultiplicity : InternalMultiplicity;
+        
+			if (EndMult < DesiredEndMult)
 			{
-				this->InsertKnot(EndParam, DesiredInternalMultiplicity - EndMult);
+				this->InsertKnot(EndParam, DesiredEndMult - EndMult);
 			}
 		}
 		else
@@ -332,7 +337,9 @@ public:
 			this->InsertKnot(EndParam, 1);
 		}
     	
-    	// Update parameter range
+    	this->FixEndKnots(NumKeys);
+    	
+    	// Update parent space range
     	this->SetParentSpaceRange(this->GetKnotRange());
 	}
 private:
@@ -404,9 +411,9 @@ private:
                 AccumLength += SegmentLengths[i];
                 
                 // Use normalized accumulated length as the knot value
-                const float t = TotalLength > UE_SMALL_NUMBER ? 
+                const float t = AccumLength;/*TotalLength > UE_SMALL_NUMBER ? 
                     (AccumLength / TotalLength) * static_cast<float>(NumSegments) : // Scale to match segment count
-                    static_cast<float>(i + 1);										// Fallback to uniform if total length ~ 0
+                    static_cast<float>(i + 1);										// Fallback to uniform if total length ~ 0*/
                 
                 // Interior knots at segment boundaries with multiplicity 3
                 for (int32 j = 0; j < 3; ++j)
@@ -417,7 +424,8 @@ private:
         }
 
         // End knots with multiplicity 4
-        const float EndKnot = static_cast<float>(NumSegments);
+        // const float EndKnot = static_cast<float>(NumSegments);
+    	const float EndKnot = AccumLength + (SegmentLengths.Num() > 0 ? SegmentLengths.Last() : 0.0f);
         while (CurrentKnot < NumKnots)
         {
             DistanceKnots[CurrentKnot++] = EndKnot;
@@ -1049,7 +1057,6 @@ public:
 		{
 			// Prepend
 			PrependPoint(Position, ParameterizationPolicy);
-			Reparameterize(ParameterizationPolicy);
 			return 0;
 		}
     
@@ -1066,11 +1073,12 @@ public:
         
 			// Add the segment
 			AddBezierSegmentInternal(P0, P1, P2, Position, StartParam, ParameterizationPolicy);
-			Reparameterize(ParameterizationPolicy);
 			return GetNumSegments();
 		}
-            	
-    	const float ClampedParam = FMath::Clamp(Parameter, 0.0f, static_cast<float>(NumSegments));
+
+		const float NormalizedParam = ParameterRange.GetT(Parameter);
+		const float SegmentSpaceParam = NormalizedParam * static_cast<float>(NumSegments);
+    	const float ClampedParam = FMath::Clamp(SegmentSpaceParam, 0.0f, static_cast<float>(NumSegments));
     	const int32 SegmentIndex = FMath::Min(FMath::FloorToInt(ClampedParam), NumSegments - 1);
     	const float LocalT = ClampedParam - static_cast<float>(SegmentIndex); // Parameter [0,1] within segment
     	
@@ -1103,13 +1111,20 @@ public:
         float OriginalEndParam = Base::GetParameterAt(BaseIndex+3);
 		float TargetSplitParam = OriginalStartParam + LocalT * (OriginalEndParam - OriginalStartParam);
     
-		// Remove the original segment
-		RemoveSegmentInternal(SegmentIndex);
+		// Instead of removing the original segment, let's directly insert a knot with multiplicity 3
+		// at the split point, which is the proper structure for an internal Bezier junction
+		this->GetEvaluator().InsertKnot(TargetSplitParam, 3);
+    
+		// Now we can safely remove the control points
+		// We're not touching the knot structure anymore - just the control points
+		for (int32 i = 0; i < 4; ++i)
+		{
+			Base::RemoveValue(BaseIndex);
+		}
 		
         // Add the two new segments using the policy
         AddBezierSegmentInternal(Q0, Q1, Q2, Q3, OriginalStartParam, ParameterizationPolicy, TargetSplitParam);
 		AddBezierSegmentInternal(R0, R1, R2, R3, TargetSplitParam, ParameterizationPolicy, OriginalEndParam);
-		Reparameterize(ParameterizationPolicy);
 				
         return InsertedPointIndex;
     }
@@ -1193,7 +1208,7 @@ public:
         int32 Result = AddBezierSegmentInternal(PositionAtParam, P1, P2, P3, Parameter, ParameterizationPolicy);
 
     	// Update the knot vector by inserting the new boundary.
-    	this->GetEvaluator().InsertLocalKnots(Parameter, Parameter);
+    	this->GetEvaluator().InsertLocalKnots(Parameter, Parameter, Base::GetNumKeys());
     	
         return Result;
     }
@@ -1376,9 +1391,17 @@ public:
 	 */
 	void Reparameterize(EParameterizationPolicy Mode = EParameterizationPolicy::Centripetal)
     {
-    	// Pass the control points to the evaluator when updating segment structure
-    	// This allows distance-based knot generation to work properly
-    	this->GetEvaluator().Reparameterize(Mode, this->Values.GetStorage().GetEntries());
+		// Store current parameter range
+		FInterval1f OldRange = this->GetParameterRange();
+    
+		// Reparameterize the curve
+		this->GetEvaluator().Reparameterize(Mode, Base::Values.GetStorage().GetEntries());
+    
+		// Get new parameter range
+		FInterval1f NewRange = this->GetParameterRange();
+    
+		// Update parameter values in storage
+		this->ReparameterizeStorage(OldRange, NewRange);
     }
     
 protected:
@@ -1401,43 +1424,82 @@ private:
 	}
 	
 	/**
-	 * Computes an appropriate parameter for a new segment based on insertion position.
-	 * For all parameterization modes, segments occupy 1.0 unit of parameter space.
-	 * 
-	 * @param RequestedParam - The requested parameter value (used to determine prepend/append/insert)
-	 * @return The parameter value to use for the new segment
+	 * Computes an appropriate parameter for a new segment based on insertion position and parameterization policy
+	 * @param P0 Start point of the segment
+	 * @param P1 First control point
+	 * @param P2 Second control point
+	 * @param P3 End point of the segment
+	 * @param RequestedParam The requested parameter (used to determine prepend vs append)
+	 * @param ParameterizationPolicy The parameterization policy to use
+	 * @return A properly computed parameter value for the segment
 	 */
-	float ComputeNewSegmentParameter(float RequestedParam)
-    {
-    	// Get current parameter range
-    	FInterval1f CurrentParameterRange = this->GetEvaluator().GetParentSpaceRange();
-    
-    	// Handle special case for empty splines
-    	bool bIsEmpty = (this->GetNumSegments() == 0);
-    	if (bIsEmpty)
-    	{
-    		return FMath::IsNearlyZero(RequestedParam) ? 0.0f : RequestedParam;
-    	}
-    
-    	// Check if parameter is within current range
-    	if (CurrentParameterRange.Contains(RequestedParam))
-    	{
-    		// For insertion within the curve, use the requested parameter directly
-    		return RequestedParam;
-    	}
-    
-    	// Determine prepend or append
-    	if (RequestedParam < CurrentParameterRange.Min)
-    	{
-    		// Prepend case - place 1.0 unit before the start
-    		return CurrentParameterRange.Min - 1.0f;
-    	}
-    	else
-    	{
-    		// Append case - place at NumSegments (end of the current parameter range)
-    		return static_cast<float>(GetNumSegments());
-    	}
-    }
+	float ComputeNewSegmentParameter(
+	    const ValueType& P0,
+	    const ValueType& P1, 
+	    const ValueType& P2, 
+	    const ValueType& P3,
+	    float RequestedParam,
+	    EParameterizationPolicy ParameterizationPolicy)
+	{
+	    // Get current parameter range
+	    FInterval1f CurrentParameterRange = this->GetParameterRange();
+	    
+	    // Handle special case for empty splines
+	    bool bIsEmpty = (this->GetNumSegments() == 0);
+	    
+	    // If the spline is empty, just use the requested parameter (or 0 if near zero)
+	    if (bIsEmpty)
+	    {
+	        return FMath::IsNearlyZero(RequestedParam) ? 0.0f : RequestedParam;
+	    }
+	    
+	    // Check if parameter is within current range
+	    if (CurrentParameterRange.Contains(RequestedParam))
+	    {
+	        // For insertion within the curve, use the requested parameter directly
+	        return RequestedParam;
+	    }
+	    
+	    // Handle out-of-bounds cases
+	    float NewParam;
+	    float SegmentLength;
+	    
+	    // Calculate segment length based on parameterization policy
+	    switch (ParameterizationPolicy)
+	    {
+	    case EParameterizationPolicy::ChordLength:
+	        SegmentLength = static_cast<float>((P3 - P0).Size());
+	        break;
+	    case EParameterizationPolicy::Centripetal:
+	        SegmentLength = FMath::Sqrt(static_cast<float>((P3 - P0).Size()));
+	        break;
+	    case EParameterizationPolicy::Uniform:
+	    default:
+	        SegmentLength = 1.0f;
+	        break;
+	    }
+	    
+	    // Determine prepend or append
+	    if (RequestedParam < CurrentParameterRange.Min)
+	    {
+	        // Prepend case - parameter before current start
+	        NewParam = CurrentParameterRange.Min - SegmentLength;
+	    }
+	    else
+	    {
+	        // Append case - parameter after current end
+	        NewParam = CurrentParameterRange.Max + SegmentLength;
+	        
+	        // For uniform parameterization, use integer values for append
+	        if (ParameterizationPolicy == EParameterizationPolicy::Uniform)
+	        {
+	            const int32 NumSegments = GetNumSegments();
+	            NewParam = static_cast<float>(NumSegments + 1);
+	        }
+	    }
+	    
+	    return NewParam;
+	}
 
 	/**
 	 * Computes the end parameter for a segment given the start parameter and policy
@@ -1489,7 +1551,7 @@ private:
      * @param P2 Second control point
      * @param P3 End point
      * @param StartParam Start parameter for the segment
-     * @param EndParam Optional End parameter for the segment, or -1 for automatic calculation
+     * @param EndParam Optional End parameter for the segment
      * @param ParameterizationPolicy How to distribute parameters
      * @return Index of the new segment
      */
@@ -1500,30 +1562,34 @@ private:
 		const ValueType& P3,
 		float StartParam,
 		EParameterizationPolicy ParameterizationPolicy = EParameterizationPolicy::Centripetal,
-		float EndParam = -1.0f) // Default to -1 for automatic calculation
+		TOptional<float> EndParam = TOptional<float>())
     {
-    	// Compute the appropriate start parameter based on requested value
-    	StartParam = ComputeNewSegmentParameter(StartParam);
-    
+    	// Compute the appropriate start parameter based on requested value and policy
+    	StartParam = ComputeNewSegmentParameter(P0, P1, P2, P3, StartParam, ParameterizationPolicy);
+		float EndParamValue;
     	// If EndParam is provided, use it; otherwise, calculate it
-    	if (EndParam < 0.0f || EndParam < StartParam)
+    	if (!EndParam.IsSet() || EndParam.GetValue() < StartParam)
     	{
-    		EndParam = StartParam + 1.0f;
+    		EndParamValue = ComputeEndParameter(StartParam, P0, P3, ParameterizationPolicy);
     	}
+	    else
+	    {
+		    EndParamValue = EndParam.GetValue();
+	    }
 
     	// Calculate intermediate control point parameters
     	float IntermediateParams[2] = {
-    		StartParam + (EndParam - StartParam) * 0.33f,
-			StartParam + (EndParam - StartParam) * 0.67f
+    		StartParam + (EndParamValue - StartParam) * 0.33f,
+			StartParam + (EndParamValue - StartParam) * 0.67f
 		};
 
     	// Add the control points
     	int32 P0Index = Base::AddValue(P0, StartParam);
     	Base::AddValue(P1, IntermediateParams[0]);
     	Base::AddValue(P2, IntermediateParams[1]);
-    	Base::AddValue(P3, EndParam);
+    	Base::AddValue(P3, EndParamValue);
     
-    	this->GetEvaluator().InsertLocalKnots(StartParam, EndParam);
+    	this->GetEvaluator().InsertLocalKnots(StartParam, EndParamValue, Base::GetNumKeys());
     
     	return P0Index / 4;
     }
@@ -1535,30 +1601,110 @@ private:
      */
     bool RemoveSegmentInternal(const int32 SegmentIndex)
     {
-    	if (SegmentIndex < 0 || SegmentIndex >= GetNumSegments())
-    	{
-    		return false;
-    	}
+    	const int32 NumSegments = GetNumSegments();
+	    if (SegmentIndex < 0 || SegmentIndex >= NumSegments)
+	    {
+	        return false;
+	    }
 
-    	const int32 BaseIndex = SegmentIndex * 4;
-    	// Get the parameters for the removed segment.
-    	float StartParam = Base::GetParameterAt(BaseIndex);
-    	float EndParam = Base::GetParameterAt(BaseIndex + 3);
-    	
-    	bool bSuccess = true;
-    	bSuccess &= Base::RemoveValue(BaseIndex + 3);
-    	bSuccess &= Base::RemoveValue(BaseIndex + 2);
-    	bSuccess &= Base::RemoveValue(BaseIndex + 1);
-    	bSuccess &= Base::RemoveValue(BaseIndex);
+	    // Handle special case - can't remove the only segment
+	    if (NumSegments == 1)
+	    {
+	        UE_LOG(LogSpline, Warning, TEXT("Cannot remove the only segment from a spline"));
+	        return false;
+	    }
 
-    	if (bSuccess)
+    	// Calculate storage indices
+    	const int32 BaseIndex = SegmentIndex * 4;
+    	const bool bIsFirstSegment = (SegmentIndex == 0);
+    	const bool bIsLastSegment = (SegmentIndex == NumSegments - 1);
+    
+    	// Get the parameters at segment boundaries
+    	auto& Storage = Base::Values.GetStorage().GetEntries();
+    	float StartParam = Storage[BaseIndex].Parameter;
+    	float EndParam = Storage[BaseIndex + 3].Parameter;
+    
+    	UE_LOG(LogSpline, Log, TEXT("Removing segment %d: Start=%f, End=%f"), 
+			   SegmentIndex, StartParam, EndParam);
+    
+    	// Get references to the evaluator and knot vector
+    	auto& Evaluator = this->GetEvaluator();
+    	const TArray<float>& KnotVector = Evaluator.GetKnotVector();
+    	const int32 Degree = Evaluator.GetDegree(); // Typically 3 for cubic
+    
+    	// Store new boundary parameters before removing anything
+    	float NewStartParam = 0.0f;
+    	float NewEndParam = 0.0f;
+    
+    	// For first segment, the new start parameter will be the end of removed segment
+    	if (bIsFirstSegment && SegmentIndex + 1 < NumSegments)
+    	{
+    		NewStartParam = EndParam;
+    	}
+    	// For last segment, the new end parameter will be the start of removed segment
+    	else if (bIsLastSegment && SegmentIndex > 0)
     	{
-    		// Remove the corresponding knots.
-    		// Here we remove one occurrence at each boundary.
-    		this->GetEvaluator().RemoveKnot(StartParam);
-    		this->GetEvaluator().RemoveKnot(EndParam);
+    		NewEndParam = StartParam;
     	}
-    	return bSuccess;
+    
+    	// Now remove points from storage
+    	bool bSuccess = true;
+	    
+	    if (bIsFirstSegment)
+	    {
+	        for (int32 i = 3; i >= 0; --i)
+	        {
+	            bSuccess &= Base::RemoveValue(BaseIndex + i);
+	        }
+	    	// Remove all knots at original start parameter
+	    	Evaluator.RemoveKnot(StartParam, true);
+	    	        
+	    	// The new start parameter becomes the end parameter of the removed segment
+        
+	    	// First count existing multiplicity
+	    	int32 NewStartMultiplicity = Evaluator.GetKnotMultiplicity(NewStartParam);
+        
+	    	// Calculate how many knots to add to reach desired multiplicity
+	    	int32 DesiredMultiplicity = Degree + 1; // Typically 4 for cubic
+	    	int32 KnotsToAdd = FMath::Max(0, DesiredMultiplicity - NewStartMultiplicity);
+        
+	    	UE_LOG(LogSpline, Log, TEXT("Adding %d knots at new start parameter %f to reach multiplicity %d"),
+				   KnotsToAdd, NewStartParam, DesiredMultiplicity);
+        
+	    	// Insert knots to reach desired multiplicity
+	    	for (int32 i = 0; i < KnotsToAdd; ++i)
+	    	{
+	    		Evaluator.InsertKnot(NewStartParam);
+	    	}
+	    }
+	    else if (bIsLastSegment)
+	    {
+	    	for (int32 i = 3; i >= 0; --i)
+	    	{
+	    		bSuccess &= Base::RemoveValue(BaseIndex + i);
+	    	}
+	        
+	    	// Remove all knots at original end parameter
+	    	Evaluator.RemoveKnot(EndParam, true);
+	    	Evaluator.FillEndKnots(NewEndParam, Base::GetNumKeys());
+	    	
+	    }
+	    else
+	    {
+	        // Remove the internal segment's P0 (duplicate), P1, and P2 - keep P3 (which is P0 of next)
+	        for (int32 i = 3; i >= 0; --i)
+	        {
+	            bSuccess &= Base::RemoveValue(BaseIndex + i - 1);
+	        }
+	    	// For internal segments, just remove all knots at the start parameter
+	    	Evaluator.RemoveKnot(StartParam, true);
+	    	
+	    }
+    	Evaluator.FixEndKnots(Base::GetNumKeys());
+	    // Update parent space range to reflect the new knot vector
+	    this->GetEvaluator().SetParentSpaceRange(this->GetEvaluator().GetKnotRange());
+	    
+	    return bSuccess;
     }
 };
 	
diff --git a/Engine/Source/Runtime/GeometryCore/Public/Splines/Spline.h b/Engine/Source/Runtime/GeometryCore/Public/Splines/Spline.h
--- a/Engine/Source/Runtime/GeometryCore/Public/Splines/Spline.h
+++ b/Engine/Source/Runtime/GeometryCore/Public/Splines/Spline.h
@@ -768,6 +768,63 @@ protected:
 			static_cast<TAttributeChannel<ValueType>*>(Pair.Value.Get())->Clear();
 		}
 	}
+
+	/**
+     * Updates value and attribute storage after reparameterization
+     * @param OldRange The parameter range before reparameterization
+     * @param NewRange The parameter range after reparameterization
+     */
+    void ReparameterizeStorage(const FInterval1f& OldRange, const FInterval1f& NewRange)
+    {
+        // Early out if ranges are nearly identical
+        if (FMath::IsNearlyEqual(OldRange.Min, NewRange.Min) &&
+            FMath::IsNearlyEqual(OldRange.Max, NewRange.Max))
+        {
+            return;
+        }
+        
+        // 1. Update main value storage
+        auto& ValueEntries = Values.GetStorage().GetEntries();
+        for (auto& Entry : ValueEntries)
+        {
+            // Convert parameter to normalized space and back to new range
+            float NormalizedParam = OldRange.GetT(Entry.Parameter);
+            Entry.Parameter = NewRange.Interpolate(NormalizedParam);
+        }
+        
+        // 2. Update attribute channels
+        for (auto& ChannelPair : AttributeChannels)
+        {
+            IAttributeChannelBase* Channel = ChannelPair.Value.Get();
+            
+            // Need to store parameters before modifying them as SetParameter can reorder entries
+            TArray<TTuple<int32, float>> IndicesAndParams;
+            for (int32 i = 0; i < Channel->NumKeys(); ++i)
+            {
+                float LocalParam = Channel->GetParameter(i);
+                float ParentParam = Channel->MapToParentSpace(LocalParam);
+                IndicesAndParams.Add(TTuple<int32, float>(i, ParentParam));
+            }
+            
+            // Update each parameter value
+            for (const auto& IndexAndParam : IndicesAndParams)
+            {
+                int32 Index = IndexAndParam.Get<0>();
+                float ParentParam = IndexAndParam.Get<1>();
+                
+                // Remap to new parameter space
+                float NormalizedParam = OldRange.GetT(ParentParam);
+                float NewParentParam = NewRange.Interpolate(NormalizedParam);
+                float NewLocalParam = Channel->MapFromParentSpace(NewParentParam);
+                
+                // Update parameter
+                Channel->SetParameter(Index, NewLocalParam);
+            }
+            
+            // Update parent evaluator reference
+            Channel->UpdateEvaluatorParent(&GetEvaluator());
+        }
+    }
 private:
 	/**
 	 * Creates a new instance of the correct derived type for cloning.
diff --git a/Engine/Source/Runtime/GeometryCore/Public/Splines/TangentBezierSpline.h b/Engine/Source/Runtime/GeometryCore/Public/Splines/TangentBezierSpline.h
--- a/Engine/Source/Runtime/GeometryCore/Public/Splines/TangentBezierSpline.h
+++ b/Engine/Source/Runtime/GeometryCore/Public/Splines/TangentBezierSpline.h
@@ -568,10 +568,7 @@ public:
 	    
 	    // Update tangents to ensure proper curve continuity
 	    UpdateTangents();
-	    
-	    // Reparameterize the spline for better parameter distribution
-	    Reparameterize();
-	    
+	    	    
 	    return true;
 	}
 
@@ -848,8 +845,6 @@ public:
                 UpdatePointTangents(NewIndex - 1, TangentModes[NewIndex - 1]);
             }
         }
-
-		InternalSpline.Reparameterize();
     }
 	/**
 	 * Gets the control point at the specified index
@@ -893,24 +888,63 @@ public:
 		{
 			return;
 		}
-		
+    
 		// Prevent removing if we'd end up with fewer than 2 points
-        if (NumPoints <= 2)
-        {
-            UE_LOG(LogSpline, Warning, TEXT("Cannot remove point - spline must have at least 2 points."));
-            return;
-        }
-		
-		if (InternalSpline.RemoveSegment(Index))
+		if (NumPoints <= 2)
+		{
+			UE_LOG(LogSpline, Warning, TEXT("Cannot remove point - spline must have at least 2 points."));
+			return;
+		}
+
+		// Remove the corresponding tangent mode
+		if (TangentModes.IsValidIndex(Index))
+		{
+			TangentModes.RemoveAt(Index);
+		}
+    
+		if (Index == 0)
 		{
-			// Remove associated auto tangent flag
-			if (TangentModes.IsValidIndex(Index))
+			// Removing the first point - remove the first segment
+			InternalSpline.RemoveSegment(0);
+		}
+		else if (Index == NumPoints - 1)
+		{
+			// Removing the last point - remove the last segment
+			InternalSpline.RemoveSegment(InternalSpline.GetNumSegments() - 1);
+		}
+		else
+		{
+			// Removing an internal point
+			// Get the parameter of the point being removed
+			float PointParam = InternalSpline.GetParameterAt((Index - 1) * 4 + 3);
+        
+			// For an internal point, we need to remove exactly 4 control points:
+			// - P2 and P3 from the previous segment (where P3 is the point)
+			// - P0 and P1 from the next segment (where P0 is the same point)
+        
+			// Start index for removal: P2 of previous segment
+			int32 StartRemoveIndex = (Index - 1) * 4 + 2;
+        
+			// Remove 4 consecutive control points
+			for (int32 i = 0; i < 4; ++i)
 			{
-				TangentModes.RemoveAt(Index);
+				InternalSpline.RemoveValue(StartRemoveIndex);
 			}
-	        
-			UpdateTangents();
+        
+			// Remove the knot associated with this point
+			InternalSpline.GetEvaluator().RemoveKnot(PointParam, true);
+			InternalSpline.GetEvaluator().FixEndKnots(InternalSpline.GetNumKeys());
+		}
+    
+		// Ensure tangent modes array size matches number of points after removal
+		if (TangentModes.Num() > GetNumPoints())
+		{
+			TangentModes.SetNum(GetNumPoints());
 		}
+    
+		// Update tangents for adjacent points
+		UpdateTangents();
+		
     }
 
 	/** 
@@ -1495,6 +1529,9 @@ public:
 	void Reparameterize(EParameterizationPolicy Policy = EParameterizationPolicy::Centripetal)
     {
 	    InternalSpline.Reparameterize(Policy);
+		
+	    // Update tangents to maintain proper curve shape
+	    UpdateTangents();
     }
     
     /** Getter for internal spline */
