use an inline allocator to remove the need to go to heap in typical cases for FDynamicMesh3::GetVtxContiguousTriangles

diff --git a/Engine/Source/Runtime/GeometryCore/Private/DynamicMesh/DynamicMesh3_Queries.cpp b/Engine/Source/Runtime/GeometryCore/Private/DynamicMesh/DynamicMesh3_Queries.cpp
--- a/Engine/Source/Runtime/GeometryCore/Private/DynamicMesh/DynamicMesh3_Queries.cpp
+++ b/Engine/Source/Runtime/GeometryCore/Private/DynamicMesh/DynamicMesh3_Queries.cpp
@@ -187,7 +187,10 @@ EMeshResult FDynamicMesh3::GetVtxContiguousTriangles(int VertexID, TArray<int>&
 		return EMeshResult::Ok;
 	}
 
-	TArray<int> StartEdgeIDs;
+	// The number of StartEdgeIDs for a vertex is NumBoundaryEdges + NumRingsWithoutBoundary
+	// This is only higher than 2 at non-manifold (e.g., "bowtie") vertices, and
+	// it should be very rare for it to exceed this inline allocation of 8
+	TArray<int, TInlineAllocator<8>> StartEdgeIDs;
 	// initial starting edge candidates == boundary edges
 	for (int EID : VertexEdgeLists.Values(VertexID))
 	{
