add parallel compute option to the dynamic mesh overlay normals recompute method + add benchmark and test coverage

diff --git a/Engine/Source/Runtime/GeometryCore/Private/DynamicMesh/MeshNormals.cpp b/Engine/Source/Runtime/GeometryCore/Private/DynamicMesh/MeshNormals.cpp
--- a/Engine/Source/Runtime/GeometryCore/Private/DynamicMesh/MeshNormals.cpp
+++ b/Engine/Source/Runtime/GeometryCore/Private/DynamicMesh/MeshNormals.cpp
@@ -6,6 +6,9 @@
 #include "Async/ParallelFor.h"
 #include "DynamicMesh/MeshIndexUtil.h"
 #include "MeshQueries.h"
+#include "HAL/UnrealMemory.h"
+
+#include <atomic>
 
 using namespace UE::Geometry;
 
@@ -17,10 +20,7 @@ void FMeshNormals::SetCount(int Count, bool bClearToZero)
 	}
 	if (bClearToZero)
 	{
-		for (int i = 0; i < Count; ++i)
-		{
-			Normals[i] = FVector3d::Zero();
-		}
+		FMemory::Memzero(Normals.GetData(), Normals.NumBytes());
 	}
 }
 
@@ -270,17 +270,23 @@ void FMeshNormals::Compute_Overlay_FaceAvg(const FDynamicMeshNormalOverlay* Norm
 
 	for (int TriIdx : Mesh->TriangleIndicesItr())
 	{
-		FVector3d TriNormal, TriCentroid; double TriArea;
-		Mesh->GetTriInfo(TriIdx, TriNormal, TriArea, TriCentroid);
+		FIndex3i Tri = NormalOverlay->GetTriangle(TriIdx);
+		if (Tri.A == INDEX_NONE)
+		{
+			continue;
+		}
+
+		FVector3d V0, V1, V2;
+		Mesh->GetTriVertices(TriIdx, V0, V1, V2);
+
+		FVector3d TriNormal;
+		double TriArea;
+		TriNormal = VectorUtil::NormalArea(V0, V1, V2, TriArea);
 		FVector3d TriNormalWeights = GetVertexWeightsOnTriangle(Mesh, TriIdx, TriArea, bWeightByArea, bWeightByAngle);
 
-		FIndex3i Tri = NormalOverlay->GetTriangle(TriIdx);
 		for (int j = 0; j < 3; ++j) 
 		{
-			if (Tri[j] != FDynamicMesh3::InvalidID)
-			{
-				Normals[Tri[j]] += TriNormal * TriNormalWeights[j];
-			}
+			Normals[Tri[j]] += TriNormal * TriNormalWeights[j];
 		}
 	}
 
@@ -384,18 +390,101 @@ void FMeshNormals::QuickComputeVertexNormalsForTriangles(FDynamicMesh3& Mesh, co
 }
 
 
+namespace UE::MeshNormalsLocals
+{
+	// This is a workaround for some platforms not supporting C++20's fetch_add for floats
+	// UE platforms should be on C++20, so hopefully we can remove this method in the future and just use .fetch_add directly
+	template<typename AtomicFloatType = std::atomic<float>>
+	static inline void AtomicFloatFetchAdd(AtomicFloatType& Value, float ToAdd)
+	{
+		constexpr bool bHasAtomicFloatFetchAdd = requires(AtomicFloatType& AtomicFloat, float Param)
+		{
+			AtomicFloat.fetch_add(Param);
+		};
+		if constexpr (bHasAtomicFloatFetchAdd)
+		{
+			Value.fetch_add(ToAdd);
+		}
+		else
+		{
+			// Manual implementation of fetch_add
+			float Old = Value.load();
+			while (!Value.compare_exchange_weak(Old, Old + ToAdd)) {}
+		}
+	}
+}
 
-bool FMeshNormals::QuickRecomputeOverlayNormals(FDynamicMesh3& Mesh, bool bInvert, bool bWeightByArea, bool bWeightByAngle)
+bool FMeshNormals::QuickRecomputeOverlayNormals(FDynamicMesh3& Mesh, bool bInvert, bool bWeightByArea, bool bWeightByAngle, bool bParallelCompute)
 {
-	if (Mesh.HasAttributes() && Mesh.Attributes()->GetNormalLayer(0) != nullptr)
+	if (!Mesh.HasAttributes() || Mesh.Attributes()->PrimaryNormals() == nullptr)
+	{
+		return false;
+	}
+	FDynamicMeshNormalOverlay* NormalOverlay = Mesh.Attributes()->PrimaryNormals();
+	if (!bParallelCompute)
 	{
-		FDynamicMeshNormalOverlay* NormalOverlay = Mesh.Attributes()->GetNormalLayer(0);
 		FMeshNormals Normals(&Mesh);
 		Normals.RecomputeOverlayNormals(NormalOverlay, bWeightByArea, bWeightByAngle);
 		Normals.CopyToOverlay(NormalOverlay, bInvert);
 		return true;
 	}
-	return false;
+	else
+	{
+		// for the parallel case we want to compute once per triangle normal, and accumulate results in element normals
+		// there is some overhead to using the atomic float buffer, so if not threading it is better to not do this
+		TArray<std::atomic<float>> Normals;
+		Normals.SetNumZeroed(NormalOverlay->MaxElementID() * 3);
+
+		constexpr bool bForceSingleThreaded = false;
+
+		ParallelFor(Mesh.MaxTriangleID(), [&Normals, &Mesh, bWeightByArea, bWeightByAngle, NormalOverlay](int32 TID)
+		{
+			if (!Mesh.IsTriangle(TID))
+			{
+				return;
+			}
+
+			FIndex3i ElTri = NormalOverlay->GetTriangle(TID);
+			if (ElTri.A == INDEX_NONE)
+			{
+				// triangle was not set / has no elements
+				return;
+			}
+
+			FVector3d V0, V1, V2;
+			Mesh.GetTriVertices(TID, V0, V1, V2);
+
+			FVector3d TriNormal;
+			double TriArea;
+			TriNormal = VectorUtil::NormalArea(V0, V1, V2, TriArea);
+			FVector3f TriNormalWeights = (FVector3f)TMeshQueries<FDynamicMesh3>::GetVertexWeightsOnTriangle(Mesh, TID, TriArea, bWeightByArea, bWeightByAngle);
+			FVector3f TriNormalf = (FVector3f)TriNormal;
+
+			for (int32 SubIdx = 0; SubIdx < 3; ++SubIdx)
+			{
+				int32 ElID = ElTri[SubIdx];
+				FVector3f AddNormal = TriNormalf * TriNormalWeights[SubIdx];
+				int32 NormalArrayIdx = ElID * 3;
+				for (int32 VecIdx = 0; VecIdx < 3; ++VecIdx)
+				{
+					UE::MeshNormalsLocals::AtomicFloatFetchAdd(Normals[NormalArrayIdx + VecIdx], AddNormal[VecIdx]);
+				}
+			}
+		}, bForceSingleThreaded);
+
+		float Sign = (bInvert) ? -1.0f : 1.0f;
+		ParallelFor(NormalOverlay->MaxElementID(), [&Normals, Sign, NormalOverlay](int32 ElID)
+		{
+			if (NormalOverlay->IsElement(ElID))
+			{
+				int32 NormalsIdx = ElID * 3;
+				// Note: Normalization intentionally computed in double precision for accuracy
+				FVector3d Normal(Normals[NormalsIdx], Normals[NormalsIdx + 1], Normals[NormalsIdx + 2]);
+				NormalOverlay->SetElement(ElID, static_cast<FVector3f>(Sign * Normalized(Normal)));
+			}
+		}, bForceSingleThreaded);
+	}
+	return true;
 }
 
 bool FMeshNormals::RecomputeOverlayTriNormals(FDynamicMesh3& Mesh, const TArray<int32>& Triangles, bool bWeightByArea, bool bWeightByAngle)
diff --git a/Engine/Source/Runtime/GeometryCore/Public/DynamicMesh/MeshNormals.h b/Engine/Source/Runtime/GeometryCore/Public/DynamicMesh/MeshNormals.h
--- a/Engine/Source/Runtime/GeometryCore/Public/DynamicMesh/MeshNormals.h
+++ b/Engine/Source/Runtime/GeometryCore/Public/DynamicMesh/MeshNormals.h
@@ -193,7 +193,7 @@ public:
 	 * Compute overlay normals for the given mesh
 	 * @param bInvert if true, normals are flipped
 	 */
-	static GEOMETRYCORE_API bool QuickRecomputeOverlayNormals(FDynamicMesh3& Mesh, bool bInvert = false, bool bWeightByArea = true, bool bWeightByAngle = true);
+	static GEOMETRYCORE_API bool QuickRecomputeOverlayNormals(FDynamicMesh3& Mesh, bool bInvert = false, bool bWeightByArea = true, bool bWeightByAngle = true, bool bParallelCompute = true);
 
 	/**
 	 * Compute overlay normals for the given mesh, for the given set of triangles
