Adaptive Displacement in TexturePatch modifier

diff --git a/Engine/Source/Runtime/GeometryCore/Public/DynamicMesh/DynamicVertexAttribute.h b/Engine/Source/Runtime/GeometryCore/Public/DynamicMesh/DynamicVertexAttribute.h
--- a/Engine/Source/Runtime/GeometryCore/Public/DynamicMesh/DynamicVertexAttribute.h
+++ b/Engine/Source/Runtime/GeometryCore/Public/DynamicMesh/DynamicVertexAttribute.h
@@ -224,7 +224,7 @@ public:
 		}
 	}
 
-	/** Get the element at a given index */
+	/** Get the element at a given index. For scalar access, use GetValue<T*const>(&Value) for a Value of type T. */
 	template<typename AsType>
 	void GetValue(int VertexID, AsType& Data) const
 	{
