GeometryCore: Add smooth min/max functions.

diff --git a/Engine/Source/Runtime/GeometryCore/Public/MathUtil.h b/Engine/Source/Runtime/GeometryCore/Public/MathUtil.h
--- a/Engine/Source/Runtime/GeometryCore/Public/MathUtil.h
+++ b/Engine/Source/Runtime/GeometryCore/Public/MathUtil.h
@@ -164,6 +164,10 @@ public:
 	static inline int32 Min3Index(const RealType A, const RealType B, const RealType C);
 	/** compute min and max of a,b,c with max 3 comparisons (sometimes 2) */
 	static inline void MinMax(RealType A, RealType B, RealType C, RealType& MinOut, RealType& MaxOut);
+	/** Uses a quadratic blend to smooth the min transition when A and B are less than BlendExtent apart. */
+	static inline RealType SmoothMin(RealType A, RealType B, RealType BlendExtent);
+	/** Uses a quadratic blend to smooth the max transition when A and B are less than BlendExtent apart. */
+	static inline RealType SmoothMax(RealType A, RealType B, RealType BlendExtent);
 	static inline RealType Sqrt(const RealType Value);
 	/** Compute cube root */
 	static inline RealType Cbrt(const RealType Value);
@@ -310,6 +314,30 @@ void TMathUtil<RealType>::MinMax(RealType A, RealType B, RealType C, RealType& M
 	}
 }
 
+template<typename RealType>
+RealType TMathUtil<RealType>::SmoothMin(RealType A, RealType B, RealType BlendExtent)
+{
+	if (BlendExtent == 0)
+	{
+		return Min(A, B);
+	}
+
+	double H = (BlendExtent + B - A) / 2;
+	return Min(B - H * Clamp(H / BlendExtent, 0.0, 1.0), A);
+}
+
+template<typename RealType>
+RealType TMathUtil<RealType>::SmoothMax(RealType A, RealType B, RealType BlendExtent)
+{
+	if (BlendExtent == 0)
+	{
+		return Max(A, B);
+	}
+
+	double H = (BlendExtent + A - B) / 2;
+	return Max(B + H * Clamp(H / BlendExtent, 0.0, 1.0), A);
+}
+
 template<typename RealType>
 RealType TMathUtil<RealType>::Sqrt(const RealType Value)
 {
