Add an AppendWithOffsets method on DynamicMesh3 that directly appends the buffers of another mesh

This can be significantly faster than the DynamicMeshEditor append method, and doesn't require creating/managing dense maps of the appended indices

diff --git a/Engine/Plugins/ChaosClothAssetEditor/Source/ChaosClothAssetTools/Private/ChaosClothAsset/ClothPatternToDynamicMeshMappingSupport.cpp b/Engine/Plugins/ChaosClothAssetEditor/Source/ChaosClothAssetTools/Private/ChaosClothAsset/ClothPatternToDynamicMeshMappingSupport.cpp
--- a/Engine/Plugins/ChaosClothAssetEditor/Source/ChaosClothAssetTools/Private/ChaosClothAsset/ClothPatternToDynamicMeshMappingSupport.cpp
+++ b/Engine/Plugins/ChaosClothAssetEditor/Source/ChaosClothAssetTools/Private/ChaosClothAsset/ClothPatternToDynamicMeshMappingSupport.cpp
@@ -27,6 +27,8 @@ public:
 	bool Apply(UE::Geometry::FDynamicMeshAttributeBase* Attribute, bool bRevert) const override;
 };
 
+// TODO: Can this class inherit from UE::Geometry::TDynamicMeshScalarTriangleAttribute<int32>
+// and then only need to implement a smaller subset of the methods? (On* and IsValid?)
 class FDynamicMeshSourceTriangleIdAttribute final :
 	public UE::Geometry::FDynamicMeshAttributeBase
 {
@@ -109,6 +111,39 @@ public:
 		return true;
 	}
 
+	virtual bool Append(const TDynamicAttributeBase& Source, const UE::Geometry::FDynamicMesh3::FAppendInfo& Info) override
+	{
+		int32 NewMaxID = Info.NumTriangle + Info.TriangleOffset;
+		if (NewMaxID > SourceTriangleIds.Num())
+		{
+			SourceTriangleIds.SetNum(NewMaxID);
+		}
+		for (int32 Idx = 0; Idx < Info.NumTriangle; ++Idx)
+		{
+			int32 TargetID = Idx + Info.TriangleOffset;
+			if (!Parent->IsTriangle(TargetID))
+			{
+				continue;
+			}
+			int32 SourceTID;
+			if (!ensure(Source.CopyOut(Idx, &SourceTID, sizeof(SourceTID))))
+			{
+				return false;
+			}
+			SetValue(Idx + Info.TriangleOffset, SourceTID);
+		}
+		return true;
+	}
+
+	virtual void AppendDefaulted(const UE::Geometry::FDynamicMesh3::FAppendInfo& Info) override
+	{
+		int32 NewMaxID = Info.NumTriangle + Info.TriangleOffset;
+		if (NewMaxID > SourceTriangleIds.Num())
+		{
+			SourceTriangleIds.Resize(NewMaxID, INDEX_NONE);
+		}
+	}
+
 	bool CopyOut(int RawID, void* Buffer, int BufferSize) const override
 	{
 		if (BufferSize != sizeof(int32) || !Parent->IsTriangle(RawID))
diff --git a/Engine/Source/Runtime/GeometryCore/Private/DynamicMesh/DynamicMesh3.cpp b/Engine/Source/Runtime/GeometryCore/Private/DynamicMesh/DynamicMesh3.cpp
--- a/Engine/Source/Runtime/GeometryCore/Private/DynamicMesh/DynamicMesh3.cpp
+++ b/Engine/Source/Runtime/GeometryCore/Private/DynamicMesh/DynamicMesh3.cpp
@@ -229,6 +229,110 @@ void FDynamicMesh3::Copy(const FDynamicMesh3& copy, bool bNormals, bool bColors,
 	TopologyChangeStamp.store(copy.TopologyChangeStamp);
 }
 
+void FDynamicMesh3::AppendWithOffsets(const FDynamicMesh3& ToAppend, FAppendInfo* OutAppendInfo)
+{
+	FAppendInfo LocalAppendInfo;
+	FAppendInfo* UseAppendInfo = OutAppendInfo ? OutAppendInfo : &LocalAppendInfo;
+
+	UseAppendInfo->VertexOffset = MaxVertexID();
+	UseAppendInfo->TriangleOffset = MaxTriangleID();
+	UseAppendInfo->EdgeOffset = MaxEdgeID();
+	UseAppendInfo->GroupOffset = HasTriangleGroups() ? MaxGroupID() : 0;
+	UseAppendInfo->NumVertex = ToAppend.MaxVertexID();
+	UseAppendInfo->NumTriangle = ToAppend.MaxTriangleID();
+	UseAppendInfo->NumEdge = ToAppend.MaxEdgeID();
+
+	if (HasAttributes())
+	{
+		for (int32 NormalLayerIdx = 0, N = FMath::Min(3, Attributes()->NumNormalLayers()); NormalLayerIdx < N; ++NormalLayerIdx)
+		{
+			UseAppendInfo->NormalOverlayOffsets[NormalLayerIdx] = Attributes()->GetNormalLayer(NormalLayerIdx)->MaxElementID();
+		}
+	}
+
+	Vertices.Add(ToAppend.Vertices);
+	auto MatchOptional = []<typename T>(TOptional<TDynamicVector<T>>& Src, const TOptional<TDynamicVector<T>>& ToAppend, int32 NumDefault, T DefaultValue = T())
+	{
+		if (Src.IsSet())
+		{
+			if (ToAppend.IsSet())
+			{
+				Src->Add(*ToAppend);
+			}
+			else
+			{
+				Src->Resize(NumDefault, DefaultValue);
+			}
+		}
+	};
+	MatchOptional(VertexNormals, ToAppend.VertexNormals, Vertices.Num(), FVector3f::UnitY());
+	MatchOptional(VertexColors, ToAppend.VertexColors, Vertices.Num(), FVector3f::One());
+	MatchOptional(VertexUVs, ToAppend.VertexUVs, Vertices.Num(), FVector2f::Zero());
+
+	VertexRefCounts.Append(ToAppend.VertexRefCounts);
+	VertexEdgeLists.AppendWithElementOffset(ToAppend.VertexEdgeLists, UseAppendInfo->EdgeOffset);
+
+	Triangles.Add(ToAppend.Triangles);
+	for (int32 Idx = UseAppendInfo->TriangleOffset, N = Triangles.Num(); Idx < N; ++Idx)
+	{
+		Triangles[Idx].A += UseAppendInfo->VertexOffset;
+		Triangles[Idx].B += UseAppendInfo->VertexOffset;
+		Triangles[Idx].C += UseAppendInfo->VertexOffset;
+	}
+	TriangleEdges.Add(ToAppend.TriangleEdges);
+	for (int32 Idx = UseAppendInfo->TriangleOffset, N = Triangles.Num(); Idx < N; ++Idx)
+	{
+		TriangleEdges[Idx].A += UseAppendInfo->EdgeOffset;
+		TriangleEdges[Idx].B += UseAppendInfo->EdgeOffset;
+		TriangleEdges[Idx].C += UseAppendInfo->EdgeOffset;
+	}
+	TriangleRefCounts.Append(ToAppend.TriangleRefCounts);
+
+	if (TriangleGroups.IsSet())
+	{
+		GroupIDCounter += ToAppend.MaxGroupID();
+		if (ToAppend.TriangleGroups.IsSet())
+		{
+			TriangleGroups->Add(*ToAppend.TriangleGroups);
+			for (int32 Idx = UseAppendInfo->TriangleOffset, N = Triangles.Num(); Idx < N; ++Idx)
+			{
+				(*TriangleGroups)[Idx] += UseAppendInfo->GroupOffset;
+			}
+		}
+		else
+		{
+			TriangleGroups->Resize(MaxTriangleID(), 0);
+		}
+	}
+
+	Edges.Add(ToAppend.Edges);
+	for (int32 Idx = UseAppendInfo->EdgeOffset, N = Edges.Num(); Idx < N; ++Idx)
+	{
+		Edges[Idx].Tri.A += UseAppendInfo->TriangleOffset;
+		if (Edges[Idx].Tri.B != INDEX_NONE)
+		{
+			Edges[Idx].Tri.B += UseAppendInfo->TriangleOffset;
+		}
+		Edges[Idx].Vert.A += UseAppendInfo->VertexOffset;
+		Edges[Idx].Vert.B += UseAppendInfo->VertexOffset;
+	}
+	EdgeRefCounts.Append(ToAppend.EdgeRefCounts);
+
+	if (HasAttributes())
+	{
+		if (ToAppend.HasAttributes())
+		{
+			AttributeSet->Append(*ToAppend.AttributeSet, *UseAppendInfo);
+		}
+		else
+		{
+			AttributeSet->AppendDefaulted(*UseAppendInfo);
+		}
+	}
+	
+	UpdateChangeStamps(true, true);
+}
+
 void FDynamicMesh3::CompactCopy(const FDynamicMesh3& copy, bool bNormals, bool bColors, bool bUVs, bool bAttributes, FCompactMaps* CompactInfo)
 {
 	if (copy.IsCompact() && ((!bAttributes || !HasAttributes()) || AttributeSet->IsCompact())) {
@@ -1052,7 +1156,7 @@ bool FDynamicMesh3::IsSameAs_Helper(const FDynamicMesh3& m2, const FSameAsOption
 					if (OutMeshDifferenceInfo)
 					{
 						OutMeshDifferenceInfo->Reason = FMeshDifferenceInfo::EReason::Group;
-						OutMeshDifferenceInfo->Detail = FString::Printf(TEXT("Triangle Group: %d != %d"), GetTriangleGroup(Tid), GetTriangleGroup(TidM2));
+						OutMeshDifferenceInfo->Detail = FString::Printf(TEXT("Triangle Group: %d != %d"), GetTriangleGroup(Tid), m2.GetTriangleGroup(TidM2));
 						OutMeshDifferenceInfo->SetTID(Tid, TidM2);
 					}
 					return false;
diff --git a/Engine/Source/Runtime/GeometryCore/Private/DynamicMesh/DynamicMeshAttributeSet.cpp b/Engine/Source/Runtime/GeometryCore/Private/DynamicMesh/DynamicMeshAttributeSet.cpp
--- a/Engine/Source/Runtime/GeometryCore/Private/DynamicMesh/DynamicMeshAttributeSet.cpp
+++ b/Engine/Source/Runtime/GeometryCore/Private/DynamicMesh/DynamicMeshAttributeSet.cpp
@@ -195,6 +195,146 @@ void FDynamicMeshAttributeSet::CompactCopy(const FCompactMaps& CompactMaps, cons
 }
 
 
+void FDynamicMeshAttributeSet::Append(const FDynamicMeshAttributeSet& ToAppend, const FDynamicMesh3::FAppendInfo& AppendInfo)
+{
+	auto AppendHelper = [&AppendInfo]<typename T>(T & Target, const T * ToAppend)
+	{
+		if (ToAppend)
+		{
+			Target.Append(*ToAppend, AppendInfo);
+		}
+		else
+		{
+			Target.AppendDefaulted(AppendInfo);
+		}
+	};
+
+	for (int32 Idx = 0; Idx < NumUVLayers(); ++Idx)
+	{
+		AppendHelper(UVLayers[Idx], ToAppend.GetUVLayer(Idx));
+	}
+	for (int32 Idx = 0; Idx < NumNormalLayers(); ++Idx)
+	{
+		AppendHelper(NormalLayers[Idx], ToAppend.GetNormalLayer(Idx));
+	}
+	if (ColorLayer.IsValid())
+	{
+		AppendHelper(*ColorLayer, ToAppend.ColorLayer.Get());
+	}
+	if (MaterialIDAttrib.IsValid())
+	{
+		AppendHelper(*MaterialIDAttrib, ToAppend.MaterialIDAttrib.Get());
+	}
+
+	for (int Idx = 0; Idx < NumPolygroupLayers(); ++Idx)
+	{
+		AppendHelper(PolygroupLayers[Idx], Idx < ToAppend.NumPolygroupLayers() ? ToAppend.GetPolygroupLayer(Idx) : nullptr);
+	}
+
+	for (int Idx = 0; Idx < NumWeightLayers(); ++Idx)
+	{
+		AppendHelper(WeightLayers[Idx], Idx < ToAppend.NumWeightLayers() ? ToAppend.GetWeightLayer(Idx) : nullptr);
+	}
+
+	// if both have bones, attempt to combine the bones and update weights accordingly
+	if (HasBones() && ToAppend.HasBones())
+	{
+		const bool bSameSkeletons = GetBoneNames()->IsSameAs(*ToAppend.GetBoneNames());
+		if (!bSameSkeletons)
+		{
+			AppendBonesUnique(ToAppend);
+		}
+		for (TPair<FName, TUniquePtr<FDynamicMeshVertexSkinWeightsAttribute>>& AttribPair : SkinWeightAttributes)
+		{
+			FDynamicMeshVertexSkinWeightsAttribute* ToAttrib = ToAppend.GetSkinWeightsAttribute(AttribPair.Key);
+			if (ToAttrib && !bSameSkeletons)
+			{
+				// Make a copy of the to-append skinning weights so we can reindex them w.r.t the new skeleton
+				// TODO: add a method to do this reindexing in-place on just the appended weights
+				FDynamicMeshVertexSkinWeightsAttribute CopyAppendMeshAttrib;
+				CopyAppendMeshAttrib.Copy(*ToAttrib);
+				CopyAppendMeshAttrib.ReindexBoneIndicesToSkeleton(ToAppend.GetBoneNames()->GetAttribValues(), GetBoneNames()->GetAttribValues());
+				AppendHelper(*AttribPair.Value, &CopyAppendMeshAttrib);
+			}
+			else
+			{
+				AppendHelper(*AttribPair.Value, ToAttrib);
+			}
+		}
+	}
+	else
+	{
+		// If only one had bones, we don't need to update bones; just transfer skin weights via standard append logic
+		for (TPair<FName, TUniquePtr<FDynamicMeshVertexSkinWeightsAttribute>>& AttribPair : SkinWeightAttributes)
+		{
+			AppendHelper(*AttribPair.Value, ToAppend.GetSkinWeightsAttribute(AttribPair.Key));
+		}
+	}
+
+	for (TPair<FName, TUniquePtr<FDynamicMeshAttributeBase>>& AttribPair : GenericAttributes)
+	{
+		const TUniquePtr<FDynamicMeshAttributeBase>* AppendAttr = ToAppend.GenericAttributes.Find(AttribPair.Key);
+		AppendHelper(*AttribPair.Value, AppendAttr ? AppendAttr->Get() : nullptr);
+	}
+
+	for (int32 Idx = 0; Idx < SculptLayers.NumLayers(); ++Idx)
+	{
+		AppendHelper(SculptLayers.Layers[Idx], ToAppend.SculptLayers.GetLayer(Idx));
+	}
+}
+
+void FDynamicMeshAttributeSet::AppendDefaulted(const FDynamicMesh3::FAppendInfo& AppendInfo)
+{
+	auto AppendHelper = [&AppendInfo]<typename T>(T & Target)
+	{
+		Target.AppendDefaulted(AppendInfo);
+	};
+
+	for (int32 Idx = 0; Idx < NumUVLayers(); ++Idx)
+	{
+		AppendHelper(UVLayers[Idx]);
+	}
+	for (int32 Idx = 0; Idx < NumNormalLayers(); ++Idx)
+	{
+		AppendHelper(NormalLayers[Idx]);
+	}
+	if (ColorLayer.IsValid())
+	{
+		AppendHelper(*ColorLayer);
+	}
+	if (MaterialIDAttrib.IsValid())
+	{
+		AppendHelper(*MaterialIDAttrib);
+	}
+
+	for (int Idx = 0; Idx < NumPolygroupLayers(); ++Idx)
+	{
+		AppendHelper(PolygroupLayers[Idx]);
+	}
+
+	for (int Idx = 0; Idx < NumWeightLayers(); ++Idx)
+	{
+		AppendHelper(WeightLayers[Idx]);
+	}
+
+	for (const TPair<FName, TUniquePtr<FDynamicMeshVertexSkinWeightsAttribute>>& AttribPair : SkinWeightAttributes)
+	{
+		AppendHelper(*AttribPair.Value);
+	}
+
+	// Note: In the 'defaulted' case, there is nothing to append re the Bone attributes (since they are not tied to the base mesh elements)
+
+	for (const TPair<FName, TUniquePtr<FDynamicMeshAttributeBase>>& AttribPair : GenericAttributes)
+	{
+		AppendHelper(*AttribPair.Value);
+	}
+
+	for (int32 Idx = 0; Idx < SculptLayers.NumLayers(); ++Idx)
+	{
+		AppendHelper(SculptLayers.Layers[Idx]);
+	}
+}
+
 
 
 void FDynamicMeshAttributeSet::CompactInPlace(const FCompactMaps& CompactMaps)
diff --git a/Engine/Source/Runtime/GeometryCore/Private/DynamicMesh/NonManifoldMappingSupport.cpp b/Engine/Source/Runtime/GeometryCore/Private/DynamicMesh/NonManifoldMappingSupport.cpp
--- a/Engine/Source/Runtime/GeometryCore/Private/DynamicMesh/NonManifoldMappingSupport.cpp
+++ b/Engine/Source/Runtime/GeometryCore/Private/DynamicMesh/NonManifoldMappingSupport.cpp
@@ -173,6 +173,22 @@ public:
 		SourceVertexIds[InVertexId] = InSourceVertexId;
 	}
 
+	virtual bool Append(const TDynamicAttributeBase<FDynamicMesh3>& Source, const FDynamicMesh3::FAppendInfo& Mapping) override
+	{
+		// Non-manifold vertex mapping only supports a single source mesh, so appended vertices cannot preserve sources ...
+		// Instead, we always append defaulted elements
+		AppendDefaulted(Mapping);
+		return true;
+	}
+	virtual void AppendDefaulted(const FDynamicMesh3::FAppendInfo& Mapping) override
+	{
+		const int32 NewNum = Mapping.VertexOffset + Mapping.NumVertex;
+		if (NewNum > SourceVertexIds.Num())
+		{
+			SourceVertexIds.Resize(NewNum, INDEX_NONE);
+		}
+	}
+
 	// Any topo operation on the mesh will invalidate the non-manifold information.
 	virtual void OnSplitEdge(const DynamicMeshInfo::FEdgeSplitInfo& ) override
 	{
diff --git a/Engine/Source/Runtime/GeometryCore/Private/DynamicMeshEditor.cpp b/Engine/Source/Runtime/GeometryCore/Private/DynamicMeshEditor.cpp
--- a/Engine/Source/Runtime/GeometryCore/Private/DynamicMeshEditor.cpp
+++ b/Engine/Source/Runtime/GeometryCore/Private/DynamicMeshEditor.cpp
@@ -1902,6 +1902,9 @@ void FDynamicMeshEditor::AppendMesh(const FDynamicMesh3* AppendMesh,
 	TFunction<FVector3d(int, const FVector3d&)> NormalTransform,
 	bool bReverseOrientation)
 {
+	// TODO: add a 'fast path' w/ the new Mesh->AppendWithOffsets, and build IndexMapsOut from the offsets
+	// (this could always be used w/ identical results if both meshes are compact, and perhaps via a flag in the non-compact case)
+
 	// todo: handle this case by making a copy?
 	check(AppendMesh != Mesh);
 
diff --git a/Engine/Source/Runtime/GeometryCore/Private/Util/SmallListSet.cpp b/Engine/Source/Runtime/GeometryCore/Private/Util/SmallListSet.cpp
--- a/Engine/Source/Runtime/GeometryCore/Private/Util/SmallListSet.cpp
+++ b/Engine/Source/Runtime/GeometryCore/Private/Util/SmallListSet.cpp
@@ -6,8 +6,6 @@
 
 using namespace UE::Geometry;
 
-const int32 FSmallListSet::NullValue = -1;
-
 void FSmallListSet::Resize(int32 NewSize)
 {
 	int32 CurSize = (int32)ListHeads.GetLength();
@@ -100,6 +98,79 @@ void FSmallListSet::Compact(int32 MaxListIndex)
 }
 
 
+void FSmallListSet::AppendWithElementOffset(const FSmallListSet& Other, int32 ElementOffset)
+{
+	int32 OrigListBlocksNum = ListBlocks.Num();
+	int32 OrigListHeadsNum = ListHeads.Num();
+	int32 OrigLinkedListElementsNum = LinkedListElements.Num();
+	int32 OrigFreeHeadIndex = FreeHeadIndex;
+
+	// Append ListHeads indices
+	ListHeads.Add(Other.ListHeads);
+	for (int32 Idx = OrigListHeadsNum, N = ListHeads.Num(); Idx < N; ++Idx)
+	{
+		// Offset appended non-null indices to point to appended ListBlock indices
+		if (ListHeads[Idx] != NullValue)
+		{
+			ListHeads[Idx] += (int32)OrigListBlocksNum;
+		}
+	}
+
+	// Append LinkedListElements entries
+	LinkedListElements.Add(Other.LinkedListElements);
+	for (int32 Idx = OrigLinkedListElementsNum, N = LinkedListElements.Num(); Idx < N; Idx += 2)
+	{
+		// Offset the element data
+		LinkedListElements[Idx] += ElementOffset;
+		// Offset appended non-null pointer indices to refer to appended indices
+		int32& Link = LinkedListElements[Idx + 1];
+		if (Link != NullValue)
+		{
+			Link += OrigLinkedListElementsNum;
+		}
+	}
+
+	// Append ListBlocks entries
+	ListBlocks.Add(Other.ListBlocks);
+	for (int32 Idx = OrigListBlocksNum, N = ListBlocks.Num(); Idx < N; Idx += BLOCKSIZE + 2)
+	{
+		int32 BlockNumEls = ListBlocks[Idx];
+		for (int32 SubIdx = 0, SubIdxNum = FMath::Min(BLOCKSIZE, BlockNumEls); SubIdx < SubIdxNum; ++SubIdx)
+		{
+			ListBlocks[Idx + 1 + SubIdx] += ElementOffset;
+		}
+		if (ListBlocks[Idx + BLOCKSIZE + 1] != NullValue)
+		{
+			ListBlocks[Idx + BLOCKSIZE + 1] += OrigLinkedListElementsNum;
+		}
+	}
+
+	// If we had a non-empty free list on Other, need to transfer it too
+	if (Other.FreeHeadIndex != NullValue)
+	{
+		FreeHeadIndex = Other.FreeHeadIndex + OrigLinkedListElementsNum;
+		// If both were non-empty, need to walk Other's free list to attach its tail to our original head
+		if (OrigFreeHeadIndex != NullValue)
+		{
+			int32 WalkListIndex = FreeHeadIndex;
+			while (true)
+			{
+				int32 NextIndex = LinkedListElements[WalkListIndex + 1];
+				if (NextIndex == NullValue)
+				{
+					break;
+				}
+				WalkListIndex = NextIndex;
+			}
+			checkSlow(LinkedListElements[WalkListIndex + 1] == NullValue);
+			LinkedListElements[WalkListIndex + 1] = OrigFreeHeadIndex;
+		}
+	}
+
+	AllocatedCount += Other.AllocatedCount;
+}
+
+
 void FSmallListSet::Insert(int32 ListIndex, int32 Value)
 {
 	checkSlow(0 <= ListIndex && ListIndex < ListHeads.Num());
diff --git a/Engine/Source/Runtime/GeometryCore/Public/DynamicMesh/DynamicAttribute.h b/Engine/Source/Runtime/GeometryCore/Public/DynamicMesh/DynamicAttribute.h
--- a/Engine/Source/Runtime/GeometryCore/Public/DynamicMesh/DynamicAttribute.h
+++ b/Engine/Source/Runtime/GeometryCore/Public/DynamicMesh/DynamicAttribute.h
@@ -118,6 +118,17 @@ public:
 	  */
 	virtual bool CopyThroughMapping(const TDynamicAttributeBase* Source, const FMeshIndexMappings& Mapping) = 0;
 
+	/**
+	 * Add elements from the Source attribute, using the AppendInfo to determine the range to copy.
+	 * @return true if the append succeeded, false otherwise (e.g. false if the data from the source attribute was not compatible and the CopyOut failed to copy across)
+	 */
+	virtual bool Append(const TDynamicAttributeBase& Source, const FDynamicMesh3::FAppendInfo& AppendInfo) = 0;
+
+	/**
+	 * Add default-valued elements, using the AppendInfo to determine the range to add. Called when appending from a mesh that does not have this attribute.
+	 */
+	virtual void AppendDefaulted(const FDynamicMesh3::FAppendInfo& AppendInfo) = 0;
+
 	/** Generic function to copy data out of an attribute; it's up to the derived class to map RawID to chunks of attribute data */
 	virtual bool CopyOut(int RawID, void* Buffer, int BufferSize) const = 0;
 
diff --git a/Engine/Source/Runtime/GeometryCore/Public/DynamicMesh/DynamicBoneAttribute.h b/Engine/Source/Runtime/GeometryCore/Public/DynamicMesh/DynamicBoneAttribute.h
--- a/Engine/Source/Runtime/GeometryCore/Public/DynamicMesh/DynamicBoneAttribute.h
+++ b/Engine/Source/Runtime/GeometryCore/Public/DynamicMesh/DynamicBoneAttribute.h
@@ -122,6 +122,19 @@ public:
 		return true;
 	}
 
+	virtual bool Append(const TDynamicAttributeBase<ParentType>& Source, const UE::Geometry::FDynamicMesh3::FAppendInfo& Info) override
+	{
+		// Appending bone attributes is a special case that must be handled AttributeSet level,
+		// because it may require also consistently remapping bone weights
+		ensureMsgf(false, TEXT("Directly appending bone attributes is not supported"));
+		return false;
+	}
+
+	virtual void AppendDefaulted(const UE::Geometry::FDynamicMesh3::FAppendInfo& Info) override
+	{
+		// bone attributes are not attached to mesh vertices/triangles/edges, so there is no work to do in the append-defaulted case
+	}
+
 	virtual bool CopyOut(int RawID, void* Buffer, int BufferSize) const override
 	{
 		return false; // Not supported for bones
diff --git a/Engine/Source/Runtime/GeometryCore/Public/DynamicMesh/DynamicMesh3.h b/Engine/Source/Runtime/GeometryCore/Public/DynamicMesh/DynamicMesh3.h
--- a/Engine/Source/Runtime/GeometryCore/Public/DynamicMesh/DynamicMesh3.h
+++ b/Engine/Source/Runtime/GeometryCore/Public/DynamicMesh/DynamicMesh3.h
@@ -241,6 +241,33 @@ public:
 	  */
 	GEOMETRYCORE_API bool Copy(const FMeshShapeGenerator* Generator);
 
+	// Tracks how IDs are offset and number of elements appended by a mesh append operation
+	struct FAppendInfo
+	{
+		// Offsets for base mesh element IDs
+		int32 VertexOffset = 0, TriangleOffset = 0, EdgeOffset = 0, GroupOffset = 0;
+
+		// Offsets for normals/tangents overlay element; will remain zero if the corresponding 'normal layer' is not present
+		int32 NormalOverlayOffsets[3]{ 0,0,0 };
+
+		// The number appended of each element type -- including 'invalid' slots, i.e. the amount by which MaxID increased
+		int32 NumVertex = 0, NumTriangle = 0, NumEdge = 0;
+	};
+
+	/**
+	 * Append the ToAppend mesh geometry to this mesh, preserving element IDs aside from a constant offset.
+	 *
+	 * Attributes from the ToAppend mesh will only be included if they already exist on this mesh. 
+	 * If more attributes are needed, consider calling EnableMatchingAttributes.
+	 * 
+	 * Note that OutAppendInfo optionally provides some offsets for convenience, but callers may also compute the offset for any element type
+	 * by storing the corresponding MaxID before calling AppendWithOffsets.
+	 * 
+	 * @param ToAppend Mesh to append
+	 * @param OutAppendInfo Optionally stores offsets and number of IDs appended for mesh elements.
+	 */
+	GEOMETRYCORE_API void AppendWithOffsets(const FDynamicMesh3& ToAppend, FAppendInfo* OutAppendInfo = nullptr);
+
 	/** 
 	 * Copy input mesh while compacting, i.e. removing unused vertices/triangles/edges. 
 	 * Note that it is currently unsafe to hold on to the attribute set pointer (returned from Attributes())
@@ -1776,6 +1803,8 @@ protected:
 private:
 	virtual EMeshResult CanCollapseEdgeInternal(int vKeep, int vRemove, double collapse_t,
 		const FCollapseEdgeOptions& Options, FEdgeCollapseInfo* OutCollapseInfo) const;
+
+	friend class FDynamicMeshAttributeSet;
 };
 
 
diff --git a/Engine/Source/Runtime/GeometryCore/Public/DynamicMesh/DynamicMeshAttributeSet.h b/Engine/Source/Runtime/GeometryCore/Public/DynamicMesh/DynamicMeshAttributeSet.h
--- a/Engine/Source/Runtime/GeometryCore/Public/DynamicMesh/DynamicMeshAttributeSet.h
+++ b/Engine/Source/Runtime/GeometryCore/Public/DynamicMesh/DynamicMeshAttributeSet.h
@@ -659,6 +659,13 @@ protected:
 	 * @param FailMode Desired behavior if mesh is found invalid
 	 */
 	GEOMETRYCORE_API virtual bool CheckValidity(bool bAllowNonmanifold, EValidityCheckFailMode FailMode) const;
+
+private:
+
+	// Called by FDynamicMesh3 during a mesh append, to append a corresponding attribute set
+	void Append(const FDynamicMeshAttributeSet& ToAppend, const FDynamicMesh3::FAppendInfo& AppendInfo);
+	// Called by FDynamicMesh3 during a mesh append, to defaulted attributes when the ToAppend mesh did not have an attribute set
+	void AppendDefaulted(const FDynamicMesh3::FAppendInfo& AppendInfo);
 };
 
 
diff --git a/Engine/Source/Runtime/GeometryCore/Public/DynamicMesh/DynamicMeshOverlay.h b/Engine/Source/Runtime/GeometryCore/Public/DynamicMesh/DynamicMeshOverlay.h
--- a/Engine/Source/Runtime/GeometryCore/Public/DynamicMesh/DynamicMeshOverlay.h
+++ b/Engine/Source/Runtime/GeometryCore/Public/DynamicMesh/DynamicMeshOverlay.h
@@ -247,6 +247,47 @@ public:
 		checkSlow(IsCompact());
 	}
 
+	void Append(const TDynamicMeshOverlay& ToAppend, const FDynamicMesh3::FAppendInfo& AppendInfo)
+	{
+		// We expect to be appending s.t. the intial ElementTriangles map to the pre-append triangles
+		if (!ensure(ElementTriangles.Num() == AppendInfo.TriangleOffset * 3))
+		{
+			ElementTriangles.SetNum(AppendInfo.TriangleOffset * 3);
+		}
+		check(ToAppend.ElementTriangles.Num() == AppendInfo.NumTriangle * 3);
+
+		int32 ElementIDOffset = ElementsRefCounts.GetMaxIndex();
+		int32 ElementTrianglesOffset = ElementTriangles.Num();
+		ElementTriangles.Add(ToAppend.ElementTriangles);
+		for (int32 Idx = ElementTrianglesOffset, N = ElementTriangles.Num(); Idx < N; ++Idx)
+		{
+			int32& ElID = ElementTriangles[Idx];
+			if (ElID != IndexConstants::InvalidID)
+			{
+				ElID += ElementIDOffset;
+			}
+		}
+
+		checkSlow(ElementIDOffset == ParentVertices.Num()); // ParentVertices must be 1:1 with Element IDs
+		ParentVertices.Add(ToAppend.ParentVertices);
+		for (int32 Idx = 0; Idx < ToAppend.ParentVertices.Num(); ++Idx)
+		{
+			int32& Parent = ParentVertices[Idx + ElementIDOffset];
+			if (Parent != IndexConstants::InvalidID)
+			{
+				Parent += AppendInfo.VertexOffset;
+			}
+		}
+		ElementsRefCounts.Append(ToAppend.ElementsRefCounts);
+		Elements.Add(ToAppend.Elements);
+	}
+	void AppendDefaulted(const FDynamicMesh3::FAppendInfo& AppendInfo)
+	{
+		// Note: ElementRefCounts, Elements and ParentVertices remain unchanged, since the new triangles are unset
+		checkSlow(ElementTriangles.Num() * 3 == AppendInfo.TriangleOffset);
+		ElementTriangles.Resize(3 * (AppendInfo.TriangleOffset + AppendInfo.NumTriangle), IndexConstants::InvalidID);
+	}
+
 	/** Discard all elements. */
 	void ClearElements();
 
diff --git a/Engine/Source/Runtime/GeometryCore/Public/DynamicMesh/DynamicMeshTriangleAttribute.h b/Engine/Source/Runtime/GeometryCore/Public/DynamicMesh/DynamicMeshTriangleAttribute.h
--- a/Engine/Source/Runtime/GeometryCore/Public/DynamicMesh/DynamicMeshTriangleAttribute.h
+++ b/Engine/Source/Runtime/GeometryCore/Public/DynamicMesh/DynamicMeshTriangleAttribute.h
@@ -196,6 +196,34 @@ public:
 		}
 		return true;
 	}
+
+	virtual bool Append(const TDynamicAttributeBase& Source, const UE::Geometry::FDynamicMesh3::FAppendInfo& Info) override
+	{
+		int32 NewMaxID = Info.NumTriangle + Info.TriangleOffset;
+		if (NewMaxID * AttribDimension > AttribValues.Num())
+		{
+			AttribValues.SetNum(NewMaxID * AttribDimension);
+		}
+
+		AttribValueType BufferData[AttribDimension];
+		int BufferSize = sizeof(BufferData);
+		for (int32 Idx = 0; Idx < Info.NumTriangle; ++Idx)
+		{
+			if (!ensure(Source.CopyOut(Idx, BufferData, BufferSize)))
+			{
+				return false;
+			}
+			SetValue(Idx + Info.TriangleOffset, BufferData);
+		}
+		return true;
+	}
+
+	virtual void AppendDefaulted(const UE::Geometry::FDynamicMesh3::FAppendInfo& Info) override
+	{
+		int32 NewMaxID = Info.NumTriangle + Info.TriangleOffset;
+		AttribValues.SetMinimumSize(NewMaxID * AttribDimension, GetDefaultAttributeValue());
+	}
+
 	virtual bool CopyOut(int RawID, void* Buffer, int BufferSize) const override
 	{
 		if (sizeof(AttribValueType) * AttribDimension != BufferSize)
diff --git a/Engine/Source/Runtime/GeometryCore/Public/DynamicMesh/DynamicVertexAttribute.h b/Engine/Source/Runtime/GeometryCore/Public/DynamicMesh/DynamicVertexAttribute.h
--- a/Engine/Source/Runtime/GeometryCore/Public/DynamicMesh/DynamicVertexAttribute.h
+++ b/Engine/Source/Runtime/GeometryCore/Public/DynamicMesh/DynamicVertexAttribute.h
@@ -185,6 +185,33 @@ public:
 		}
 		return true;
 	}
+
+	virtual bool Append(const TDynamicAttributeBase<ParentType>& Source, const UE::Geometry::FDynamicMesh3::FAppendInfo& Info) override
+	{
+		AttribValueType BufferData[AttribDimension];
+		int BufferSize = sizeof(BufferData);
+
+		int32 NewMaxID = Info.NumVertex + Info.VertexOffset;
+		if (NewMaxID * AttribDimension > AttribValues.Num())
+		{
+			AttribValues.SetNum(NewMaxID * AttribDimension);
+		}
+		for (int32 Idx = 0; Idx < Info.NumVertex; ++Idx)
+		{
+			if (!ensure(Source.CopyOut(Idx, BufferData, BufferSize)))
+			{
+				return false;
+			}
+			SetValue(Idx + Info.VertexOffset, BufferData);
+		}
+		return true;
+	}
+
+	virtual void AppendDefaulted(const UE::Geometry::FDynamicMesh3::FAppendInfo& Info) override
+	{
+		AttribValues.SetMinimumSize((Info.NumVertex + Info.VertexOffset) * AttribDimension, GetDefaultAttributeValue());
+	}
+
 	virtual bool CopyOut(int RawID, void* Buffer, int BufferSize) const override
 	{
 		if (sizeof(AttribValueType)*AttribDimension != BufferSize)
diff --git a/Engine/Source/Runtime/GeometryCore/Public/DynamicMesh/DynamicVertexSkinWeightsAttribute.h b/Engine/Source/Runtime/GeometryCore/Public/DynamicMesh/DynamicVertexSkinWeightsAttribute.h
--- a/Engine/Source/Runtime/GeometryCore/Public/DynamicMesh/DynamicVertexSkinWeightsAttribute.h
+++ b/Engine/Source/Runtime/GeometryCore/Public/DynamicMesh/DynamicVertexSkinWeightsAttribute.h
@@ -293,6 +293,36 @@ public:
 		}
 		return true;
 	}
+
+	virtual bool Append(const TDynamicAttributeBase<ParentType>& Source, const UE::Geometry::FDynamicMesh3::FAppendInfo& Info) override
+	{
+		// Don't snarf the FBoneWeight as a concrete object, since it _may_ contain a pointer
+		// and we don't want it to destruct.
+		constexpr int32 BufferSize = sizeof(FBoneWeights);
+		int8 BufferData[BufferSize];
+
+		int32 NewMaxID = Info.NumVertex + Info.VertexOffset;
+		if (NewMaxID > VertexBoneWeights.Num())
+		{
+			VertexBoneWeights.SetNum(NewMaxID);
+		}
+		for (int32 Idx = 0; Idx < Info.NumVertex; ++Idx)
+		{
+			if (!ensure(Source.CopyOut(Idx, BufferData, BufferSize)))
+			{
+				return false;
+			}
+			SetValue(Idx + Info.VertexOffset, *reinterpret_cast<FBoneWeights*>(BufferData));
+		}
+		return true;
+	}
+
+	virtual void AppendDefaulted(const UE::Geometry::FDynamicMesh3::FAppendInfo& Info) override
+	{
+		int32 NewMaxID = Info.NumVertex + Info.VertexOffset;
+		VertexBoneWeights.SetMinimumSize(NewMaxID, FBoneWeights{});
+	}
+
 	bool CopyOut(int RawID, void* Buffer, int BufferSize) const override
 	{
 		if (sizeof(FBoneWeights) != BufferSize)
diff --git a/Engine/Source/Runtime/GeometryCore/Public/Util/RefCountVector.h b/Engine/Source/Runtime/GeometryCore/Public/Util/RefCountVector.h
--- a/Engine/Source/Runtime/GeometryCore/Public/Util/RefCountVector.h
+++ b/Engine/Source/Runtime/GeometryCore/Public/Util/RefCountVector.h
@@ -73,6 +73,21 @@ public:
 		return RefCounts[Index];
 	}
 
+	// Append all ref counts from another RefCountVector, offsetting FreeIndices to refer to their corresponded new array positions
+	// Note this does not try to 'fill in' original free indices, by design -- all existing IDs remain untouched, and all new IDs are simply offsets of the Other's IDs
+	void Append(const FRefCountVector& Other)
+	{
+		size_t OrigNum = RefCounts.Num();
+		size_t OrigFree = FreeIndices.Num();
+		RefCounts.Add(Other.RefCounts);
+		FreeIndices.Add(Other.FreeIndices);
+		for (size_t Idx = OrigFree, N = FreeIndices.Num(); Idx < N; ++Idx)
+		{
+			FreeIndices[Idx] += OrigNum;
+		}
+		UsedCount += Other.UsedCount;
+	}
+
 	int Allocate()
 	{
 		UsedCount++;
diff --git a/Engine/Source/Runtime/GeometryCore/Public/Util/SmallListSet.h b/Engine/Source/Runtime/GeometryCore/Public/Util/SmallListSet.h
--- a/Engine/Source/Runtime/GeometryCore/Public/Util/SmallListSet.h
+++ b/Engine/Source/Runtime/GeometryCore/Public/Util/SmallListSet.h
@@ -36,7 +36,7 @@ class FSmallListSet
 {
 protected:
 	/** This value is used to indicate Null in internal pointers */
-	static GEOMETRYCORE_API const int32 NullValue; // = -1;		// note: cannot be constexpr because we pass as reference to several functions, requires C++17
+	static GEOMETRYCORE_API constexpr int32 NullValue = -1;
 
 	/** size of initial linear-memory portion of lists */
 	static constexpr int32 BLOCKSIZE = 8;
@@ -64,7 +64,7 @@ protected:
 	 */
 	TDynamicVector<int32> LinkedListElements{};
 
-	/** index of first free element in linked_store */
+	/** index of first free element in LinkedListElements */
 	int32 FreeHeadIndex{NullValue};
 
 
@@ -316,6 +316,14 @@ public:
 		}
 	}
 
+	/**
+	 * Append all elements from another small list set, with ElementOffset added to all such appended elements
+	 * Note: Also directly appends free blocks / free lists; does no compacting
+	 * @param Other the list to append
+	 * @param ElementOffset amount to shift all appended elements
+	 */
+	void AppendWithElementOffset(const FSmallListSet& Other, int32 ElementOffset);
+
 	/**
 	 * Call ApplyFunc on each element of the list at ListIndex, until ApplyFunc returns false
 	 * @return true if all elements were processed and ApplyFunc never returned false
