Splines
- Support for 0 or 1 points and the ability to clear the spline re-added in TangentBezierSpline
- AddPoint/InsertPoint API clarification
- Fixed issues with closed loop not correctly updating the loop segment's tangents correctly

diff --git a/Engine/Source/Runtime/Engine/Private/Curves/Spline.cpp b/Engine/Source/Runtime/Engine/Private/Curves/Spline.cpp
--- a/Engine/Source/Runtime/Engine/Private/Curves/Spline.cpp
+++ b/Engine/Source/Runtime/Engine/Private/Curves/Spline.cpp
@@ -138,8 +138,6 @@ public:
 
 	float GetSegmentLength(const int32 Index, const float Param, const FVector& Scale3D = FVector(1.0f)) const;
 	float GetSplineLength() const;
-	int32 GetNumPoints() const;
-	int32 GetNumSegments() const;
 	
     void SetClosedLoop(bool bInClosedLoop);
 
@@ -224,6 +222,28 @@ private:
 
 	/** Clears and repopulates PositionCurve, RotationCurve, and ScaleCurve */
 	void RebuildLegacyCurves();
+
+	/**
+     * Converts FSplinePoint to FTangentBezierControlPoint
+     */
+    FTangentBezierControlPoint ConvertToTangentBezierControlPoint(const FSplinePoint& Point) const;
+
+    /**
+     * Updates the attributes of a point in the spline
+     * @param Point 
+     * @param PointIndex 
+     */
+    void UpdatePointAttributes(const FSplinePoint& Point, int32 PointIndex);
+	
+	/**
+	 * Converts an index (possibly fractional) to a parameter value
+	 */
+	float ConvertIndexToParameter(int32 Index, float Fraction = 0.0f) const;
+
+    /**
+	 * Converts a parameter value to the closest index
+	 */
+	int32 ConvertParameterToIndex(float Parameter) const;
 };
 
 /**
@@ -245,15 +265,13 @@ TArray<FName> FNewSpline::GetAttributeChannelNames() const
 template <typename AttrType>
 AttrType FNewSpline::EvaluateAttribute(FName Name, float Parameter) const
 {
-	const int32 NumEffectivePoints = GetNumPoints();
-	Parameter = GetParameterRange().Interpolate(Parameter / (NumEffectivePoints - 1));
+	Parameter = GetParameterRange().Interpolate(Parameter / GetNumSegments());
 	return InternalSpline.EvaluateAttribute<AttrType>(Name, Parameter);
 }
 
 int32 FNewSpline::SetAttributeParameter(const FName& Name, int32 Index, float Parameter)
 {
-	const int32 NumEffectivePoints = GetNumPoints();
-   	Parameter = GetParameterRange().Interpolate(Parameter / (NumEffectivePoints - 1));
+  	Parameter = GetParameterRange().Interpolate(Parameter / GetNumSegments());
 	return InternalSpline.SetAttributeParameter(Name, Index, Parameter);
 }
 
@@ -261,7 +279,7 @@ float FNewSpline::GetAttributeParameter(const FName& Name, int32 Index) const
 {
 	float SplineSpaceParam = InternalSpline.GetAttributeParameter(Name, Index);
 	const float NormalizedParam = this->GetParameterRange().GetT(SplineSpaceParam);
-	const float SegmentSpaceParam = NormalizedParam * static_cast<float>(GetNumPoints() - 1);
+	const float SegmentSpaceParam = NormalizedParam * static_cast<float>(GetNumSegments());
 	return SegmentSpaceParam;
 }
 
@@ -1220,7 +1238,7 @@ FNewSpline::FNewSpline(const FLegacySpline& Other)
 
 	// initialize spline
 	SetControlPoints(Points);
-
+	
 	auto NormalizeAttributeParameter = [MaxParameter = static_cast<float>(GetNumPoints() - 1)](int32 Index)
 	{
 		return Index / MaxParameter;
@@ -1382,15 +1400,6 @@ float FNewSpline::GetSplineLength() const
     return EvaluateAttribute<float>(ParamToLengthAttrName, 1.f);
 }
 
-int32 FNewSpline::GetNumPoints() const
-{
-	return IsClosedLoop() ? FTangentBezierSpline::GetNumPoints() - 1 : FTangentBezierSpline::GetNumPoints();
-}
-
-int32 FNewSpline::GetNumSegments() const
-{
-	return FTangentBezierSpline::GetNumPoints() - 1;
-}
 
 void FNewSpline::SetClosedLoop(bool bInClosedLoop)
 {
@@ -1400,9 +1409,7 @@ void FNewSpline::SetClosedLoop(bool bInClosedLoop)
 
 void FNewSpline::Reset()
 {
-    //Points.Empty();
-	//TODO: Not sure what to do here as we don't have a reset to no segments implementation for our Bezier splines
-    UpdateSpline();
+    InternalSpline.Clear();
 }
 
 void FNewSpline::ResetToDefault()
@@ -1448,75 +1455,54 @@ void FNewSpline::ResetScale()
 
 void FNewSpline::AddPoint(const FSplinePoint& Point)
 {
-    // auto UpperBound = [](const TArray<FTangentBezierControlPoint>& SplinePoints, float Value)
-    // {
-    // 	const int32 InsertIndex = static_cast<int32>(Value);
-    // 	int32 Count = SplinePoints.Num();
-    // 	int32 First = 0;
-    //
-    // 	while (Count > 0)
-    // 	{
-    // 		const int32 Middle = Count / 2;
-    // 		const int32 CurrentIndex = First + Middle;
-    //
-    // 		if (CurrentIndex <= InsertIndex)
-    // 		{
-    // 			First = CurrentIndex + 1;
-    // 			Count -= Middle + 1;
-    // 		}
-    // 		else
-    // 		{
-    // 			Count = Middle;
-    // 		}
-    // 	}
-    //
-    // 	return First;
-    // };
-    //
-    // const int32 Index = UpperBound(Points, Point.InputKey);
+    FTangentBezierControlPoint ControlPoint = ConvertToTangentBezierControlPoint(Point);
+    const int32 NumPoints = GetNumPoints();
+    
+    // Special case for empty spline
+    if (NumPoints == 0)
+    {
+        AppendPoint(ControlPoint);
+        UpdatePointAttributes(Point, 0);
+        return;
+    }
+	// Convert InputKey to parameter value
+	int32 BaseIndex = FMath::FloorToInt(Point.InputKey);
+	float Fraction = Point.InputKey - BaseIndex;
+	float Parameter = ConvertIndexToParameter(BaseIndex, Fraction);
     
-    FTangentBezierControlPoint NewPoint;
-    NewPoint.Position = Point.Position;
-    NewPoint.TangentIn = Point.ArriveTangent;
-    NewPoint.TangentOut = Point.LeaveTangent;
-    //NewPoint.Rotation = Point.Rotation.Quaternion();
-    // NewPoint.Scale = Point.Scale;
-    //
-    // Points.Insert(NewPoint, Index);
-	FTangentBezierSpline::AddPoint(NewPoint);
-
-	//TODO: Make sure this works as expected
-	AddAttributeValue(RotationAttrName, Point.Rotation.Quaternion(), Point.InputKey);
-	AddAttributeValue(ScaleAttrName, Point.Scale, Point.InputKey);
+	// Insert at calculated parameter
+	int32 InsertedIndex = FTangentBezierSpline::InsertPoint(Parameter, ControlPoint);
 	
-    RebuildLegacyCurves();
+	UpdatePointAttributes(Point, InsertedIndex);
 }
 
 void FNewSpline::InsertPoint(const FSplinePoint& Point, int32 Index)
 {
-    Index = FMath::Clamp(Index, 0, GetNumPoints() - 1);
-	float SqDistance;
-	float SplineParam = FTangentBezierSpline::FindNearestOnSegment(Point.Position, Index - 1, SqDistance);
-	
-	EInterpCurveMode InterpCurveMode = ConvertSplinePointTypeToInterpCurveMode(Point.Type);
-    int32 InsertedIndex = FTangentBezierSpline::InsertPoint(SplineParam, Point.Position, ConvertInterpCurveModeToTangentMode(InterpCurveMode));
 
-	// If insertion was successful, update tangents if not auto-computed
-	if (InsertedIndex >= 0 && InterpCurveMode != CIM_CurveAuto && InterpCurveMode != CIM_CurveAutoClamped)
+	FTangentBezierControlPoint ControlPoint = ConvertToTangentBezierControlPoint(Point);
+	const int32 NumPoints = GetNumPoints();
+    
+	// Special case for empty spline
+	if (NumPoints == 0 || Index >= NumPoints)
 	{
-		SetTangentIn(InsertedIndex, Point.ArriveTangent);
-		SetTangentOut(InsertedIndex, Point.LeaveTangent);
+		AppendPoint(ControlPoint);
+		UpdatePointAttributes(Point, NumPoints);
+		return;
+	}
+	else if (Index == 0)
+	{
+		PrependPoint(ControlPoint);
+		UpdatePointAttributes(Point, 0);
+		return;
 	}
-
-	// Update rotation and scale attributes
-	// Use the normalized parameter for the new point
-	float NormalizedParam = static_cast<float>(InsertedIndex) / 
-						   FMath::Max(1.0f, static_cast<float>(GetNumPoints() - 1));
-    
-	AddAttributeValue(RotationAttrName, Point.Rotation.Quaternion(), NormalizedParam);
-	AddAttributeValue(ScaleAttrName, Point.Scale, NormalizedParam);
 	
-    RebuildLegacyCurves();
+    Index = FMath::Clamp(Index, 0, GetNumPoints());
+	float SqDistance;
+	float SplineParam = FTangentBezierSpline::FindNearestOnSegment(Point.Position, Index - 1, SqDistance);
+		
+   	int32 InsertedSegmentIndex = FTangentBezierSpline::InsertPoint(SplineParam, ControlPoint);
+	
+	UpdatePointAttributes(Point, InsertedSegmentIndex);
 }
 
 FSplinePoint FNewSpline::GetPoint(int32 Index) const
@@ -1716,8 +1702,7 @@ float FNewSpline::GetParameterAtDistance(float Distance) const
 
 float FNewSpline::GetDistanceAtParameter(float Parameter) const
 {
-    const int32 NumEffectivePoints = GetNumPoints();
-    Parameter = Parameter / (NumEffectivePoints - 1);
+    Parameter = Parameter / GetNumSegments();
     return EvaluateAttribute<float>(ParamToLengthAttrName, Parameter);
 }
 
@@ -1756,8 +1741,7 @@ float FNewSpline::FindNearestOnSegment(const FVector& InLocation, int32 SegmentI
 
 FVector FNewSpline::Evaluate(float Parameter) const
 {
-    const int32 NumEffectivePoints = GetNumPoints();
-    Parameter = GetParameterRange().Interpolate(Parameter / (NumEffectivePoints - 1));
+    Parameter = GetParameterRange().Interpolate(Parameter / GetNumSegments());
 	
     return FVector(GetPosition(Parameter));
 }
@@ -1765,21 +1749,21 @@ FVector FNewSpline::Evaluate(float Parameter) const
 FVector FNewSpline::EvaluateDerivative(float Parameter) const
 {
     const int32 NumEffectivePoints = GetNumPoints();
-	Parameter = GetParameterRange().Interpolate(Parameter / (NumEffectivePoints - 1));
+	Parameter = GetParameterRange().Interpolate(Parameter / GetNumSegments());
     return FVector(GetTangent(Parameter));
 }
 
 FQuat FNewSpline::EvaluateRotation(float Parameter) const
 {
     const int32 NumEffectivePoints = GetNumPoints();
-	Parameter = GetParameterRange().Interpolate(Parameter / (NumEffectivePoints - 1));
+	Parameter = GetParameterRange().Interpolate(Parameter / GetNumSegments());
     return EvaluateAttribute<FQuat>(RotationAttrName, Parameter);
 }
 
 FVector FNewSpline::EvaluateScale(float Parameter) const
 {
     const int32 NumEffectivePoints = GetNumPoints();
-	Parameter = GetParameterRange().Interpolate(Parameter / (NumEffectivePoints - 1));
+	Parameter = GetParameterRange().Interpolate(Parameter / GetNumSegments());
     return EvaluateAttribute<FVector>(ScaleAttrName, Parameter);
 }
 
@@ -1867,7 +1851,76 @@ void FNewSpline::RebuildLegacyCurves()
     }
 }
 
+FTangentBezierControlPoint FNewSpline::ConvertToTangentBezierControlPoint(const FSplinePoint& Point) const
+{
+	FTangentBezierControlPoint NewPoint;
+	NewPoint.Position = Point.Position;
+	NewPoint.TangentIn = Point.ArriveTangent;
+	NewPoint.TangentOut = Point.LeaveTangent;
+	NewPoint.TangentMode = ConvertInterpCurveModeToTangentMode(
+		ConvertSplinePointTypeToInterpCurveMode(Point.Type));
+	return NewPoint;
+}
+
+void FNewSpline::UpdatePointAttributes(const FSplinePoint& Point, int32 PointIndex)
+{
+	if (PointIndex < 0)
+		return;
+            
+	// Calculate normalized parameter
+	float NormalizedParam = static_cast<float>(PointIndex) / 
+		FMath::Max(1.0f, static_cast<float>(GetNumPoints() - 1));
+        
+	AddAttributeValue(RotationAttrName, Point.Rotation.Quaternion(), NormalizedParam);
+	AddAttributeValue(ScaleAttrName, Point.Scale, NormalizedParam);
+        
+	RebuildLegacyCurves();
+}
+
+float FNewSpline::ConvertIndexToParameter(int32 Index, float Fraction) const
+{
+	// Handle empty spline
+	if (GetNumPoints() <= 1)
+		return 0.0f;
+    
+	// Clamp to valid range
+	Index = FMath::Clamp(Index, 0, GetNumPoints() - 1);
+    
+	// For exact index, just return parameter at that index
+	if (FMath::IsNearlyZero(Fraction))
+		return GetParameterAt(Index);
+    
+	// For fractional indices, interpolate
+	int32 NextIndex = FMath::Min(Index + 1, GetNumPoints() - 1);
+	float StartParam = GetParameterAt(Index);
+	float EndParam = GetParameterAt(NextIndex);
+    
+	return FMath::Lerp(StartParam, EndParam, Fraction);
+}
 
+int32 FNewSpline::ConvertParameterToIndex(float Parameter) const
+{
+	// Handle empty spline
+	if (GetNumPoints() <= 1)
+		return 0;
+    
+	// Find segment containing this parameter
+	for (int32 i = 0; i < GetNumPoints() - 1; i++)
+	{
+		float StartParam = GetParameterAt(i);
+		float EndParam = GetParameterAt(i + 1);
+        
+		if (Parameter >= StartParam && Parameter <= EndParam)
+		{
+			// Return closer index
+			float Fraction = (Parameter - StartParam) / (EndParam - StartParam);
+			return (Fraction <= 0.5f) ? i : (i + 1);
+		}
+	}
+    
+	// Out of range
+	return (Parameter < GetParameterAt(0)) ? 0 : (GetNumPoints() - 1);
+}
 
 
 /** FLegacySpline Implementation */
diff --git a/Engine/Source/Runtime/GeometryCore/Public/Splines/BSpline.h b/Engine/Source/Runtime/GeometryCore/Public/Splines/BSpline.h
--- a/Engine/Source/Runtime/GeometryCore/Public/Splines/BSpline.h
+++ b/Engine/Source/Runtime/GeometryCore/Public/Splines/BSpline.h
@@ -31,7 +31,7 @@ template <typename ScalarType, int32 Dimensions, int32 Degree>
 class TBSplineEvaluator : public TBaseEvaluator<typename TChooseVectorType<ScalarType, Dimensions>::Type>
 {
 	static_assert(Degree > 0, "BSpline degree must be positive");
-	
+	friend class TBSpline<ScalarType, Dimensions, Degree, TBSplineEvaluator>;
 public:
 	using ValueType = typename TChooseVectorType<ScalarType, Dimensions>::Type;
 	using Base = TBaseEvaluator<ValueType>;
@@ -924,6 +924,14 @@ public:
 	virtual ~TBSpline() override = default;
 
 
+	/**
+	 * Clears all values from the spline including custom attribute channels and knot vector
+	 */
+	virtual void Clear()
+	{
+		Base::Clear();
+		this->GetEvaluator().KnotVector.Reset();
+	}
 	/**
 	 * Sets a custom knot vector for the spline
 	 * @param NewKnots - Array of knot values. Must be non-decreasing and have length numPoints + degree + 1
diff --git a/Engine/Source/Runtime/GeometryCore/Public/Splines/PolyBezierSpline.h b/Engine/Source/Runtime/GeometryCore/Public/Splines/PolyBezierSpline.h
--- a/Engine/Source/Runtime/GeometryCore/Public/Splines/PolyBezierSpline.h
+++ b/Engine/Source/Runtime/GeometryCore/Public/Splines/PolyBezierSpline.h
@@ -13,12 +13,15 @@ namespace Spline
  * Custom evaluator for PolyBezier splines.
  * Handles parameter mapping and control point selection for proper Bezier evaluation.
  */
+	
 template<typename T, int32 D>
 class TPolyBezierEvaluator : public TBSplineEvaluator<T, D, 3>
 {	
 	using ValueType = typename TChooseVectorType<T, D>::Type;
 	using Base = TBSplineEvaluator<T, D, 3>;
 
+	template<typename ScalarType, int32 Dimensions>
+	friend class TPolyBezierSpline;
     
 public:
     TPolyBezierEvaluator():
@@ -161,7 +164,7 @@ public:
             return *Window[3];  // Last control point (P3) of final Bezier segment
         }
     	const float ClampedParam = FMath::Clamp(SegmentSpaceParam, 0.0f, NumSegments);
-    	int32 SegmentIndex = FMath::FloorToInt(ClampedParam);
+    	int32 SegmentIndex = FMath::FloorToInt(ClampedParam); 
     	const float LocalT = ClampedParam - static_cast<float>(SegmentIndex);
 
         // Calculate Bernstein basis functions
@@ -301,12 +304,12 @@ public:
 		{
 			// Get parameter range to determine if these are endpoints
 			FInterval1f KnotRange = this->GetKnotRange();
-        
+			float CachedEndKnot = KnotRange.Max;
 			// Check if StartParam is at the beginning of the knot range
-			bool bIsStartEndpoint = FMath::IsNearlyEqual(StartParam, KnotRange.Min) || StartParam < KnotRange.Min;
+			bool bIsStartParamEndKnot = StartParam <= KnotRange.Min || StartParam >= KnotRange.Max;
         
 			// Check if EndParam is at the end of the knot range
-			bool bIsEndEndpoint = FMath::IsNearlyEqual(EndParam, KnotRange.Max) || EndParam > KnotRange.Max;
+			bool bIsEndParamEndKnot = EndParam <= KnotRange.Min || EndParam >= KnotRange.Max;
         
 			// For endpoints, we want multiplicity of degree+1 for clamped splines
 			const int32 EndpointMultiplicity = this->GetDegree() + 1;
@@ -315,20 +318,29 @@ public:
 			const int32 InternalMultiplicity = this->GetDegree();
 
 			int32 StartMult = this->GetKnotMultiplicity(StartParam);
-			int32 DesiredStartMult = bIsStartEndpoint ? EndpointMultiplicity : InternalMultiplicity;
-        
+			int32 DesiredStartMult = bIsStartParamEndKnot ? EndpointMultiplicity : InternalMultiplicity;
+
+			
 			if (StartMult < DesiredStartMult)
 			{
 				this->InsertKnot(StartParam, DesiredStartMult - StartMult);
 			}
+			else if (StartMult > DesiredStartMult)
+			{
+				this->RemoveKnot(StartParam, false, StartMult - DesiredStartMult);
+			}
 
 			int32 EndMult = this->GetKnotMultiplicity(EndParam);
-			int32 DesiredEndMult = bIsEndEndpoint ? EndpointMultiplicity : InternalMultiplicity;
+			int32 DesiredEndMult = bIsEndParamEndKnot ? EndpointMultiplicity : InternalMultiplicity;
         
 			if (EndMult < DesiredEndMult)
 			{
 				this->InsertKnot(EndParam, DesiredEndMult - EndMult);
 			}
+			else if (EndMult > DesiredEndMult)
+			{
+				this->RemoveKnot(EndParam, false, EndMult - DesiredEndMult);
+			}
 		}
 		else
 		{
@@ -512,6 +524,13 @@ public:
 		AddBezierSegmentInternal(P0, P1, P2, P3, 0, Parameterization);
 		Reparameterize(Parameterization);
 	}
+
+	// Default constructor used for serialization
+    TPolyBezierSpline()
+    	: Base(MakeUnique<TPolyBezierEvaluator<ScalarType, Dimensions>>())
+    {
+        
+    }
     
     // Virtual destructor
     virtual ~TPolyBezierSpline() = default;
@@ -652,7 +671,8 @@ public:
     
 		// Set as closed loop
 		Result.SetClosedLoop(true);
-    
+		// Apply consistent parameterization
+		Result.Reparameterize(EParameterizationPolicy::Centripetal);
 		return Result;
     }
     
@@ -1028,52 +1048,39 @@ public:
 		return AddBezierSegmentInternal(P0, P1, P2, P3, PrependParam, ParameterizationPolicy, ParameterRange.Min);
 	}
 
-    /**
-     * Inserts a new point at the specified parameter along the spline.
-     * Uses the specified parameterization policy to preserve attribute positions.
-     * 
-     * @param Parameter Parameter value where to insert the point (Spline Space)
-     * @param Position Position for the new point
-     * @param ParameterizationPolicy How to distribute parameters
-     * @return The index of the newly created point
-     */
+	/**
+	 * Inserts a point at the specified parameter along the spline.
+	 * This splits the segment at the parameter value and places the new point
+	 * at the specified position, which may change the spline shape.
+	 * 
+	 * @param Parameter Parameter value where to split the spline
+	 * @param Position Position for the new point (doesn't need to be on the original curve)
+	 * @param ParameterizationPolicy How to distribute parameters
+	 * @return The index of the newly created point
+	 */
     int32 InsertPoint(
         float Parameter, 
         const ValueType& Position,
         EParameterizationPolicy ParameterizationPolicy = EParameterizationPolicy::Centripetal)
     {
+		int32 NewPointIndex;
+		if (InitializeFirstSegment(Parameter, Position, NewPointIndex)) return NewPointIndex;
+		
         // Find segment to split
         const int32 NumSegments = GetNumSegments();
-        if (NumSegments == 0)
-        {
-            // If no segments, just add this as the first point
-            Base::AddValue(Position, 0.0f);
-            return 0;
-        }
+        
         FInterval1f ParameterRange = this->GetParameterRange();
 		
 		// Handle out-of-range parameters
 		if (Parameter <= ParameterRange.Min)
 		{
 			// Prepend
-			PrependPoint(Position, ParameterizationPolicy);
-			return 0;
+			return PrependPoint(Position, ParameterizationPolicy);
 		}
     
 		if (Parameter >= ParameterRange.Max)
 		{
-			// Append - create a point at the end
-			float StartParam = Base::GetParameterAt(Base::GetNumKeys() - 1);
-			ValueType P0 = Base::GetValueAt(Base::GetNumKeys() - 1);
-        
-			// Calculate reasonable auto-tangents
-			ValueType Dir = Position - P0;
-			ValueType P1 = P0 + Dir * 0.33f;
-			ValueType P2 = P0 + Dir * 0.66f;
-        
-			// Add the segment
-			AddBezierSegmentInternal(P0, P1, P2, Position, StartParam, ParameterizationPolicy);
-			return GetNumSegments();
+			return AppendPoint(Position, ParameterizationPolicy);
 		}
 
 		const float NormalizedParam = ParameterRange.GetT(Parameter);
@@ -1104,9 +1111,6 @@ public:
         ValueType R2 = P2 * mt + P3 * t;
         ValueType R3 = P3;
         
-        // Store indices for return value calculation
-        const int32 InsertedPointIndex = BaseIndex + 3;
-        
         float OriginalStartParam = Base::GetParameterAt(BaseIndex);
         float OriginalEndParam = Base::GetParameterAt(BaseIndex+3);
 		float TargetSplitParam = OriginalStartParam + LocalT * (OriginalEndParam - OriginalStartParam);
@@ -1124,41 +1128,9 @@ public:
 		
         // Add the two new segments using the policy
         AddBezierSegmentInternal(Q0, Q1, Q2, Q3, OriginalStartParam, ParameterizationPolicy, TargetSplitParam);
-		AddBezierSegmentInternal(R0, R1, R2, R3, TargetSplitParam, ParameterizationPolicy, OriginalEndParam);
-				
-        return InsertedPointIndex;
+		return AddBezierSegmentInternal(R0, R1, R2, R3, TargetSplitParam, ParameterizationPolicy, OriginalEndParam);
     }
 
-    /**
-     * Prepends a point to the spline at the beginning.
-     * The new point will be connected to the first segment.
-     * @param Position The position of the new point
-     * @param ParameterizationPolicy How to distribute parameters
-     * @return Index of the new point
-     */
-    int32 PrependPoint(
-		const ValueType& Position,
-		EParameterizationPolicy ParameterizationPolicy = EParameterizationPolicy::Centripetal)
-	{
-		// Get current parameter range start
-		
-		FInterval1f ParameterRange = this->GetParameterRange();
-    
-		// Use a parameter just before the current start
-		float PrependParam = ParameterRange.Min - 0.1f;
-    
-		// Create a control point at the beginning
-		ValueType P3 = Base::GetValueAt(0); // First point of the spline
-    
-		// Calculate reasonable auto-tangents
-		ValueType Dir = P3 - Position;
-		ValueType P1 = Position + Dir * 0.33f;
-		ValueType P2 = Position + Dir * 0.66f;
-    
-		// Add the segment by reusing existing code
-		return AddBezierSegmentInternal(Position, P1, P2, P3, PrependParam, ParameterizationPolicy, ParameterRange.Min);
-	}
-	
     /**
      * Inserts a Bezier segment at the specified parameter
      * 
@@ -1403,24 +1375,106 @@ public:
 		// Update parameter values in storage
 		this->ReparameterizeStorage(OldRange, NewRange);
     }
-    
+	/**
+	 * Clears all values from the spline including custom attribute channels and knot vector
+	 */
+	virtual void Clear()
+	{
+		Base::Clear();
+		this->GetEvaluator().KnotVector.Reset();
+	}
 protected:
     virtual TSharedPtr<Base> CreateNewInstance() const override
     {
     	return MakeShared<TPolyBezierSpline>(CreateDefault());
     }
+	
+	/**
+	 * Helper method to add the first segment to the spline
+	 * @param Parameter - Parameter value where to insert the point (Spline Space)
+	 * @param ControlPoint - Bezier control point to insert
+	 * @param PointIndex - Index of the newly created bezier control point
+	 * @return true if the points were added, false otherwise
+	 */
+	bool InitializeFirstSegment(float Parameter, const ValueType& ControlPoint, int32& PointIndex)
+    {
+    	if (Base::GetNumKeys() < 3)
+    	{
+    		// First point - just store it directly
+    		PointIndex = Base::AddValue(ControlPoint, Parameter);
+    		return true;
+    	}
 
-private:
-	// Default constructor used for serialization
-    TPolyBezierSpline()
-		: Base(MakeUnique<TPolyBezierEvaluator<ScalarType, Dimensions>>())
+    	// Adding the fourth control point would create a segment
+    	if (Base::GetNumKeys() == 3)
+    	{
+    		const ValueType& P0 = Base::GetValueAt(0);
+    		const ValueType& P1 = Base::GetValueAt(1);
+    		const ValueType& P2 = Base::GetValueAt(2);
+    		const ValueType& P3 = ControlPoint;
+	        
+    		// Add the first segment
+    		SetControlPoints( {P0, P1, P2, P3}, EParameterizationPolicy::Centripetal);
+    		PointIndex = 3;
+    		return true;
+    	}
+    	return false;
+    }
+
+	/**
+     * Prepends a point to the spline at the beginning.
+     * The new point will be connected to the first segment.
+     * @param Position The position of the new point
+     * @param ParameterizationPolicy How to distribute parameters
+     * @return Index of the new segment
+     */
+    int32 PrependPoint(
+		const ValueType& Position,
+		EParameterizationPolicy ParameterizationPolicy = EParameterizationPolicy::Centripetal)
+	{    	
+		// Get current parameter range start
+		FInterval1f ParameterRange = this->GetParameterRange();
+    
+		// Use a parameter just before the current start
+		float PrependParam = ParameterRange.Min - 0.1f;
+    
+		// Create a control point at the beginning
+		ValueType P3 = Base::GetValueAt(0); // First point of the spline
+    
+		// Calculate reasonable auto-tangents
+		ValueType Dir = P3 - Position;
+		ValueType P1 = Position + Dir * 0.33f;
+		ValueType P2 = Position + Dir * 0.66f;
+    
+		// Add the segment by reusing existing code
+		int32 NewIndex = AddBezierSegmentInternal(Position, P1, P2, P3, PrependParam, ParameterizationPolicy, ParameterRange.Min);
+		check (NewIndex == 0);
+		return NewIndex;
+	}
+
+	/**
+	 * Appends a point to the end of the spline.
+	 * Creates a new cubic Bezier segment connecting the last point to the new point.
+	 * 
+	 * @param Position The position of the new point
+	 * @param ParameterizationPolicy How to distribute parameters
+	 * @return Index of the new segment
+	 */
+	int32 AppendPoint(const ValueType& Position, EParameterizationPolicy ParameterizationPolicy)
 	{
-    	// Create a minimal valid default segment for serialization
-    	ValueType P0(0, 0, 0);
-    	ValueType P1(0.33f, 0, 0);
-    	ValueType P2(0.66f, 0, 0);
-    	ValueType P3(1, 0, 0);
-    	AddBezierSegmentInternal(P0, P1, P2, P3, 0, EParameterizationPolicy::Centripetal);
+		// Append - create a point at the end
+		float StartParam = Base::GetParameterAt(Base::GetNumKeys() - 1);
+		ValueType P0 = Base::GetValueAt(Base::GetNumKeys() - 1);
+        
+		// Calculate reasonable auto-tangents
+		ValueType Dir = Position - P0;
+		ValueType P1 = P0 + Dir * 0.33f;
+		ValueType P2 = P0 + Dir * 0.66f;
+        
+		// Add the segment
+		int32 NewIndex = AddBezierSegmentInternal(P0, P1, P2, Position, StartParam, ParameterizationPolicy);
+		check (NewIndex == GetNumSegments() - 1)
+		return NewIndex;
 	}
 	
 	/**
@@ -1588,7 +1642,7 @@ private:
     	Base::AddValue(P1, IntermediateParams[0]);
     	Base::AddValue(P2, IntermediateParams[1]);
     	Base::AddValue(P3, EndParamValue);
-    
+
     	this->GetEvaluator().InsertLocalKnots(StartParam, EndParamValue, Base::GetNumKeys());
     
     	return P0Index / 4;
diff --git a/Engine/Source/Runtime/GeometryCore/Public/Splines/TangentBezierSpline.h b/Engine/Source/Runtime/GeometryCore/Public/Splines/TangentBezierSpline.h
--- a/Engine/Source/Runtime/GeometryCore/Public/Splines/TangentBezierSpline.h
+++ b/Engine/Source/Runtime/GeometryCore/Public/Splines/TangentBezierSpline.h
@@ -123,14 +123,13 @@ struct FTangentBezierControlPoint
  */
 struct FTangentBezierSpline
 {
+	friend FPolyBezierSpline3d;
 
-protected:
 	/** 
 	 ** Initializes with a minimal default segment
 	 */
 	FTangentBezierSpline()
 		: Tension(0.0f)
-		, InternalSpline(FPolyBezierSpline3d::CreateLine(FVector::ZeroVector, FVector(100,0,0)))
 	{
 	}
 	
@@ -138,6 +137,7 @@ public:
 	/** Default constructor with at least one segment */
 	FTangentBezierSpline(const FVector& StartPoint, const FVector& EndPoint)
 		: Tension(0.0f)
+		, TangentModes({ETangentMode::Auto, ETangentMode::Auto})
 		, InternalSpline(FPolyBezierSpline3d::CreateLine(StartPoint, EndPoint))
 	{
 		
@@ -373,9 +373,6 @@ public:
 	            
 	            // Add the closing segment to the internal spline
 	            InternalSpline.AppendBezierSegment(P1, P2, P3);
-
-	        	// Add corresponding tangent mode
-	        	TangentModes.Add(ETangentMode::Auto);
 	        }
 	    }
 	    else
@@ -387,9 +384,6 @@ public:
 	    	{
 	    		// Remove the last segment
 	    		InternalSpline.RemoveSegment(NumSegments - 1);
-
-	    		// Remove last tangent mode
-	    		TangentModes.RemoveAt(TangentModes.Num() - 1);
 	    	}
         
 	    	// Now update the flag on the internal spline
@@ -628,7 +622,7 @@ public:
                 }
                 
                 // Add each point normally
-                AddPoint(Points[i]);
+                AppendPoint(Points[i]);
             }
         }
         else
@@ -667,31 +661,74 @@ public:
         return true;
     }
 
-	/**
-	 * Helper method to prepend a point to the spline.
-	 * 
-	 * @param ControlPoint Control point structure to prepend
-	 */
-	void PrependPoint(const FTangentBezierControlPoint& ControlPoint)
-	{
-		const int32 NumPoints = GetNumPoints();
-		if (NumPoints == 0)
+	
+
+    /**
+     * Inserts a new point at the specified parameter along the spline.
+     * 
+     * @param Parameter - Parameter value where to insert the point (Spline Space)
+	 * @param ControlPoint - Control point to insert
+     * @return The index of the newly created point
+     */
+    int32 InsertPoint(float Parameter, 
+                     const FTangentBezierControlPoint& ControlPoint)
+    {
+	    int32 NewPointIndex;
+	    if (InitializeFirstSegment(Parameter, ControlPoint, NewPointIndex)) return NewPointIndex;
+
+	    const int32 NumPoints = GetNumPoints();
+		int32 InsertedPointIndex = -1;
+		
+		// Insert at that parameter with the exact position
+		// Have the internal spline handle the insertion
+        InsertedPointIndex = InternalSpline.InsertPoint(Parameter, ControlPoint.Position);
+		
+		// If insertion was successful, update tangents if not auto-computed
+		if (InsertedPointIndex >= 0 &&
+			(ControlPoint.TangentMode == ETangentMode::User ||
+			 ControlPoint.TangentMode == ETangentMode::Broken))
 		{
-			// If spline is empty (shouldn't happen with our design), just use regular AddPoint
-			AddPoint(ControlPoint);
-			return;
+			SetTangentIn(InsertedPointIndex, ControlPoint.TangentIn);
+			SetTangentOut(InsertedPointIndex, ControlPoint.TangentOut);
 		}
         
-		// First ensure the tangent modes array is large enough for existing points
-		if (TangentModes.Num() < NumPoints)
+		SetPointTangentMode(InsertedPointIndex, ControlPoint.TangentMode);
+				
+		// Update adjacent points if they have automatic tangents
+		if (InsertedPointIndex > 0 && 
+			(TangentModes[InsertedPointIndex-1] == ETangentMode::Auto || 
+			 TangentModes[InsertedPointIndex-1] == ETangentMode::AutoClamped))
 		{
-			// Fill with default values (Auto) for any missing entries
-			TangentModes.AddDefaulted(NumPoints - TangentModes.Num());
+			UpdatePointTangents(InsertedPointIndex-1, TangentModes[InsertedPointIndex-1]);
 		}
         
-		// Now insert the new tangent mode at the beginning
-		TangentModes.Insert(ControlPoint.TangentMode, 0);
-        
+		if (NewPointIndex+1 < GetNumPoints() && 
+			(TangentModes[InsertedPointIndex+1] == ETangentMode::Auto || 
+			 TangentModes[InsertedPointIndex+1] == ETangentMode::AutoClamped))
+		{
+			UpdatePointTangents(InsertedPointIndex+1, TangentModes[InsertedPointIndex+1]);
+		}
+        return InsertedPointIndex;
+    }
+    
+    /**
+	 * Helper method to prepend a point to the spline.
+	 * 
+	 * @param ControlPoint Control point structure to prepend
+	 * @return Index of the new point
+	 */
+	void PrependPoint(const FTangentBezierControlPoint& ControlPoint)
+	{
+    	int32 NewIndex = GetNumPoints();
+    	FInterval1f ParameterRange = this->GetParameterRange();
+		float PrependParam = ParameterRange.Min - 0.1f;
+
+    	// Special case: handle first two points
+    	if (InitializeFirstSegment(PrependParam, ControlPoint, NewIndex))
+    	{
+    		return;
+    	}
+		        
 		// Calculate control points for the cubic Bezier segment
 		FVector P0 = ControlPoint.Position;
         
@@ -716,124 +753,76 @@ public:
 		// Let InternalSpline handle the prepending efficiently
 		InternalSpline.PrependBezierSegment(P0, P1, P2);
         
-		// Update tangents based on tangent mode
-		UpdatePointTangents(0, ControlPoint.TangentMode);
-        
+		SetPointTangentMode(0, ControlPoint.TangentMode);
+                
 		// Update tangents for the next point if it's auto
-		if (NumPoints > 0 && 
+		if (GetNumPoints() > 0 && 
 			(TangentModes[1] == ETangentMode::Auto || TangentModes[1] == ETangentMode::AutoClamped))
 		{
 			UpdatePointTangents(1, TangentModes[1]);
 		}
 	}
-    
-    /**
-     * Inserts a new point at the specified parameter along the spline.
-     * Uses distance-based parameterization to preserve attribute positions.
-     * 
-     * @param Parameter - Parameter value where to insert the point (Spline Space)
-     * @param Position - Position for the new point
-	 * @param InTangentMode - Tangent mode for the new point
-     * @return The index of the newly created point
+    /** 
+	 * Appends a control point to the end of the spline.
+	 * @param ControlPoint Structure containing position and tangent information
      */
-    int32 InsertPoint(float Parameter, 
-                     const FVector& Position,
-                     ETangentMode InTangentMode = ETangentMode::Auto)
+    void AppendPoint(const FTangentBezierControlPoint& ControlPoint)
     {
-               
-        // Have the internal spline handle the insertion
-        // This efficiently splits the Bezier segment while preserving attributes
-        int32 InsertedIndex = InternalSpline.InsertPoint(Parameter, Position);
-        
-        // Convert from PolyBezier control point index to TangentBezier point index
-        int32 NewPointIndex = (InsertedIndex + 1) / 4;
-        
-		// Insert tangent mode at the correct position
-		if (NewPointIndex >= TangentModes.Num())
-		{
-			TangentModes.SetNum(NewPointIndex + 1);
-		}
-		else
+        const int32 NumPoints = GetNumPoints();
+
+    	int32 NewIndex = GetNumPoints();
+
+    	FInterval1f ParameterRange = this->GetParameterRange();
+    	float AppendParam = ParameterRange.Max + 0.1f;
+    	// Special case: handle first two points
+    	if (InitializeFirstSegment(AppendParam, ControlPoint, NewIndex))
+    	{
+    		return;
+    	}
+    	
+    	
+	    // Normal case - add segment to existing spline
+        FVector P0 = GetPosition(NewIndex - 1);
+	    
+		        
+		FVector P3 = ControlPoint.Position;
+    
+		// Update the previous point's tangents if needed
+		if (TangentModes[NewIndex - 1] == ETangentMode::Auto || 
+			TangentModes[NewIndex - 1] == ETangentMode::AutoClamped)
 		{
-			TangentModes.Insert(InTangentMode, NewPointIndex);
+			UpdatePointTangents(NewIndex - 1, TangentModes[NewIndex - 1]);
 		}
-		TangentModes[NewPointIndex] = InTangentMode;
-        
-		// Update tangents to ensure proper curve continuity
-		UpdatePointTangents(NewPointIndex, InTangentMode);
-        
-		// Update adjacent points if they have automatic tangents
-		if (NewPointIndex > 0 && 
-			(TangentModes[NewPointIndex-1] == ETangentMode::Auto || 
-			 TangentModes[NewPointIndex-1] == ETangentMode::AutoClamped))
+    
+		// Get its outgoing tangent (may have been updated above)
+		FVector OutTangentPrev = GetTangentOut(NewIndex - 1);
+    
+		// Calculate control points for the new segment
+		FVector P1 = P0 + (OutTangentPrev / 3.0f);
+    
+		// For the incoming tangent of the new point
+		FVector P2;
+    
+		if (ControlPoint.TangentMode == ETangentMode::User || 
+			ControlPoint.TangentMode == ETangentMode::Broken)
 		{
-			UpdatePointTangents(NewPointIndex-1, TangentModes[NewPointIndex-1]);
+			// For user-defined tangents, use the provided tangent
+			P2 = P3 - (ControlPoint.TangentIn / 3.0f);
 		}
-        
-		if (NewPointIndex+1 < GetNumPoints() && 
-			(TangentModes[NewPointIndex+1] == ETangentMode::Auto || 
-			 TangentModes[NewPointIndex+1] == ETangentMode::AutoClamped))
+		else
 		{
-			UpdatePointTangents(NewPointIndex+1, TangentModes[NewPointIndex+1]);
+			// For other modes, we'll set up a reasonable default
+			// The UpdatePointTangents call below will update it properly
+			FVector DefaultTangent = (1.0f - Tension) * (P3 - P0);
+			P2 = P3 - (DefaultTangent / 3.0f);
 		}
-                
-        return NewPointIndex;
-    }
     
-    /** 
-     * Adds a new point to the spline
-	 * @param ControlPoint Structure containing position and tangent information
-     */
-    void AddPoint(const FTangentBezierControlPoint& ControlPoint)
-    {
-        const int32 NewIndex = GetNumPoints();
-        
-		// Store tangent mode
-		if (TangentModes.Num() <= NewIndex)
-		{
-			TangentModes.AddZeroed(NewIndex - TangentModes.Num() + 1);
-		}
-		TangentModes[NewIndex] = ControlPoint.TangentMode;
-     	    
-        FVector P0 = GetPosition(NewIndex - 1);
-        FVector P3 = ControlPoint.Position;
-        
-        // Update the previous point's tangents if needed
-        if (TangentModes[NewIndex - 1] == ETangentMode::Auto || 
-            TangentModes[NewIndex - 1] == ETangentMode::AutoClamped)
-        {
-            UpdatePointTangents(NewIndex - 1, TangentModes[NewIndex - 1]);
-        }
-        
-        // Get its outgoing tangent (may have been updated above)
-        FVector OutTangentPrev = GetTangentOut(NewIndex - 1);
-        
-        // Calculate control points for the new segment
-        FVector P1 = P0 + (OutTangentPrev / 3.0f);
-        
-        // For the incoming tangent of the new point
-        FVector P2;
-        
-        if (ControlPoint.TangentMode == ETangentMode::User || 
-            ControlPoint.TangentMode == ETangentMode::Broken)
-        {
-            // For user-defined tangents, use the provided tangent
-            P2 = P3 - (ControlPoint.TangentIn / 3.0f);
-        }
-        else
-        {
-            // For other modes, we'll set up a reasonable default
-            // The UpdatePointTangents call below will update it properly
-            FVector DefaultTangent = (1.0f - Tension) * (P3 - P0);
-            P2 = P3 - (DefaultTangent / 3.0f);
-        }
-        
-        // Add the segment
-        InternalSpline.AppendBezierSegment(P1, P2, P3);
-        
-        // Now update the tangents based on the tangent mode
-        UpdatePointTangents(NewIndex, ControlPoint.TangentMode);
-        
+		// Add the segment
+		InternalSpline.AppendBezierSegment(P1, P2, P3);
+	    
+		
+		SetPointTangentMode(NewIndex, ControlPoint.TangentMode);
+
         // If this is auto mode, update tangents for adjacent points
         if (ControlPoint.TangentMode == ETangentMode::Auto || 
             ControlPoint.TangentMode == ETangentMode::AutoClamped)
@@ -846,6 +835,7 @@ public:
             }
         }
     }
+
 	/**
 	 * Gets the control point at the specified index
 	 * @param Index Index of the point to retrieve
@@ -888,14 +878,26 @@ public:
 		{
 			return;
 		}
-    
-		// Prevent removing if we'd end up with fewer than 2 points
-		if (NumPoints <= 2)
+		// special case for handling the last point and second last point
+		if (NumPoints == 1)
 		{
-			UE_LOG(LogSpline, Warning, TEXT("Cannot remove point - spline must have at least 2 points."));
+			// Removing the only point
+			InternalSpline.RemoveValue(0);
+			TangentModes.Reset();
 			return;
 		}
-
+		else if (NumPoints == 2)
+		{
+			// Removing the last point and tangents
+			// For CP = 0, we remove P0, P1, P2
+			// For CP = 1, we remove P1, P2, P3
+			InternalSpline.RemoveValue(Index+2);
+			InternalSpline.RemoveValue(Index+1);
+			InternalSpline.RemoveValue(Index);
+			TangentModes.RemoveAt(Index);
+			return;
+		}
+    
 		// Remove the corresponding tangent mode
 		if (TangentModes.IsValidIndex(Index))
 		{
@@ -963,23 +965,27 @@ public:
         const FVector OldPosition = GetPosition(Index);
         bool bPositionChanged = (OldPosition != ControlPoint.Position);
         
-        // Store the tangent mode
-        if (TangentModes.Num() <= Index)
-        {
-            TangentModes.AddZeroed(Index - TangentModes.Num() + 1);
-        }
-        TangentModes[Index] = ControlPoint.TangentMode;
+        SetPointTangentMode(Index, ControlPoint.TangentMode);
         
         // Update position
         if (Index == 0)
         {
             // First point
             UpdateBezierControlPoint(Index * 4, ControlPoint.Position);
+
+        	if (IsClosedLoop())
+        	{
+        		UpdateBezierControlPoint((NumPoints - 1) * 4 + 3, ControlPoint.Position);
+        	}
         }
         else if (Index == NumPoints - 1)
         {
             // Last point
             UpdateBezierControlPoint(Index * 4 - 1, ControlPoint.Position);
+        	if (IsClosedLoop())
+        	{
+        		UpdateBezierControlPoint(Index * 4, ControlPoint.Position);
+        	}
         }
         else
         {
@@ -1095,18 +1101,10 @@ public:
     FVector GetTangentOut(int32 Index) const
     {
         const int32 NumPoints = GetNumPoints();
-        if (Index == NumPoints - 1)
-        {
-            // Special case for last point
-            if (InternalSpline.IsClosedLoop() && InternalSpline.GetNumSegments() > 0)
-            {
-                // For closed loop, look at first segment
-                FVector P0 = FVector(InternalSpline.GetValueAt(0));
-                FVector P1 = FVector(InternalSpline.GetValueAt(1));
-                return 3.0f * (P1 - P0);
-            }
-            return FVector::ZeroVector;
-        }
+        if (Index == NumPoints - 1 && !IsClosedLoop())
+		{
+			return FVector::ZeroVector;
+		}
         
         // Normal case - get from next segment's control points
         const int32 BezierIndex = Index * 4;
@@ -1154,15 +1152,8 @@ public:
     void SetTangentOut(int32 Index, const FVector& NewTangent)
     {
         const int32 NumPoints = GetNumPoints();
-        if (Index == NumPoints - 1)
+        if (Index == NumPoints - 1 && !IsClosedLoop())
         {
-            // Special case for last point in closed loop
-            if (InternalSpline.IsClosedLoop() && InternalSpline.GetNumSegments() > 0)
-            {
-                FVector P0 = FVector(InternalSpline.GetValueAt(0));
-                FVector P1 = P0 + (NewTangent / 3.0f);
-                UpdateBezierControlPoint(1, P1);
-            }
             return;
         }
         
@@ -1186,15 +1177,20 @@ public:
     	else if (NumKeys == 1)
     		return 1;  // Special case when we just have the first point
     	else
-    		return InternalSpline.GetNumSegments() + 1;  // Regular case: segments + 1
+    		return IsClosedLoop() ? InternalSpline.GetNumSegments() : InternalSpline.GetNumSegments() + 1;
     }
 
+	int32 GetNumSegments() const
+	{
+		return InternalSpline.GetNumSegments();
+	}
+
     /** Checks if a point's tangents are automatically computed */
     bool IsAutoTangent(int32 Index) const
     {
-		return TangentModes.IsValidIndex(Index) && TangentModes[Index] == ETangentMode::Auto;
+		return TangentModes.IsValidIndex(Index) && (TangentModes[Index] == ETangentMode::Auto || TangentModes[Index] == ETangentMode::AutoClamped);
     }
-    
+
 	/**
 	  * Sets the tangent mode for a specific point
 	  * @param Index - Index of the point
@@ -1209,15 +1205,15 @@ public:
 		}
         
 		// Ensure array is big enough
-		if (TangentModes.Num() <= Index)
+		if (TangentModes.Num() <= NumPoints)
 		{
-			TangentModes.AddZeroed(Index - TangentModes.Num() + 1);
+			TangentModes.AddZeroed(NumPoints - TangentModes.Num() + 1);
 		}
         
 		TangentModes[Index] = Mode;
-        
+
 		// Update tangents based on new mode
-		UpdatePointTangents(Index, Mode);
+        UpdatePointTangents(Index, Mode);
 	}
 
 	FVector GetValueAt(int32 Index) const
@@ -1381,10 +1377,12 @@ public:
 	void UpdatePointTangents(int32 Index, ETangentMode Mode)
 	{
 		const int32 NumPoints = GetNumPoints();
-		if (Index < 0 || Index >= NumPoints)
+		if (Index < 0 || Index >= NumPoints || Mode == ETangentMode::User || Mode == ETangentMode::Broken)
 		{
+			// Invalid index or user-defined mode, no action needed
 			return;
 		}
+		
 
 		// Get the position of the current point
 		const FVector& Current = GetPosition(Index);
@@ -1492,12 +1490,6 @@ public:
 			}
 			break;
 
-		case ETangentMode::Broken:
-		case ETangentMode::User:
-			// For User and Broken modes, we're keeping the existing tangents
-			// which we already fetched above, so no action needed here
-			break;
-
 		case ETangentMode::Unknown:
 		default:
 			// Default to Auto mode if unknown
@@ -1512,6 +1504,7 @@ public:
 		// Apply the calculated tangents
 		SetTangentIn(Index, InTangent);
 		SetTangentOut(Index, OutTangent);
+
 	}
     
     float FindNearest(const FVector& Point, float& OutSquaredDistance) const
@@ -1593,6 +1586,71 @@ protected:
             InternalSpline.UpdateValue(Index, NewValue);
         }
     }
+private:
+
+	static bool IsCurveKey(ETangentMode Mode)
+	{
+		return ((Mode == ETangentMode::Auto) || (Mode == ETangentMode::AutoClamped) || (Mode == ETangentMode::User) || (Mode == ETangentMode::Broken));
+	}
+	/**
+	 * Helper method to add the first segment to the spline
+	 * @param Parameter - Parameter value where to insert the point (Spline Space)
+	 * @param ControlPoint - Control point to insert
+	 * @param PointIndex - Index of the newly created point
+	 * @return true if the points were added, false otherwise
+	 */
+	bool InitializeFirstSegment(float Parameter, const FTangentBezierControlPoint& ControlPoint, int32& PointIndex)
+    {
+    	if (InternalSpline.GetNumKeys() == 0)
+    	{
+    		// First point - just store it directly
+    		InternalSpline.AddValue(FVector(ControlPoint.Position), Parameter);
+    		// update tangent mode
+            if (TangentModes.IsEmpty())
+            {
+	            TangentModes.Add(ControlPoint.TangentMode);
+            }
+    		else
+    		{
+    			TangentModes[0] = ControlPoint.TangentMode;
+    		}
+    		PointIndex = 0;
+    		return true;
+    	}
+
+    	if (InternalSpline.GetNumKeys() == 1)
+    	{
+    		// We have exactly one point stored - this is our second point
+    		// Get the first point, then clear to start building segments
+    		FVector P0 = FVector(InternalSpline.GetValueAt(0));
+    		FVector P3 = ControlPoint.Position;
+    		// Calculate control points
+    		FVector P1, P2;
+    		if (ControlPoint.TangentMode == ETangentMode::Auto ||
+				ControlPoint.TangentMode == ETangentMode::AutoClamped)
+    		{
+    			// For first segment with auto tangents, use simple interpolation
+    			FVector AutoTangent = (1.0f - Tension) * (P3 - P0);
+    			P1 = P0 + (AutoTangent / 3.0f);
+    			P2 = P3 - (AutoTangent / 3.0f);
+    		}
+    		else
+    		{
+    			// Use provided tangents
+    			P1 = P0 + (ControlPoint.TangentOut / 3.0f); // Need to save first point's OutTangent
+    			P2 = P3 - (ControlPoint.TangentIn / 3.0f);
+    		}
+    		
+	        // update tangent mode
+    		SetPointTangentMode(1, ControlPoint.TangentMode);
+    		
+    		// Add the first segment
+    		InternalSpline.SetControlPoints( {P0, P1, P2, P3}, EParameterizationPolicy::Centripetal);
+    		PointIndex = 1;
+    		return true;
+    	}
+    	return false;
+    }
 };
 } // end namespace UE::Geometry::Spline
 } // end namespace UE::Geometry
