Added validation when adding attribute value to a mismatched attribute channel

diff --git a/Engine/Source/Runtime/GeometryCore/Public/Splines/Spline.h b/Engine/Source/Runtime/GeometryCore/Public/Splines/Spline.h
--- a/Engine/Source/Runtime/GeometryCore/Public/Splines/Spline.h
+++ b/Engine/Source/Runtime/GeometryCore/Public/Splines/Spline.h
@@ -681,8 +681,22 @@ public:
 	template <typename AttrType>
 	const TAttributeChannel<AttrType>* GetAttributeChannel(const FName& Name) const
 	{
-		const TUniquePtr<IAttributeChannelBase>* Found = AttributeChannels.Find(Name);
-		return Found ? static_cast<const TAttributeChannel<AttrType>*>(Found->Get()) : nullptr;
+		
+		// Check if a channel with this name already exists
+		if (const TUniquePtr<IAttributeChannelBase>* ExistingChannel = AttributeChannels.Find(Name))
+		{
+			const FName TypeName = TAttributeTypeTraits<AttrType>::Name;
+			const FName ExistingTypeName = (*ExistingChannel)->GetTypeName();
+			
+			if (TypeName == ExistingTypeName)
+			{
+				return static_cast<const TAttributeChannel<AttrType>*>(ExistingChannel->Get());
+			}
+			// Channel exists but with a different type
+			UE_LOG(LogSpline, Warning, TEXT("Failed to fetch attribute channel '%s': Channel already exists with type '%s', but requested type is '%s'."),
+					*Name.ToString(), *ExistingTypeName.ToString(), *TypeName.ToString());
+		}
+		return nullptr;
 	}
 	
 	const TAttributeChannel<ValueType>& GetValuesChannel() const
@@ -693,8 +707,21 @@ public:
 	template <typename AttrType>
 	TAttributeChannel<AttrType>* GetAttributeChannel(const FName& Name)
 	{
-		TUniquePtr<IAttributeChannelBase>* Found = AttributeChannels.Find(Name);
-		return Found ? static_cast<TAttributeChannel<AttrType>*>(Found->Get()) : nullptr;
+		// Check if a channel with this name already exists
+		if (TUniquePtr<IAttributeChannelBase>* ExistingChannel = AttributeChannels.Find(Name))
+		{
+			const FName TypeName = TAttributeTypeTraits<AttrType>::Name;
+			const FName ExistingTypeName = (*ExistingChannel)->GetTypeName();
+			
+			if (TypeName == ExistingTypeName)
+			{
+				return static_cast<TAttributeChannel<AttrType>*>(ExistingChannel->Get());
+			}
+			// Channel exists but with a different type
+			UE_LOG(LogSpline, Warning, TEXT("Failed to fetch attribute channel '%s': Channel already exists with type '%s', but requested type is '%s'."),
+					*Name.ToString(), *ExistingTypeName.ToString(), *TypeName.ToString());
+		}
+		return nullptr;
 	}
 
 	const TMap<FName, TUniquePtr<IAttributeChannelBase>>& GetAttributeChannels() const
@@ -710,7 +737,6 @@ public:
     
 		for (const auto& Pair : AttributeChannels)
 		{
-			// Compare type names directly instead of using dynamic_cast
 			if (Pair.Value->GetTypeName() == TypeName)
 			{
 				Names.Add(Pair.Key);
