add a bitflag enum to the mesh transforms so the caller can specify which attributes should be transformed
+ clean up the mesh transform logic to mostly use a shared helper function
+ make the lattice deform tool use the bitflags to transform w/out the sculpt layers

diff --git a/Engine/Plugins/Runtime/MeshModelingToolset/Source/MeshModelingTools/Private/LatticeDeformerTool.cpp b/Engine/Plugins/Runtime/MeshModelingToolset/Source/MeshModelingTools/Private/LatticeDeformerTool.cpp
--- a/Engine/Plugins/Runtime/MeshModelingToolset/Source/MeshModelingTools/Private/LatticeDeformerTool.cpp
+++ b/Engine/Plugins/Runtime/MeshModelingToolset/Source/MeshModelingTools/Private/LatticeDeformerTool.cpp
@@ -464,25 +464,9 @@ void ULatticeDeformerTool::OnShutdown(EToolShutdownType ShutdownType)
 
 			// The lattice and its output mesh are in world space, so get them in local space.
 			// TODO: Would it make more sense to do all the lattice computation in local space?
-			const FTransformSRT3d LocalToWorld(TargetComponent->GetWorldTransform());
-			const bool bHasVertexNormals = DynamicMeshResult->HasVertexNormals();
-			const int32 NumVertices = DynamicMeshResult->MaxVertexID();
-			ParallelFor(NumVertices, [&DynamicMeshResult, &LocalToWorld, bHasVertexNormals](int32 VertexIndex)
-			{
-				if (DynamicMeshResult->IsVertex(VertexIndex))
-				{
-					FVector3d Position = DynamicMeshResult->GetVertex(VertexIndex);
-					Position = LocalToWorld.InverseTransformPosition(Position);
-					DynamicMeshResult->SetVertex(VertexIndex, Position);
-
-					if (bHasVertexNormals)
-					{
-						FVector3f Normal = DynamicMeshResult->GetVertexNormal(VertexIndex);
-						Normal = LocalToWorld.InverseTransformNormal(Normal);
-						DynamicMeshResult->SetVertexNormal(VertexIndex, Normalized(Normal));
-					}
-				}
-			});
+			// Note: We skip transforming sculpt layers, since they were never transformed to world space
+			FTransform3d LocalToWorld(TargetComponent->GetWorldTransform());
+			MeshTransforms::ApplyTransformInverse(*DynamicMeshResult, LocalToWorld, true, ~MeshTransforms::ETransformAttributes::SculptLayers);
 
 
 
diff --git a/Engine/Source/Runtime/GeometryCore/Private/DynamicMesh/MeshTransforms.cpp b/Engine/Source/Runtime/GeometryCore/Private/DynamicMesh/MeshTransforms.cpp
--- a/Engine/Source/Runtime/GeometryCore/Private/DynamicMesh/MeshTransforms.cpp
+++ b/Engine/Source/Runtime/GeometryCore/Private/DynamicMesh/MeshTransforms.cpp
@@ -43,287 +43,190 @@ namespace MeshTransformsLocal
 			);
 		}
 	}
-}
-
 
-void MeshTransforms::Translate(FDynamicMesh3& Mesh, const FVector3d& Translation)
-{
-	int NumVertices = Mesh.MaxVertexID();
-	ParallelFor(NumVertices, [&](int vid) 
+	template<typename TransformPosF, typename TransformNormalF, typename TransformTangentF, typename TransformOffsetF>
+	void TransformMeshHelper(FDynamicMesh3& Mesh, TransformPosF TransformPosition, TransformNormalF TransformNormal, TransformTangentF TransformTangent, TransformOffsetF TransformOffset, MeshTransforms::ETransformAttributes TransformAttributes)
 	{
-		if (Mesh.IsVertex(vid))
-		{
-			Mesh.SetVertex(vid, Mesh.GetVertex(vid) + Translation);
-		}
-	});
-	MeshTransformsLocal::TransformSculptLayers(Mesh, 
-		[&Translation](const FVector3d& Pos) {return Pos + Translation;}, 
-		MeshTransformsLocal::ESkipOffset::SkipOffset);
-}
+		using namespace MeshTransforms;
 
+		const bool bVertexNormals = bool(TransformAttributes & ETransformAttributes::VertexNormals) && Mesh.HasVertexNormals();
+		const bool bPositions = bool(TransformAttributes & ETransformAttributes::Positions);
 
-void MeshTransforms::Scale(FDynamicMesh3& Mesh, const FVector3d& Scale, const FVector3d& Origin, bool bReverseOrientationIfNeeded, bool bOnlyPositions)
-{
-	bool bVertexNormals = Mesh.HasVertexNormals();
-	bool bNeedNormalTangentScaling = !bOnlyPositions && !Scale.IsUniform();
-	int NumVertices = Mesh.MaxVertexID();
-	FVector3f NormalScale = (FVector3f)Scale;
-	if (bNeedNormalTangentScaling)
-	{
-		// compute a safe inverse-scale to apply to normal vectors
-		for (int32 Idx = 0; Idx < 3; ++Idx)
+		const int NumVertices = Mesh.MaxVertexID();
+		ParallelFor(NumVertices, [bPositions, bVertexNormals, &Mesh, &TransformPosition, &TransformNormal](int vid)
 		{
-			if (!FMath::IsNearlyZero(NormalScale[Idx]))
+			if (Mesh.IsVertex(vid)) 
 			{
-				NormalScale[Idx] = 1.0f / NormalScale[Idx];
-			}
-		}
-	}
-	ParallelFor(NumVertices, [&](int32 vid)
-	{
-		if (Mesh.IsVertex(vid))
-		{
-			Mesh.SetVertex(vid, (Mesh.GetVertex(vid) - Origin) * Scale + Origin);
+				if (bPositions)
+				{
+					FVector3d Position = Mesh.GetVertex(vid);
+					Position = TransformPosition(Position);
+					Mesh.SetVertex(vid, Position);
+				}
 
-			if (bNeedNormalTangentScaling && bVertexNormals)
-			{
-				FVector3f ScaledNormal = Mesh.GetVertexNormal(vid) * NormalScale;
-				ScaledNormal.Normalize();
-				Mesh.SetVertexNormal(vid, ScaledNormal);
+				if (bVertexNormals)
+				{
+					FVector3f Normal = Mesh.GetVertexNormal(vid);
+					Normal = TransformNormal(Normal);
+					Mesh.SetVertexNormal(vid, Normal);
+				}
 			}
-		}
-	});
-	if (bNeedNormalTangentScaling && Mesh.HasAttributes())
-	{
-		const int32 NumNormalLayers = FMath::Min(Mesh.Attributes()->NumNormalLayers(), 3);
-		for (int32 NormalTangentLayerIdx = 0; NormalTangentLayerIdx < NumNormalLayers; ++NormalTangentLayerIdx)
+		});
+
+		if (Mesh.HasAttributes())
 		{
-			FVector3f ScaleBy = (FVector3f)Scale; // tangents are transformed by the scale directly
-			if (NormalTangentLayerIdx == 0) // normal is transformed by the inverse scale
+			FDynamicMeshNormalOverlay* Normals = Mesh.Attributes()->PrimaryNormals();
+			if (Normals && bool(TransformAttributes & ETransformAttributes::Normals))
 			{
-				ScaleBy = NormalScale;
-			}
-			FDynamicMeshNormalOverlay* Normals = Mesh.Attributes()->GetNormalLayer(NormalTangentLayerIdx);
-			int NumNormals = Normals->MaxElementID();
-			ParallelFor(NumNormals, [&](int ElemID)
+				int NumNormals = Normals->MaxElementID();
+				ParallelFor(NumNormals, [&TransformNormal, &Normals](int ElemID)
 				{
 					if (Normals->IsElement(ElemID))
 					{
-						FVector3f ScaledNormal = Normals->GetElement(ElemID) * ScaleBy;
-						ScaledNormal.Normalize();
-						Normals->SetElement(ElemID, ScaledNormal);
+						FVector3f Normal = Normals->GetElement(ElemID);
+						Normal = TransformNormal(Normal);
+						Normals->SetElement(ElemID, Normal);
 					}
 				});
+			}
+			if (Mesh.Attributes()->HasTangentSpace() && bool(TransformAttributes & ETransformAttributes::Tangents))
+			{
+				for (int32 TangentLayerIdx = 1; TangentLayerIdx < 3; ++TangentLayerIdx)
+				{
+					FDynamicMeshNormalOverlay* TangentLayer = Mesh.Attributes()->GetNormalLayer(TangentLayerIdx);
+					int NumTangents = TangentLayer->MaxElementID();
+					ParallelFor(NumTangents, [&TransformTangent, &TangentLayer](int ElemID)
+					{
+						if (TangentLayer->IsElement(ElemID))
+						{
+							FVector3f Tangent = TangentLayer->GetElement(ElemID);
+							Tangent = TransformTangent(Tangent);
+							TangentLayer->SetElement(ElemID, Normalized(Tangent));
+						}
+					});
+				}
+			}
+		
+			if (bool(TransformAttributes & ETransformAttributes::SculptLayers))
+			{
+				TransformSculptLayers(Mesh,
+					TransformPosition,
+					TransformOffset);
+			}
 		}
 	}
-	
-	MeshTransformsLocal::TransformSculptLayers(Mesh, 
-		[&Origin, &Scale](const FVector3d& Pos) {return (Pos - Origin) * Scale + Origin;}, 
-		[&Scale](const FVector3d& Offset) { return Offset * Scale;});
-
-	if (bReverseOrientationIfNeeded && Scale.X * Scale.Y * Scale.Z < 0)
-	{
-		Mesh.ReverseOrientation(false);
-	}
 }
 
 
-
-void MeshTransforms::WorldToFrameCoords(FDynamicMesh3& Mesh, const FFrame3d& Frame)
+void MeshTransforms::Translate(FDynamicMesh3& Mesh, const FVector3d& Translation, ETransformAttributes TransformAttributes)
 {
-	bool bVertexNormals = Mesh.HasVertexNormals();
-
 	int NumVertices = Mesh.MaxVertexID();
-	ParallelFor(NumVertices, [&](int vid)
+	if (bool(TransformAttributes & ETransformAttributes::Positions))
 	{
-		if (Mesh.IsVertex(vid))
+		ParallelFor(NumVertices, [&](int vid) 
 		{
-			FVector3d Position = Mesh.GetVertex(vid);
-			Mesh.SetVertex(vid, Frame.ToFramePoint(Position));
-
-			if (bVertexNormals)
+			if (Mesh.IsVertex(vid))
 			{
-				FVector3f Normal = Mesh.GetVertexNormal(vid);
-				Mesh.SetVertexNormal(vid, (FVector3f)Frame.ToFrameVector((FVector3d)Normal));
+				Mesh.SetVertex(vid, Mesh.GetVertex(vid) + Translation);
 			}
-		}
-	});
-
-	if (Mesh.HasAttributes())
+		});
+	}
+	if (bool(TransformAttributes & ETransformAttributes::SculptLayers))
 	{
-		for (int32 NormalTangentLayerIdx = 0; NormalTangentLayerIdx < Mesh.Attributes()->NumNormalLayers(); ++NormalTangentLayerIdx)
-		{
-			// Note: Because this is a rotation transform, normals and tangents transform the same way
-			FDynamicMeshNormalOverlay* Normals = Mesh.Attributes()->GetNormalLayer(NormalTangentLayerIdx);
-			int NumNormals = Normals->MaxElementID();
-			ParallelFor(NumNormals, [&](int ElemID)
-			{
-				if (Normals->IsElement(ElemID))
-				{
-					FVector3f Normal = Normals->GetElement(ElemID);
-					Normals->SetElement(ElemID, (FVector3f)Frame.ToFrameVector((FVector3d)Normal));
-				}
-			});
-		}
-
-		MeshTransformsLocal::TransformSculptLayers(Mesh, 
-			[&Frame](const FVector3d& Pos) {return Frame.ToFramePoint(Pos); }, 
-			[&Frame](const FVector3d& Offset) { return Frame.ToFrameVector(Offset); });
+		MeshTransformsLocal::TransformSculptLayers(Mesh,
+			[&Translation](const FVector3d& Pos) {return Pos + Translation;},
+			MeshTransformsLocal::ESkipOffset::SkipOffset);
 	}
 }
 
-
-
-
-
-void MeshTransforms::FrameCoordsToWorld(FDynamicMesh3& Mesh, const FFrame3d& Frame)
+void MeshTransforms::Scale(FDynamicMesh3& Mesh, const FVector3d& Scale, const FVector3d& Origin, bool bReverseOrientationIfNeeded, bool bOnlyTransformPositions)
 {
-	bool bVertexNormals = Mesh.HasVertexNormals();
+	ETransformAttributes TransformAttributes = bOnlyTransformPositions ? ETransformAttributes::Positions : ETransformAttributes::All;
+	MeshTransforms::Scale(Mesh, Scale, Origin, bReverseOrientationIfNeeded, TransformAttributes);
+}
 
+void MeshTransforms::Scale(FDynamicMesh3& Mesh, const FVector3d& Scale, const FVector3d& Origin, bool bReverseOrientationIfNeeded, ETransformAttributes TransformAttributes)
+{
+	const ETransformAttributes NormalsOrTangents = ETransformAttributes::VertexNormals | ETransformAttributes::Normals | ETransformAttributes::Tangents;
+	bool bNeedNormalTangentScaling = bool(TransformAttributes & NormalsOrTangents) && !Scale.IsUniform();
 	int NumVertices = Mesh.MaxVertexID();
-	ParallelFor(NumVertices, [&](int vid)
+	FVector3f TangentScale = (FVector3f)Scale;
+	FVector3f NormalScale = TangentScale;
+	if (bNeedNormalTangentScaling)
 	{
-		if (Mesh.IsVertex(vid))
+		// compute a safe inverse-scale to apply to normal vectors
+		for (int32 Idx = 0; Idx < 3; ++Idx)
 		{
-			FVector3d Position = Mesh.GetVertex(vid);
-			Mesh.SetVertex(vid, Frame.FromFramePoint(Position));
-
-			if (bVertexNormals)
+			if (!FMath::IsNearlyZero(NormalScale[Idx]))
 			{
-				FVector3f Normal = Mesh.GetVertexNormal(vid);
-				Mesh.SetVertexNormal(vid, (FVector3f)Frame.FromFrameVector((FVector3d)Normal));
+				NormalScale[Idx] = 1.0f / NormalScale[Idx];
 			}
 		}
-	});
-
-	if (Mesh.HasAttributes())
+	}
+	else
 	{
-		for (int32 NormalTangentLayerIdx = 0; NormalTangentLayerIdx < Mesh.Attributes()->NumNormalLayers(); ++NormalTangentLayerIdx)
-		{
-			// Note: Because this is a rotation transform, normals and tangents transform the same way
-			FDynamicMeshNormalOverlay* Normals = Mesh.Attributes()->GetNormalLayer(NormalTangentLayerIdx);
-			int NumNormals = Normals->MaxElementID();
-			ParallelFor(NumNormals, [&](int ElemID)
-			{
-				if (Normals->IsElement(ElemID))
-				{
-					FVector3f Normal = Normals->GetElement(ElemID);
-					Normals->SetElement(ElemID, (FVector3f)Frame.FromFrameVector((FVector3d)Normal));
-				}
-			});
-		}
+		// we don't need normal/tangent scaling, so disable them
+		TransformAttributes = TransformAttributes & (~NormalsOrTangents);
+	}
+	
+	MeshTransformsLocal::TransformMeshHelper(Mesh, 
+		[&Origin, &Scale](const FVector3d& Pos) -> FVector3d {return (Pos - Origin) * Scale + Origin;},
+		[NormalScale](const FVector3f& Normal) -> FVector3f { return Normalized(Normal * NormalScale); },
+		[TangentScale](const FVector3f& Tangent) -> FVector3f { return Normalized(Tangent * TangentScale); },
+		[&Scale](const FVector3d& Offset) -> FVector3d { return Offset * Scale;},
+		TransformAttributes);
 
-		MeshTransformsLocal::TransformSculptLayers(Mesh, 
-			[&Frame](const FVector3d& Pos) {return Frame.FromFramePoint(Pos); }, 
-			[&Frame](const FVector3d& Offset) { return Frame.FromFrameVector(Offset); });
+	if (bReverseOrientationIfNeeded && Scale.X * Scale.Y * Scale.Z < 0)
+	{
+		Mesh.ReverseOrientation(false);
 	}
 }
 
 
-void MeshTransforms::Rotate(FDynamicMesh3& Mesh, const FRotator& Rotation, const FVector3d& RotationOrigin)
+
+void MeshTransforms::WorldToFrameCoords(FDynamicMesh3& Mesh, const FFrame3d& Frame, ETransformAttributes TransformAttributes)
 {
-	bool bVertexNormals = Mesh.HasVertexNormals();
+	MeshTransformsLocal::TransformMeshHelper(Mesh,
+		[&Frame](const FVector3d& Pos) {return Frame.ToFramePoint(Pos); },
+		[&Frame](const FVector3f& Normal) -> FVector3f { return (FVector3f)Frame.ToFrameVector((FVector3d)Normal); },
+		[&Frame](const FVector3f& Tangent) -> FVector3f { return (FVector3f)Frame.ToFrameVector((FVector3d)Tangent); },
+		[&Frame](const FVector3d& Offset) { return Frame.ToFrameVector(Offset); },
+		TransformAttributes);
+}
+
 
-	int NumVertices = Mesh.MaxVertexID();
-	ParallelFor(NumVertices, [&](int vid)
-	{
-		if (Mesh.IsVertex(vid)) 
-		{
-			FVector3d Position = Rotation.RotateVector((Mesh.GetVertex(vid) - RotationOrigin)) + RotationOrigin;
-			Mesh.SetVertex(vid, Position);
 
-			if (bVertexNormals)
-			{
-				FVector3f Normal = Mesh.GetVertexNormal(vid);
-				Normal = (FVector3f)Rotation.RotateVector((FVector3d)Normal);
-				Mesh.SetVertexNormal(vid, Normal);
-			}
-		}
-	});
 
-	if (Mesh.HasAttributes())
-	{
-		for (int32 NormalLayerIdx = 0; NormalLayerIdx < Mesh.Attributes()->NumNormalLayers(); ++NormalLayerIdx)
-		{
-			FDynamicMeshNormalOverlay* NormalTangentLayer = Mesh.Attributes()->GetNormalLayer(NormalLayerIdx);
-			int NumElements = NormalTangentLayer->MaxElementID();
-			ParallelFor(NumElements, [&](int ElemID)
-			{
-				if (NormalTangentLayer->IsElement(ElemID))
-				{
-					FVector3f NormalTangent = NormalTangentLayer->GetElement(ElemID);
-					NormalTangent = (FVector3f)Rotation.RotateVector((FVector3d)NormalTangent);
-					NormalTangentLayer->SetElement(ElemID, NormalTangent);
-				}
-			});
-		}
-		
-		MeshTransformsLocal::TransformSculptLayers(Mesh, 
-			[&Rotation, &RotationOrigin](const FVector3d& Pos) {return Rotation.RotateVector((Pos - RotationOrigin)) + RotationOrigin; }, 
-			[&Rotation](const FVector3d& Offset) { return Rotation.RotateVector(Offset); });
-	}
-}
 
-void MeshTransforms::ApplyTransform(FDynamicMesh3& Mesh, const FTransformSRT3d& Transform, bool bReverseOrientationIfNeeded)
+void MeshTransforms::FrameCoordsToWorld(FDynamicMesh3& Mesh, const FFrame3d& Frame, ETransformAttributes TransformAttributes)
 {
-	bool bVertexNormals = Mesh.HasVertexNormals();
+	MeshTransformsLocal::TransformMeshHelper(Mesh,
+		[&Frame](const FVector3d& Pos) {return Frame.FromFramePoint(Pos); },
+		[&Frame](const FVector3f& Normal) -> FVector3f { return (FVector3f)Frame.FromFrameVector((FVector3d)Normal); },
+		[&Frame](const FVector3f& Tangent) -> FVector3f { return (FVector3f)Frame.FromFrameVector((FVector3d)Tangent); },
+		[&Frame](const FVector3d& Offset) { return Frame.FromFrameVector(Offset); },
+		TransformAttributes);
+}
 
-	int NumVertices = Mesh.MaxVertexID();
-	ParallelFor(NumVertices, [&](int vid)
-	{
-		if (Mesh.IsVertex(vid)) 
-		{
-			FVector3d Position = Mesh.GetVertex(vid);
-			Position = Transform.TransformPosition(Position);
-			Mesh.SetVertex(vid, Position);
 
-			if (bVertexNormals)
-			{
-				FVector3f Normal = Mesh.GetVertexNormal(vid);
-				Normal = (FVector3f)Transform.TransformNormal((FVector3d)Normal);
-				Mesh.SetVertexNormal(vid, Normal);
-			}
-		}
-	});
-
-	if (Mesh.HasAttributes())
-	{
-		FDynamicMeshNormalOverlay* Normals = Mesh.Attributes()->PrimaryNormals();
-		if (Normals)
-		{
-			int NumNormals = Normals->MaxElementID();
-			ParallelFor(NumNormals, [&](int ElemID)
-			{
-				if (Normals->IsElement(ElemID))
-				{
-					FVector3f Normal = Normals->GetElement(ElemID);
-					Normal = (FVector3f)Transform.TransformNormal((FVector3d)Normal);
-					Normals->SetElement(ElemID, Normal);
-				}
-			});
-		}
-		if (Mesh.Attributes()->HasTangentSpace())
-		{
-			for (int32 TangentLayerIdx = 1; TangentLayerIdx < 3; ++TangentLayerIdx)
-			{
-				FDynamicMeshNormalOverlay* TangentLayer = Mesh.Attributes()->GetNormalLayer(TangentLayerIdx);
-				int NumTangents = TangentLayer->MaxElementID();
-				ParallelFor(NumTangents, [&](int ElemID)
-				{
-					if (TangentLayer->IsElement(ElemID))
-					{
-						FVector3f Tangent = TangentLayer->GetElement(ElemID);
-						Tangent = (FVector3f)Transform.TransformVector((FVector3d)Tangent);
-						TangentLayer->SetElement(ElemID, Normalized(Tangent));
-					}
-				});
-			}
-		}
+void MeshTransforms::Rotate(FDynamicMesh3& Mesh, const FRotator& Rotation, const FVector3d& RotationOrigin, ETransformAttributes TransformAttributes)
+{
+	MeshTransformsLocal::TransformMeshHelper(Mesh,
+		[&Rotation, &RotationOrigin](const FVector3d& Pos) {return Rotation.RotateVector((Pos - RotationOrigin)) + RotationOrigin; },
+		[&Rotation](const FVector3f& Normal) -> FVector3f { return (FVector3f)Rotation.RotateVector((FVector3d)Normal); },
+		[&Rotation](const FVector3f& Tangent) -> FVector3f { return (FVector3f)Rotation.RotateVector((FVector3d)Tangent); },
+		[&Rotation](const FVector3d& Offset) { return Rotation.RotateVector(Offset); },
+		TransformAttributes);
+}
 
-		MeshTransformsLocal::TransformSculptLayers(Mesh,
-			[&Transform](const FVector3d& Pos) {return Transform.TransformPosition(Pos); },
-			[&Transform](const FVector3d& Offset) { return Transform.TransformVector(Offset); });
-	}
+void MeshTransforms::ApplyTransform(FDynamicMesh3& Mesh, const FTransformSRT3d& Transform, bool bReverseOrientationIfNeeded, ETransformAttributes TransformAttributes)
+{
+	MeshTransformsLocal::TransformMeshHelper(Mesh,
+		[&Transform](const FVector3d& Pos) {return Transform.TransformPosition(Pos); },
+		[&Transform](const FVector3f& Normal) -> FVector3f { return (FVector3f)Transform.TransformNormal((FVector3d)Normal); },
+		[&Transform](const FVector3f& Tangent) -> FVector3f { return Normalized((FVector3f)Transform.TransformVector((FVector3d)Tangent)); },
+		[&Transform](const FVector3d& Offset) { return Transform.TransformVector(Offset); },
+		TransformAttributes);
 
 	if (bReverseOrientationIfNeeded && Transform.GetDeterminant() < 0)
 	{
@@ -333,66 +236,14 @@ void MeshTransforms::ApplyTransform(FDynamicMesh3& Mesh, const FTransformSRT3d&
 
 
 
-void MeshTransforms::ApplyTransformInverse(FDynamicMesh3& Mesh, const FTransformSRT3d& Transform, bool bReverseOrientationIfNeeded)
+void MeshTransforms::ApplyTransformInverse(FDynamicMesh3& Mesh, const FTransformSRT3d& Transform, bool bReverseOrientationIfNeeded, ETransformAttributes TransformAttributes)
 {
-	bool bVertexNormals = Mesh.HasVertexNormals();
-
-	int NumVertices = Mesh.MaxVertexID();
-	ParallelFor(NumVertices, [&](int vid)
-	{
-		if (Mesh.IsVertex(vid)) 
-		{
-			FVector3d Position = Mesh.GetVertex(vid);
-			Position = Transform.InverseTransformPosition(Position);
-			Mesh.SetVertex(vid, Position);
-
-			if (bVertexNormals)
-			{
-				FVector3f Normal = Mesh.GetVertexNormal(vid);
-				Normal = (FVector3f)Transform.InverseTransformNormal((FVector3d)Normal);
-				Mesh.SetVertexNormal(vid, Normal);
-			}
-		}
-	});
-
-	if (Mesh.HasAttributes())
-	{
-		FDynamicMeshNormalOverlay* Normals = Mesh.Attributes()->PrimaryNormals();
-		if (Normals)
-		{
-			int NumNormals = Normals->MaxElementID();
-			ParallelFor(NumNormals, [&](int ElemID)
-			{
-				if (Normals->IsElement(ElemID))
-				{
-					FVector3f Normal = Normals->GetElement(ElemID);
-					Normal = (FVector3f)Transform.InverseTransformNormal((FVector3d)Normal);
-					Normals->SetElement(ElemID, Normal);
-				}
-			});
-		}
-		if (Mesh.Attributes()->HasTangentSpace())
-		{
-			for (int32 TangentLayerIdx = 1; TangentLayerIdx < 3; ++TangentLayerIdx)
-			{
-				FDynamicMeshNormalOverlay* TangentLayer = Mesh.Attributes()->GetNormalLayer(TangentLayerIdx);
-				int NumTangents = TangentLayer->MaxElementID();
-				ParallelFor(NumTangents, [&](int ElemID)
-				{
-					if (TangentLayer->IsElement(ElemID))
-					{
-						FVector3f Tangent = TangentLayer->GetElement(ElemID);
-						Tangent = (FVector3f)Transform.InverseTransformVector((FVector3d)Tangent);
-						TangentLayer->SetElement(ElemID, Normalized(Tangent));
-					}
-				});
-			}
-		}
-		
-		MeshTransformsLocal::TransformSculptLayers(Mesh,
-			[&Transform](const FVector3d& Pos) {return Transform.InverseTransformPosition(Pos); },
-			[&Transform](const FVector3d& Offset) { return Transform.InverseTransformVector(Offset); });
-	}
+	MeshTransformsLocal::TransformMeshHelper(Mesh,
+		[&Transform](const FVector3d& Pos) {return Transform.InverseTransformPosition(Pos); },
+		[&Transform](const FVector3f& Normal) -> FVector3f { return (FVector3f)Transform.InverseTransformNormal((FVector3d)Normal); },
+		[&Transform](const FVector3f& Tangent) -> FVector3f { return Normalized((FVector3f)Transform.InverseTransformVector((FVector3d)Tangent)); },
+		[&Transform](const FVector3d& Offset) { return Transform.InverseTransformVector(Offset); },
+		TransformAttributes);
 
 	if (bReverseOrientationIfNeeded && Transform.GetDeterminant() < 0)
 	{
@@ -415,113 +266,29 @@ void MeshTransforms::ApplyTransform(FDynamicMesh3& Mesh,
 	TFunctionRef<FVector3d(const FVector3d&)> PositionTransform,
 	TFunctionRef<FVector3f(const FVector3f&)> NormalTransform)
 {
-	bool bVertexNormals = Mesh.HasVertexNormals();
-
-	int NumVertices = Mesh.MaxVertexID();
-	ParallelFor(NumVertices, [&](int vid)
-	{
-		if (Mesh.IsVertex(vid))
+	ApplyTransform(Mesh, PositionTransform, NormalTransform, 
+		[](const FVector3f& T)->FVector3f 
 		{
-			FVector3d Position = Mesh.GetVertex(vid);
-			Position = PositionTransform(Position);
-			Mesh.SetVertex(vid, Position);
-
-			if (bVertexNormals)
-			{
-				FVector3f Normal = Mesh.GetVertexNormal(vid);
-				Normal = NormalTransform(Normal);
-				Mesh.SetVertexNormal(vid, Normalized(Normal));
-			}
-		}
-	});
-
-	if (Mesh.HasAttributes())
-	{
-		FDynamicMeshNormalOverlay* Normals = Mesh.Attributes()->PrimaryNormals();
-		if (Normals)
-		{
-			int NumNormals = Normals->MaxElementID();
-			ParallelFor(NumNormals, [&](int ElemID)
-			{
-				if (Normals->IsElement(ElemID))
-				{
-					FVector3f Normal = Normals->GetElement(ElemID);
-					Normal = NormalTransform(Normal);
-					Normals->SetElement(ElemID, Normalized(Normal));
-				}
-			});
-		}
-
-		FVector3d ZeroTransformed = PositionTransform(FVector::ZeroVector);
-		MeshTransformsLocal::TransformSculptLayers(Mesh,
-			[&PositionTransform](const FVector3d& Pos) {return PositionTransform(Pos); },
-			[&PositionTransform, ZeroTransformed](const FVector3d& Offset) { return PositionTransform(Offset) - ZeroTransformed; });
-	}
+			// we un-set the transform tangent flag, so this method should not be called
+			checkNoEntry();
+			return T;
+		},
+		// transform everything except tangents
+		~ETransformAttributes::Tangents);
 }
 
 void MeshTransforms::ApplyTransform(FDynamicMesh3& Mesh,
 	TFunctionRef<FVector3d(const FVector3d&)> PositionTransform,
 	TFunctionRef<FVector3f(const FVector3f&)> NormalTransform,
-	TFunctionRef<FVector3f(const FVector3f&)> TangentTransform)
+	TFunctionRef<FVector3f(const FVector3f&)> TangentTransform,
+	ETransformAttributes TransformAttributes)
 {
-	bool bVertexNormals = Mesh.HasVertexNormals();
-
-	int NumVertices = Mesh.MaxVertexID();
-	ParallelFor(NumVertices, [&](int vid)
-	{
-		if (Mesh.IsVertex(vid))
-		{
-			FVector3d Position = Mesh.GetVertex(vid);
-			Position = PositionTransform(Position);
-			Mesh.SetVertex(vid, Position);
-
-			if (bVertexNormals)
-			{
-				FVector3f Normal = Mesh.GetVertexNormal(vid);
-				Normal = NormalTransform(Normal);
-				Mesh.SetVertexNormal(vid, Normalized(Normal));
-			}
-		}
-	});
-
-	if (Mesh.HasAttributes())
-	{
-		FDynamicMeshNormalOverlay* Normals = Mesh.Attributes()->PrimaryNormals();
-		if (Normals)
-		{
-			int NumNormals = Normals->MaxElementID();
-			ParallelFor(NumNormals, [&](int ElemID)
-			{
-				if (Normals->IsElement(ElemID))
-				{
-					FVector3f Normal = Normals->GetElement(ElemID);
-					Normal = NormalTransform(Normal);
-					Normals->SetElement(ElemID, Normalized(Normal));
-				}
-			});
-		}
-		if (Mesh.Attributes()->HasTangentSpace())
-		{
-			for (int32 TangentLayerIdx = 1; TangentLayerIdx < 3; ++TangentLayerIdx)
-			{
-				FDynamicMeshNormalOverlay* TangentLayer = Mesh.Attributes()->GetNormalLayer(TangentLayerIdx);
-				int32 NumTangents = TangentLayer->MaxElementID();
-				ParallelFor(NumTangents, [&](int ElemID)
-				{
-					if (TangentLayer->IsElement(ElemID))
-					{
-						FVector3f Tangent = TangentLayer->GetElement(ElemID);
-						Tangent = TangentTransform(Tangent);
-						TangentLayer->SetElement(ElemID, Normalized(Tangent));
-					}
-				});
-			}
-		}
-		
-		FVector3d ZeroTransformed = PositionTransform(FVector::ZeroVector);
-		MeshTransformsLocal::TransformSculptLayers(Mesh,
-			[&PositionTransform](const FVector3d& Pos) {return PositionTransform(Pos); },
-			[&PositionTransform, ZeroTransformed](const FVector3d& Offset) { return PositionTransform(Offset) - ZeroTransformed; });
-	}
+	FVector3d ZeroTransformed = PositionTransform(FVector::ZeroVector);
+	MeshTransformsLocal::TransformMeshHelper(Mesh,
+		PositionTransform,
+		NormalTransform,
+		TangentTransform,
+		[&PositionTransform, ZeroTransformed](const FVector3d& Offset) { return PositionTransform(Offset) - ZeroTransformed; },
+		TransformAttributes);
 }
 
diff --git a/Engine/Source/Runtime/GeometryCore/Public/DynamicMesh/MeshTransforms.h b/Engine/Source/Runtime/GeometryCore/Public/DynamicMesh/MeshTransforms.h
--- a/Engine/Source/Runtime/GeometryCore/Public/DynamicMesh/MeshTransforms.h
+++ b/Engine/Source/Runtime/GeometryCore/Public/DynamicMesh/MeshTransforms.h
@@ -21,48 +21,72 @@ namespace MeshTransforms
 {
 	using namespace UE::Geometry;
 
+	// enum to control which attributes are transformed
+	enum class ETransformAttributes
+	{
+		Positions = 1 << 0,
+		VertexNormals = 1 << 1,
+		Normals = 1 << 2,
+		Tangents = 1 << 3,
+		SculptLayers = 1 << 4,
+		All = Positions | VertexNormals | Normals | Tangents | SculptLayers
+	};
+	ENUM_CLASS_FLAGS(ETransformAttributes);
+
 	/**
-	 * Apply Translation to vertex positions of Mesh. Does not modify any other attributes.
+	 * Apply Translation to the Mesh.
+	 * @param TransformAttributes				Bitflag enum specifying which attributes to transform
 	 */
-	GEOMETRYCORE_API void Translate(FDynamicMesh3& Mesh, const FVector3d& Translation);
+	GEOMETRYCORE_API void Translate(FDynamicMesh3& Mesh, const FVector3d& Translation, ETransformAttributes TransformAttributes = ETransformAttributes::All);
 
 	/**
-	 * Apply Scale to vertex positions of Mesh, relative to given Origin. For non-uniform scale, also scales normals/tangents.
+	 * Apply Scale to Mesh, relative to given Origin. Correctly updates normals/tangents as needed (unless bOnlyPositions is true).
 	 * @param bReverseOrientationIfNeeded		If negative scaling inverts the mesh, also invert the triangle orientations
-	 * @param OnlyPositions						If true, only apply scale positions and do not affect normals/tangents
+	 * @param bOnlyPositions						If true, only apply scale positions and do not affect normals/tangents
 	 */
 	GEOMETRYCORE_API void Scale(FDynamicMesh3& Mesh, const FVector3d& Scale, const FVector3d& Origin, bool bReverseOrientationIfNeeded = false, bool bOnlyPositions = false);
+	
+	/**
+	 * Apply Scale to Mesh, relative to given Origin. Correctly updates normals/tangents as needed (unless TransformAttributes flags are set not to update them).
+	 * @param bReverseOrientationIfNeeded		If negative scaling inverts the mesh, also invert the triangle orientations
+	 * @param TransformAttributes				Bitflag enum specifying which attributes to transform
+	 */
+	GEOMETRYCORE_API void Scale(FDynamicMesh3& Mesh, const FVector3d& Scale, const FVector3d& Origin, bool bReverseOrientationIfNeeded, ETransformAttributes TransformAttributes);
 
 	/**
-	 * Apply Rotation to the vertex positions, relative to the given RotationOrigin.
+	 * Apply Rotation to the Mesh, relative to the given RotationOrigin.
+	 * @param TransformAttributes				Bitflag enum specifying which attributes to transform
 	 */
-	GEOMETRYCORE_API void Rotate(FDynamicMesh3& Mesh, const FRotator& Rotation, const FVector3d& RotationOrigin);
+	GEOMETRYCORE_API void Rotate(FDynamicMesh3& Mesh, const FRotator& Rotation, const FVector3d& RotationOrigin, ETransformAttributes TransformAttributes = ETransformAttributes::All);
 
 	/**
 	 * Transform Mesh into local coordinates of Frame
+	 * @param TransformAttributes				Bitflag enum specifying which attributes to transform
 	 */
-	GEOMETRYCORE_API void WorldToFrameCoords(FDynamicMesh3& Mesh, const FFrame3d& Frame);
+	GEOMETRYCORE_API void WorldToFrameCoords(FDynamicMesh3& Mesh, const FFrame3d& Frame, ETransformAttributes TransformAttributes = ETransformAttributes::All);
 
 	/**
 	 * Transform Mesh out of local coordinates of Frame
+	 * @param TransformAttributes				Bitflag enum specifying which attributes to transform
 	 */
-	GEOMETRYCORE_API void FrameCoordsToWorld(FDynamicMesh3& Mesh, const FFrame3d& Frame);
+	GEOMETRYCORE_API void FrameCoordsToWorld(FDynamicMesh3& Mesh, const FFrame3d& Frame, ETransformAttributes TransformAttributes = ETransformAttributes::All);
 
 
 	/**
 	 * Apply given Transform to a Mesh.
-	 * Modifies Vertex Positions and Normals, and any Per-Triangle Normal Overlays
 	 * @param bReverseOrientationIfNeeded		If the Transform inverts the mesh with negative scale, also invert the triangle orientations
+	 * @param TransformAttributes				Bitflag enum specifying which attributes to transform
 	 */
-	GEOMETRYCORE_API void ApplyTransform(FDynamicMesh3& Mesh, const FTransformSRT3d& Transform, bool bReverseOrientationIfNeeded = false);
+	GEOMETRYCORE_API void ApplyTransform(FDynamicMesh3& Mesh, const FTransformSRT3d& Transform, bool bReverseOrientationIfNeeded = false, ETransformAttributes TransformAttributes = ETransformAttributes::All);
 
 
 	/**
 	 * Apply inverse of given Transform to a Mesh.
 	 * Modifies Vertex Positions and Normals, and any Per-Triangle Normal Overlays
-	  @param bReverseOrientationIfNeeded		If the Transform inverts the mesh with negative scale, also invert the triangle orientations
+	 * @param bReverseOrientationIfNeeded		If the Transform inverts the mesh with negative scale, also invert the triangle orientations
+	 * @param TransformAttributes				Bitflag enum specifying which attributes to transform
 	 */
-	GEOMETRYCORE_API void ApplyTransformInverse(FDynamicMesh3& Mesh, const FTransformSRT3d& Transform, bool bReverseOrientationIfNeeded = false);
+	GEOMETRYCORE_API void ApplyTransformInverse(FDynamicMesh3& Mesh, const FTransformSRT3d& Transform, bool bReverseOrientationIfNeeded = false, ETransformAttributes TransformAttributes = ETransformAttributes::All);
 
 
 	/**
@@ -75,12 +99,13 @@ namespace MeshTransforms
 
 	/**
 	 * Apply given Transform to a Mesh.
-	 * Modifies Vertex Positions, Normals, Tangents, and any Per-Triangle Normal Overlays
+	 * @param TransformAttributes				Bitflag enum specifying which attributes to transform
 	 */
 	GEOMETRYCORE_API void ApplyTransform(FDynamicMesh3& Mesh,
 		TFunctionRef<FVector3d(const FVector3d&)> PositionTransform,
 		TFunctionRef<FVector3f(const FVector3f&)> NormalTransform,
-		TFunctionRef<FVector3f(const FVector3f&)> TangentTransform
+		TFunctionRef<FVector3f(const FVector3f&)> TangentTransform,
+		ETransformAttributes TransformAttributes = ETransformAttributes::All
 	);
 
 	/**
