[GeometryCore]
* Moved ctor of DynamicMesh3 to cpp file because removing dllstorage on the inlined ctors cause linker errors (even though they are inline)

diff --git a/Engine/Source/Runtime/GeometryCore/Private/DynamicMesh/DynamicMesh3.cpp b/Engine/Source/Runtime/GeometryCore/Private/DynamicMesh/DynamicMesh3.cpp
--- a/Engine/Source/Runtime/GeometryCore/Private/DynamicMesh/DynamicMesh3.cpp
+++ b/Engine/Source/Runtime/GeometryCore/Private/DynamicMesh/DynamicMesh3.cpp
@@ -16,6 +16,8 @@ const FVector3d FDynamicMesh3::InvalidVertex = FVector3d(TNumericLimits<double>:
 constexpr FIndex3i  FDynamicMesh3::InvalidTriangle;
 constexpr FIndex2i  FDynamicMesh3::InvalidEdge;
 
+FDynamicMesh3::FDynamicMesh3() : FDynamicMesh3(false, false, false, false) {}
+
 FDynamicMesh3::FDynamicMesh3(bool bWantNormals, bool bWantColors, bool bWantUVs, bool bWantTriGroups)
 {
 	if ( bWantNormals )   { VertexNormals = TDynamicVector<FVector3f>{}; }
@@ -24,6 +26,14 @@ FDynamicMesh3::FDynamicMesh3(bool bWantNormals, bool bWantColors, bool bWantUVs,
 	if ( bWantTriGroups ) { TriangleGroups = TDynamicVector<int>{}; }
 }
 
+FDynamicMesh3::FDynamicMesh3(EMeshComponents flags)
+	: FDynamicMesh3(((int)flags & (int)EMeshComponents::VertexNormals) != 0,
+	                ((int)flags & (int)EMeshComponents::VertexColors) != 0,
+	                ((int)flags & (int)EMeshComponents::VertexUVs) != 0,
+	                ((int)flags & (int)EMeshComponents::FaceGroups) != 0)
+{
+}
+
 // normals/colors/uvs will only be copied if they exist
 FDynamicMesh3::FDynamicMesh3(const FDynamicMesh3& Other)
 	:
diff --git a/Engine/Source/Runtime/GeometryCore/Public/DynamicMesh/DynamicMesh3.h b/Engine/Source/Runtime/GeometryCore/Public/DynamicMesh/DynamicMesh3.h
--- a/Engine/Source/Runtime/GeometryCore/Public/DynamicMesh/DynamicMesh3.h
+++ b/Engine/Source/Runtime/GeometryCore/Public/DynamicMesh/DynamicMesh3.h
@@ -205,7 +205,7 @@ protected:
 
 public:
 	/** Default constructor */
-	GEOMETRYCORE_API FDynamicMesh3() : FDynamicMesh3(false, false, false, false) {}
+	GEOMETRYCORE_API FDynamicMesh3();
 
 	/** Copy/Move construction */
 	GEOMETRYCORE_API FDynamicMesh3(const FDynamicMesh3& CopyMesh);
@@ -221,12 +221,7 @@ public:
 
 	/** Construct an empty mesh with specified attributes */
 	GEOMETRYCORE_API explicit FDynamicMesh3(bool bWantNormals, bool bWantColors, bool bWantUVs, bool bWantTriGroups);
-	GEOMETRYCORE_API explicit FDynamicMesh3(EMeshComponents flags)
-	    : FDynamicMesh3(((int)flags & (int)EMeshComponents::VertexNormals) != 0,
-	                    ((int)flags & (int)EMeshComponents::VertexColors) != 0,
-	                    ((int)flags & (int)EMeshComponents::VertexUVs) != 0,
-	                    ((int)flags & (int)EMeshComponents::FaceGroups) != 0)
-	{}
+	GEOMETRYCORE_API explicit FDynamicMesh3(EMeshComponents flags);
 
 	/** Construction from Mesh Generator */
 	GEOMETRYCORE_API FDynamicMesh3(const FMeshShapeGenerator* Generator);
@@ -1467,7 +1462,7 @@ public:
 	/**
 	 * Weld DiscardVid to KeepVid. Equivalent to calling the options overload with default options.
 	 */
-	GEOMETRYCORE_API virtual EMeshResult MergeVertices(int KeepVid, int DiscardVid, double InterpolationT, FMergeVerticesInfo& MergeInfo)
+	virtual EMeshResult MergeVertices(int KeepVid, int DiscardVid, double InterpolationT, FMergeVerticesInfo& MergeInfo)
 	{
 		return MergeVertices(KeepVid, DiscardVid, InterpolationT, FMergeVerticesOptions(), MergeInfo);
 	}
@@ -1475,7 +1470,7 @@ public:
 	/**
 	 * Weld DiscardVid to KeepVid. Equivalent to calling the options overload with default options and 0 for InterpolationT.
 	 */
-	GEOMETRYCORE_API virtual EMeshResult MergeVertices(int KeepVid, int DiscardVid, FMergeVerticesInfo& MergeInfo)
+	virtual EMeshResult MergeVertices(int KeepVid, int DiscardVid, FMergeVerticesInfo& MergeInfo)
 	{
 		return MergeVertices(KeepVid, DiscardVid, 0, FMergeVerticesOptions(), MergeInfo);
 	}
