add dynamic mesh overlay / attribute set methods to test if a vertex is a seam intersection

diff --git a/Engine/Source/Runtime/GeometryCore/Private/DynamicMesh/DynamicMeshAttributeSet.cpp b/Engine/Source/Runtime/GeometryCore/Private/DynamicMesh/DynamicMeshAttributeSet.cpp
--- a/Engine/Source/Runtime/GeometryCore/Private/DynamicMesh/DynamicMeshAttributeSet.cpp
+++ b/Engine/Source/Runtime/GeometryCore/Private/DynamicMesh/DynamicMeshAttributeSet.cpp
@@ -918,6 +918,29 @@ bool FDynamicMeshAttributeSet::IsSeamVertex(int VID, bool bBoundaryIsSeam) const
 	return false;
 }
 
+bool FDynamicMeshAttributeSet::IsSeamIntersectionVertex(int32 VertexID) const
+{
+	for (const FDynamicMeshUVOverlay& UVLayer : UVLayers)
+	{
+		if (UVLayer.IsSeamIntersectionVertex(VertexID))
+		{
+			return true;
+		}
+	}
+	for (const FDynamicMeshNormalOverlay& NormalLayer : NormalLayers)
+	{
+		if (NormalLayer.IsSeamIntersectionVertex(VertexID))
+		{
+			return true;
+		}
+	}
+	if (ColorLayer && ColorLayer->IsSeamIntersectionVertex(VertexID))
+	{
+		return true;
+	}
+	return false;
+}
+
 bool FDynamicMeshAttributeSet::IsMaterialBoundaryEdge(int EdgeID) const
 {
 	if ( MaterialIDAttrib == nullptr )
diff --git a/Engine/Source/Runtime/GeometryCore/Private/DynamicMesh/DynamicMeshOverlay.cpp b/Engine/Source/Runtime/GeometryCore/Private/DynamicMesh/DynamicMeshOverlay.cpp
--- a/Engine/Source/Runtime/GeometryCore/Private/DynamicMesh/DynamicMeshOverlay.cpp
+++ b/Engine/Source/Runtime/GeometryCore/Private/DynamicMesh/DynamicMeshOverlay.cpp
@@ -785,6 +785,17 @@ bool TDynamicMeshOverlay<RealType, ElementSize>::IsSeamVertex(int vid, bool bBou
 	return false;
 }
 
+template<typename RealType, int ElementSize>
+bool TDynamicMeshOverlay<RealType, ElementSize>::IsSeamIntersectionVertex(int32 VertexID) const
+{
+	int32 SeamCount = 0;
+	for (int32 EdgeID : ParentMesh->VtxEdgesItr(VertexID))
+	{
+		SeamCount += int32(IsSeamEdge(EdgeID));
+	}
+	return SeamCount == 1 || SeamCount > 2;
+}
+
 template<typename RealType, int ElementSize>
 bool TDynamicMeshOverlay<RealType, ElementSize>::IsBowtieInOverlay(int32 VertexID) const
 {
diff --git a/Engine/Source/Runtime/GeometryCore/Public/DynamicMesh/DynamicMeshAttributeSet.h b/Engine/Source/Runtime/GeometryCore/Public/DynamicMesh/DynamicMeshAttributeSet.h
--- a/Engine/Source/Runtime/GeometryCore/Public/DynamicMesh/DynamicMeshAttributeSet.h
+++ b/Engine/Source/Runtime/GeometryCore/Public/DynamicMesh/DynamicMeshAttributeSet.h
@@ -152,6 +152,9 @@ public:
 	/** @return true if the given vertex is a seam vertex in any overlay */
 	GEOMETRYCORE_API virtual bool IsSeamVertex(int VertexID, bool bBoundaryIsSeam = true) const;
 
+	/** @return true if the given vertex is a seam intersection vertex in any overlay */
+	GEOMETRYCORE_API virtual bool IsSeamIntersectionVertex(int32 VertexID) const;
+
 	/** @return true if the given edge is a material ID boundary */
 	GEOMETRYCORE_API virtual bool IsMaterialBoundaryEdge(int EdgeID) const;
 
diff --git a/Engine/Source/Runtime/GeometryCore/Public/DynamicMesh/DynamicMeshOverlay.h b/Engine/Source/Runtime/GeometryCore/Public/DynamicMesh/DynamicMeshOverlay.h
--- a/Engine/Source/Runtime/GeometryCore/Public/DynamicMesh/DynamicMeshOverlay.h
+++ b/Engine/Source/Runtime/GeometryCore/Public/DynamicMesh/DynamicMeshOverlay.h
@@ -571,6 +571,8 @@ public:
 	bool IsSeamEndEdge(int EdgeID) const;
 	/** Returns true if the parent-mesh vertex is connected to any seam edges */
 	bool IsSeamVertex(int VertexID, bool bBoundaryIsSeam = true) const;
+	/** Returns true if the parent-mesh vertex is at a seam 'intersection' -- i.e., the end of a seam, or the intersection w/ another seam. */
+	bool IsSeamIntersectionVertex(int32 VertexID) const;
 
 	/**
 	 * Determines whether the base-mesh vertex has "bowtie" topology in the Overlay.
