Fix SetValue in BSpline to return bool as previously

diff --git a/Engine/Source/Runtime/GeometryCore/Public/Splines/BSpline.h b/Engine/Source/Runtime/GeometryCore/Public/Splines/BSpline.h
--- a/Engine/Source/Runtime/GeometryCore/Public/Splines/BSpline.h
+++ b/Engine/Source/Runtime/GeometryCore/Public/Splines/BSpline.h
@@ -131,8 +131,16 @@ public:
 
     int32 NumKeys() const { return Values.Num(); }
     const ValueType& GetValue(int32 Idx) const { return Values[Idx]; }
-    void SetValue(int32 Idx, const ValueType& NewValue) { Values[Idx] = NewValue; }
     int32 AddValue(const ValueType& NewValue)  { return Values.Add(NewValue); }
+	bool SetValue(int32 Idx, const ValueType& NewValue)
+    {
+        if (!Values.IsValidIndex(Idx))
+    	{
+    		return false;
+    	}
+        Values[Idx] = NewValue;
+        return true;
+    }
 	int32 InsertValue(int32 Idx, const ValueType& NewValue)
     {
     	if (Idx < 0 || Idx > Values.Num())
diff --git a/Engine/Source/Runtime/GeometryCore/Public/Splines/PolyBezierSpline.h b/Engine/Source/Runtime/GeometryCore/Public/Splines/PolyBezierSpline.h
--- a/Engine/Source/Runtime/GeometryCore/Public/Splines/PolyBezierSpline.h
+++ b/Engine/Source/Runtime/GeometryCore/Public/Splines/PolyBezierSpline.h
@@ -1017,8 +1017,7 @@ public:
 
         // Convert segment+point index to global index in spline values
         const int32 GlobalIndex = SegmentIndex * 4 + PointIndex;
-        Base::SetValue(GlobalIndex, NewValue);
-		return true;
+        return Base::SetValue(GlobalIndex, NewValue);
     }
 
     /**
