BSpline: Fix operator== using || instead of &&

diff --git a/Engine/Source/Runtime/GeometryCore/Public/Splines/BSpline.h b/Engine/Source/Runtime/GeometryCore/Public/Splines/BSpline.h
--- a/Engine/Source/Runtime/GeometryCore/Public/Splines/BSpline.h
+++ b/Engine/Source/Runtime/GeometryCore/Public/Splines/BSpline.h
@@ -79,9 +79,9 @@ public:
     }
 	bool operator==(const TBSpline& Other) const
     {
-    	return	Values == Other.Values ||
-    			Knots == Other.Knots ||
-    			bIsClosedLoop == Other.bIsClosedLoop ||
+    	return	Values == Other.Values &&
+    			Knots == Other.Knots &&
+    			bIsClosedLoop == Other.bIsClosedLoop &&
     			bClampEnds == Other.bClampEnds;
     }
 	friend FArchive& operator<<(FArchive& Ar, TBSpline& BSpline)
