add TMeshAABBTree3::GetByteCount

#rnx

diff --git a/Engine/Source/Runtime/GeometryCore/Public/Spatial/MeshAABBTree3.h b/Engine/Source/Runtime/GeometryCore/Public/Spatial/MeshAABBTree3.h
--- a/Engine/Source/Runtime/GeometryCore/Public/Spatial/MeshAABBTree3.h
+++ b/Engine/Source/Runtime/GeometryCore/Public/Spatial/MeshAABBTree3.h
@@ -1196,6 +1196,10 @@ public:
 		return Box.Intersects(TestBox);
 	}
 
+	SIZE_T GetByteCount() const
+	{
+		return BoxToIndex.GetByteCount() + BoxCenters.GetByteCount() + BoxExtents.GetByteCount() + IndexList.GetByteCount();
+	}
 
 	// storage for Box Nodes.
 	//   - BoxToIndex is a pointer into IndexList
