Splines - Multispline channel creation interface improvements

diff --git a/Engine/Source/Runtime/Engine/Private/Components/SplineComponent.cpp b/Engine/Source/Runtime/Engine/Private/Components/SplineComponent.cpp
--- a/Engine/Source/Runtime/Engine/Private/Components/SplineComponent.cpp
+++ b/Engine/Source/Runtime/Engine/Private/Components/SplineComponent.cpp
@@ -690,8 +690,7 @@ bool CreatePropertyChannel(const USplineMetadata* Metadata, const FSpline& Splin
 	{
 		if (!Spline.HasAttributeChannel(PropertyName))
 		{
-			Spline.CreateAttributeChannel<T>(PropertyName);
-			return true;
+			return Spline.CreateAttributeChannel<T>(PropertyName);
 		}
 	}
 
diff --git a/Engine/Source/Runtime/Engine/Private/Curves/Spline.cpp b/Engine/Source/Runtime/Engine/Private/Curves/Spline.cpp
--- a/Engine/Source/Runtime/Engine/Private/Curves/Spline.cpp
+++ b/Engine/Source/Runtime/Engine/Private/Curves/Spline.cpp
@@ -305,7 +305,7 @@ void FNewSpline::SetAttributeValue(FName Name, const AttrType& Value, int32 Inde
 {
 	using AttrSplineType = TTangentBezierSpline<AttrType>;
 
-	if (AttrSplineType* Channel = GetAttributeChannelImpl<AttrSplineType>(Name))
+	if (AttrSplineType* Channel = GetTypedAttributeChannel<AttrSplineType>(Name))
 	{
 		Channel->SetValue(Index, Value);
 	}
@@ -316,7 +316,7 @@ void FNewSpline::RemoveAttributeValue(FName Name, int32 Index)
 {
 	using AttrSplineType = TTangentBezierSpline<AttrType>;
 
-	if (AttrSplineType* Channel = GetAttributeChannelImpl<AttrSplineType>(Name))
+	if (AttrSplineType* Channel = GetTypedAttributeChannel<AttrSplineType>(Name))
 	{
 		Channel->RemovePoint(Index);
 	}
@@ -327,7 +327,7 @@ AttrType FNewSpline::GetAttributeValue(const FName& Name, int32 Index) const
 {
 	using AttrSplineType = TTangentBezierSpline<AttrType>;
 	
-	if (AttrSplineType* Channel = GetAttributeChannelImpl<AttrSplineType>(Name))
+	if (AttrSplineType* Channel = GetTypedAttributeChannel<AttrSplineType>(Name))
 	{
 		return Channel->GetValue(Index);
 	}
@@ -340,7 +340,7 @@ int32 FNewSpline::NumAttributeValues(FName Name) const
 {
 	using AttrSplineType = TTangentBezierSpline<AttrType>;
 	
-	if (AttrSplineType* Channel = GetAttributeChannelImpl<AttrSplineType>(Name))
+	if (AttrSplineType* Channel = GetTypedAttributeChannel<AttrSplineType>(Name))
 	{
 		return Channel->GetNumPoints();
 	}
@@ -900,7 +900,7 @@ bool FSpline::CreateAttributeChannel(FName AttributeName) const
 {
 	if (NewData)
 	{
-	//	return NewData->CreateAttributeChannel<AttrType>(AttributeName);
+		return NewData->CreateAttributeChannel<TTangentBezierSpline<AttrType>>(AttributeName) != nullptr;
 	}
 	
 	return false;
@@ -924,10 +924,10 @@ template ENGINE_API int32 FSpline::AddAttributeValue<FVector>(float Param, const
 template <typename AttrType>
 void FSpline::RemoveAttributeValue(int32 Index, FName AttributeName)
 {
-	// if (NewData)
-	// {
-	// 	NewData->RemoveAttributeValue(AttributeName, Index);
-	// }
+	if (NewData)
+	{
+		NewData->RemoveAttributeValue<AttrType>(AttributeName, Index);
+	}
 }
 template ENGINE_API void FSpline::RemoveAttributeValue<float>(int32 Index, FName AttributeName);
 template ENGINE_API void FSpline::RemoveAttributeValue<FVector>(int32 Index, FName AttributeName);
@@ -1296,10 +1296,10 @@ void FSpline::OnSplineImplementationChanged()
 
 FNewSpline::FNewSpline()
 {
-	CreateAttributeChannel<FQuat>(RotationAttrName, TTangentBezierSpline<FQuat>::GetSplineTypeName());
-	CreateAttributeChannel<FVector>(ScaleAttrName, TTangentBezierSpline<FVector>::GetSplineTypeName());
-	CreateAttributeChannel<float>(ParamToLengthAttrName, TTangentBezierSpline<float>::GetSplineTypeName());
-	CreateAttributeChannel<float>(LengthToParamAttrName, TTangentBezierSpline<float>::GetSplineTypeName());
+	CreateAttributeChannel<TTangentBezierSpline<FQuat>>(RotationAttrName);
+	CreateAttributeChannel<TTangentBezierSpline<FVector>>(ScaleAttrName);
+	CreateAttributeChannel<TTangentBezierSpline<float>>(ParamToLengthAttrName);
+	CreateAttributeChannel<TTangentBezierSpline<float>>(LengthToParamAttrName);
     
     ResetToDefault();
 }
@@ -1311,10 +1311,10 @@ FNewSpline::FNewSpline(const FNewSpline& Other)
 
 FNewSpline::FNewSpline(const FLegacySpline& Other)
 {
-	CreateAttributeChannel<FQuat>(RotationAttrName, TTangentBezierSpline<FQuat>::GetSplineTypeName());
-	CreateAttributeChannel<FVector>(ScaleAttrName, TTangentBezierSpline<FVector>::GetSplineTypeName());
-	CreateAttributeChannel<float>(ParamToLengthAttrName, TTangentBezierSpline<float>::GetSplineTypeName());
-	CreateAttributeChannel<float>(LengthToParamAttrName, TTangentBezierSpline<float>::GetSplineTypeName());
+	CreateAttributeChannel<TTangentBezierSpline<FQuat>>(RotationAttrName);
+	CreateAttributeChannel<TTangentBezierSpline<FVector>>(ScaleAttrName);
+	CreateAttributeChannel<TTangentBezierSpline<float>>(ParamToLengthAttrName);
+	CreateAttributeChannel<TTangentBezierSpline<float>>(LengthToParamAttrName);
 	
 	TArray<FTangentBezierControlPoint> Points;
     for (int32 Idx = 0; Idx < Other.GetNumControlPoints(); Idx++)
@@ -1345,10 +1345,10 @@ FNewSpline::FNewSpline(const FLegacySpline& Other)
 
 FNewSpline::FNewSpline(const FSplineCurves& Other)
 {
-	CreateAttributeChannel<FQuat>(RotationAttrName, TTangentBezierSpline<FQuat>::GetSplineTypeName());
-	CreateAttributeChannel<FVector>(ScaleAttrName, TTangentBezierSpline<FVector>::GetSplineTypeName());
-	CreateAttributeChannel<float>(ParamToLengthAttrName, TTangentBezierSpline<float>::GetSplineTypeName());
-	CreateAttributeChannel<float>(LengthToParamAttrName, TTangentBezierSpline<float>::GetSplineTypeName());
+	CreateAttributeChannel<TTangentBezierSpline<FQuat>>(RotationAttrName);
+	CreateAttributeChannel<TTangentBezierSpline<FVector>>(ScaleAttrName);
+	CreateAttributeChannel<TTangentBezierSpline<float>>(ParamToLengthAttrName);
+	CreateAttributeChannel<TTangentBezierSpline<float>>(LengthToParamAttrName);
 	
 	TArray<FTangentBezierControlPoint> Points;
     for (int32 Idx = 0; Idx < Other.Position.Points.Num(); Idx++)
diff --git a/Engine/Source/Runtime/GeometryCore/Public/Splines/CustomAttributeExample.h b/Engine/Source/Runtime/GeometryCore/Public/Splines/CustomAttributeExample.h
--- a/Engine/Source/Runtime/GeometryCore/Public/Splines/CustomAttributeExample.h
+++ b/Engine/Source/Runtime/GeometryCore/Public/Splines/CustomAttributeExample.h
@@ -189,11 +189,8 @@ void CreatePolyBezierWithFooAttributes()
     TMultiSpline<FPolyBezierSpline3d> Spline;
     Spline.GetSpline() = BaseSpline; // Set the main spline
     
-    // Create attribute channel using type registry by name
-    Spline.CreateAttributeChannel<FFoo>("FooValues", TEXT("BSpline3"));
-
-	// Get direct access to the BSpline implementation
-    TBSpline<FFoo, 3>* FooSpline = Spline.GetAttributeChannelImpl<TBSpline<FFoo, 3>>("FooValues");
+    // Get direct access to the BSpline implementation
+    TBSpline<FFoo, 3>* FooSpline = Spline.CreateAttributeChannel<TBSpline<FFoo, 3>>("FooValues");
     if (FooSpline)
     {
         // Now we can directly add values using the BSpline API
diff --git a/Engine/Source/Runtime/GeometryCore/Public/Splines/MultiSpline.h b/Engine/Source/Runtime/GeometryCore/Public/Splines/MultiSpline.h
--- a/Engine/Source/Runtime/GeometryCore/Public/Splines/MultiSpline.h
+++ b/Engine/Source/Runtime/GeometryCore/Public/Splines/MultiSpline.h
@@ -116,21 +116,34 @@
     SplineType& GetSpline() { return TSplineWrapper<SPLINETYPE>::InternalSpline; }
     const SplineType& GetSpline() const { return TSplineWrapper<SPLINETYPE>::InternalSpline; }
     
-    /**
-     * Gets an attribute channel of the specified type and tries to cast it to the expected implementation type
-     * @param Name Name of the attribute channel
-     * @return Pointer to the channel implementation, or nullptr if not found or can't be cast
-     */
+	/**
+	 * Gets a strongly-typed attribute channel with the specific implementation type
+	 * 
+	 * @tparam ImplType - The specific spline implementation type to get
+	 * @param Name - Name of the attribute channel
+	 * @return Pointer to the channel implementation, or nullptr if not found or can't be cast
+	 */
     template <typename ImplType>
-    ImplType* GetAttributeChannelImpl(FName Name) const
+    ImplType* GetTypedAttributeChannel(FName Name) const
     {
         if (const FChildSpline* ExistingChannel = Children.Find(Name))
         {
             if (ExistingChannel->Spline.IsValid())
             {
-                // Try to cast to the implementation type (this is safe as it's just a pointer cast)
-                // The caller is responsible for ensuring the correct ImplType
-                return static_cast<ImplType*>(ExistingChannel->Spline.Get());
+            	// Get implementation names for comparison
+            	const FName RequestedImplName = ImplType::GetSplineTypeName();
+            	const FName ExistingImplName = ExistingChannel->Spline->GetImplementationName();
+            
+            	// Check implementation type compatibility
+            	if (RequestedImplName == ExistingImplName)
+            	{
+            		// Safe to cast if implementation types match
+            		return static_cast<ImplType*>(ExistingChannel->Spline.Get());
+            	}
+
+            	// Log warning about implementation type mismatch
+            	UE_LOG(LogSpline, Warning, TEXT("Failed to fetch attribute channel '%s': Channel exists with implementation type '%s', but requested type is '%s'."),
+					   *Name.ToString(), *ExistingImplName.ToString(), *RequestedImplName.ToString());
             }
         }
         return nullptr;
@@ -165,6 +178,117 @@
 		return nullptr;
     }
 
+	/**
+	 * Gets all attribute channel names present in this spline
+	 * @return Array of channel names
+	 */
+	TArray<FName> GetAllAttributeChannelNames() const
+	{
+	    TArray<FName> ChannelNames;
+	    Children.GetKeys(ChannelNames);
+	    return ChannelNames;
+	}
+
+	/**
+	 * Gets names of attribute channels with a specific value type
+	 * 
+	 * @tparam AttrType - The value type to filter by (e.g., float, FVector)
+	 * @return Array of channel names with the specified value type
+	 */
+	template <typename AttrType>
+	TArray<FName> GetAttributeChannelNamesOfType() const
+	{
+	    TArray<FName> ChannelNames;
+	    const FName TypeName = TSplineValueTypeTraits<AttrType>::Name;
+	    
+	    for (auto& Pair : Children)
+	    {
+	        if (Pair.Value.Spline.IsValid() && 
+	            Pair.Value.Spline->GetValueTypeName() == TypeName)
+	        {
+	            ChannelNames.Add(Pair.Key);
+	        }
+	    }
+	    
+	    return ChannelNames;
+	}
+
+	/**
+	 * Removes an attribute channel completely
+	 * 
+	 * @param Name - Name of the channel to remove
+	 * @return True if the channel was found and removed
+	 */
+	bool RemoveAttributeChannel(const FName& Name)
+	{
+	    return Children.Remove(Name) > 0;
+	}
+
+	/**
+	 * Clones an attribute channel to a new name
+	 * 
+	 * @tparam AttrType - The value type of the attribute channel
+	 * @param SourceName - Name of the source channel to clone
+	 * @param DestName - Name for the new cloned channel
+	 * @return True if the channel was successfully cloned
+	 */
+	template <typename AttrType>
+	bool CloneAttributeChannel(const FName& SourceName, const FName& DestName)
+	{
+	    // Check if destination already exists
+	    if (Children.Contains(DestName))
+	    {
+	        UE_LOG(LogSpline, Warning, TEXT("Cannot clone channel: Destination channel '%s' already exists."), 
+	               *DestName.ToString());
+	        return false;
+	    }
+	    
+	    // Get source channel
+	    TSplineInterface<AttrType>* SourceChannel = GetAttributeChannel<AttrType>(SourceName);
+	    if (!SourceChannel)
+	    {
+	        UE_LOG(LogSpline, Warning, TEXT("Cannot clone channel: Source channel '%s' not found or type mismatch."), 
+	               *SourceName.ToString());
+	        return false;
+	    }
+	    
+	    // Get implementation information from source
+	    const FName ImplName = SourceChannel->GetImplementationName();
+	    
+	    // Create new destination channel with same implementation
+	    if (!CreateAttributeChannelInternal<AttrType>(DestName, ImplName))
+	    {
+	        UE_LOG(LogSpline, Error, TEXT("Failed to create destination channel '%s'."), 
+	               *DestName.ToString());
+	        return false;
+	    }
+	    
+	    // Clone the source channel's data to the destination channel
+	    TSplineInterface<AttrType>* DestChannel = GetAttributeChannel<AttrType>(DestName);
+	    if (!DestChannel)
+	    {
+	        // This should not happen, but just in case
+	        UE_LOG(LogSpline, Error, TEXT("Internal error: Unable to retrieve just-created channel '%s'."), 
+	               *DestName.ToString());
+	        return false;
+	    }
+	    
+	    // Clone the source to the destination
+	    // We need to use the ISplineInterface::Clone() approach and move the clone into our structure
+	    TUniquePtr<ISplineInterface> ClonedSpline = SourceChannel->Clone();
+	    
+	    // Find and update the destination child's spline pointer
+	    FChildSpline& DestChild = Children[DestName];
+	    DestChild.Spline = MoveTemp(ClonedSpline);
+	    
+	    // Copy mapping settings from source
+	    const FChildSpline& SourceChild = *Children.Find(SourceName);
+	    DestChild.MappingRange = SourceChild.MappingRange;
+	    DestChild.MappingRangeSpace = SourceChild.MappingRangeSpace;
+	    
+	    return true;
+	}
+
     /**
      * Maps a parameter from parent space to child space
      * @param Name Name of the child spline
@@ -197,55 +321,29 @@
         return 0.f;
     }
 
-    /**
-     * Creates an attribute channel with the specified type and implementation name
-     * @param Name Name of the channel to create
-     * @param ImplName Name of the implementation to use, defaults to "BSpline3"
-     * @return True if channel was created successfully
+    
+
+	/**
+     * Creates an attribute channel and returns the typed implementation
+     * 
+     * @tparam ImplType - The specific spline implementation type to create
+     * @param Name - Name of the channel to create
+     * @return Pointer to the created attribute channel, or nullptr if creation failed
      */
-    template <typename AttrType>
-    bool CreateAttributeChannel(const FName& Name, const FName& ImplName)
+    template <typename ImplType>
+    ImplType* CreateAttributeChannel(const FName& Name)
     {
-        // Check if channel already exists
-        if (Children.Contains(Name))
+        // Extract the value type from the implementation type
+        using ValueType = typename ImplType::ValueType;
+        
+        // Create the channel using the existing name-based method
+        if (!CreateAttributeChannelInternal<ValueType>(Name, ImplType::GetSplineTypeName()))
         {
-            UE_LOG(LogSpline, Warning, TEXT("Failed to create attribute channel '%s': Channel already exists."), *Name.ToString());
-            return false;
+            return nullptr;
         }
-
-        // Look up the type ID
-        uint32 TypeId = FSplineTypeRegistry::GetTypeId(ImplName, TSplineValueTypeTraits<AttrType>::Name);
-        if (TypeId == 0)
-        {
-            UE_LOG(LogSpline, Error, TEXT("Failed to find type ID for implementation '%s' with value type '%s'."),
-                   *ImplName.ToString(), *TSplineValueTypeTraits<AttrType>::Name.ToString());
-            return false;
-        }
-
-        // Create child spline by type ID
-        TUniquePtr<ISplineInterface> NewSpline = FSplineTypeRegistry::CreateSpline(TypeId);
-        if (!NewSpline)
-        {
-                UE_LOG(LogSpline, Error, TEXT("Failed to create spline of type '%s'/'%s' (ID: 0x%08X) for attribute channel '%s'."),
-               *ImplName.ToString(), *TSplineValueTypeTraits<AttrType>::Name.ToString(), TypeId, *Name.ToString());
-            return false;
-        }
-
-        // Verify that the created spline is of the expected value type
-        if (NewSpline->GetValueTypeName() != TSplineValueTypeTraits<AttrType>::Name)
-        {
-            UE_LOG(LogSpline, Error, TEXT("Created spline has incorrect value type. Expected '%s', got '%s'."),
-                   *TSplineValueTypeTraits<AttrType>::Name.ToString(), *NewSpline->GetValueTypeName().ToString());
-            return false;
-        }
-
-        // Add to children map
-        FChildSpline& Child = Children.Add(Name);
-        Child.Spline = MoveTemp(NewSpline);
-        Child.MappingRangeSpace = EMappingRangeSpace::Normalized;
-        Child.MappingRange = FInterval1f(0.0f, 1.0f);
-
-        return true;
+    
+        // Return the typed implementation
+        return GetTypedAttributeChannel<ImplType>(Name);
     }
 
     /**
@@ -320,6 +418,58 @@
     }
     
 private:
+
+	/**
+     * Helper method that creates an attribute channel with the specified type and implementation name
+     * @param Name Name of the channel to create
+     * @param ImplName Name of the implementation to use, defaults to "BSpline3"
+     * @return True if channel was created successfully
+     */
+    template <typename AttrType>
+    bool CreateAttributeChannelInternal(const FName& Name, const FName& ImplName)
+    {
+        // Check if channel already exists
+        if (Children.Contains(Name))
+        {
+            UE_LOG(LogSpline, Warning, TEXT("Failed to create attribute channel '%s': Channel already exists."), *Name.ToString());
+            return false;
+        }
+
+        // Look up the type ID
+        uint32 TypeId = FSplineTypeRegistry::GetTypeId(ImplName, TSplineValueTypeTraits<AttrType>::Name);
+        if (TypeId == 0)
+        {
+            UE_LOG(LogSpline, Error, TEXT("Failed to find type ID for implementation '%s' with value type '%s'."),
+                   *ImplName.ToString(), *TSplineValueTypeTraits<AttrType>::Name.ToString());
+            return false;
+        }
+
+        // Create child spline by type ID
+        TUniquePtr<ISplineInterface> NewSpline = FSplineTypeRegistry::CreateSpline(TypeId);
+        if (!NewSpline)
+        {
+                UE_LOG(LogSpline, Error, TEXT("Failed to create spline of type '%s'/'%s' (ID: 0x%08X) for attribute channel '%s'."),
+               *ImplName.ToString(), *TSplineValueTypeTraits<AttrType>::Name.ToString(), TypeId, *Name.ToString());
+            return false;
+        }
+
+        // Verify that the created spline is of the expected value type
+        if (NewSpline->GetValueTypeName() != TSplineValueTypeTraits<AttrType>::Name)
+        {
+            UE_LOG(LogSpline, Error, TEXT("Created spline has incorrect value type. Expected '%s', got '%s'."),
+                   *TSplineValueTypeTraits<AttrType>::Name.ToString(), *NewSpline->GetValueTypeName().ToString());
+            return false;
+        }
+
+        // Add to children map
+        FChildSpline& Child = Children.Add(Name);
+        Child.Spline = MoveTemp(NewSpline);
+        Child.MappingRangeSpace = EMappingRangeSpace::Normalized;
+        Child.MappingRange = FInterval1f(0.0f, 1.0f);
+
+        return true;
+    }
+	
     TMap<FName, FChildSpline> Children;
 };
 
