DynamicMesh3 unit test weighted SplitEdge on seams and FlipEdge edge ordering

diff --git a/Engine/Source/Runtime/GeometryCore/Public/DynamicMesh/DynamicMesh3.h b/Engine/Source/Runtime/GeometryCore/Public/DynamicMesh/DynamicMesh3.h
--- a/Engine/Source/Runtime/GeometryCore/Public/DynamicMesh/DynamicMesh3.h
+++ b/Engine/Source/Runtime/GeometryCore/Public/DynamicMesh/DynamicMesh3.h
@@ -1298,6 +1298,10 @@ public:
 	/**
 	 * Flip/Rotate an edge of the mesh. This does not change the number of edges, vertices, or triangles.
 	 * Boundary edges of the mesh cannot be flipped.
+	 * 
+	 * On success, the edges of the new triangles are ordered such that the shared edge is on index 0. i.e.,
+	 * GetTriEdge(FlipInfo.Triangles[0], 0) == GetTriEdge(FlipInfo.Triangles[1], 0)
+	 * 
 	 * @param EdgeAB index of edge to be flipped
 	 * @param FlipInfo returned information about new and modified mesh elements
 	 * @return Ok on success, or enum value indicates why operation cannot be applied. Mesh remains unmodified on error.
