Fixes for inserting multiple points at the same location with degenerate zero length segments.

diff --git a/Engine/Source/Runtime/GeometryCore/Public/Splines/BSpline.h b/Engine/Source/Runtime/GeometryCore/Public/Splines/BSpline.h
--- a/Engine/Source/Runtime/GeometryCore/Public/Splines/BSpline.h
+++ b/Engine/Source/Runtime/GeometryCore/Public/Splines/BSpline.h
@@ -106,7 +106,6 @@ public:
 		Clone->Knots = this->Knots;
 		Clone->bIsClosedLoop = this->bIsClosedLoop;
 		Clone->bClampEnds = this->bClampEnds;
-		Clone->bEnableKnotUpdates = this->bEnableKnotUpdates;
     
 		// Copy infinity modes
 		Clone->PreInfinityMode = this->PreInfinityMode;
@@ -885,10 +884,6 @@ protected:
 				   ? (NumValues + (2 * Degree) + 1) // closed loop: n + 2p + 1
 				   : (NumValues + Degree + 1);
 	}
-	void SetEnableKnotUpdates(bool bEnable)
-    {
-        bEnableKnotUpdates = bEnable;
-    }
 private:
 	/** Find knot span using binary search */
 	int32 FindKnotSpan(float Parameter) const
@@ -1002,8 +997,6 @@ protected:
     TArray<ValueType> Values;
     TArray<float> Knots;
 	bool bIsClosedLoop = false;
-	// flag for enabling/disabling knot updates
-	bool bEnableKnotUpdates = true;
 	
 	/** Whether to clamp the endpoints for interpolation of first/last control points */
 	bool bClampEnds = true;
diff --git a/Engine/Source/Runtime/GeometryCore/Public/Splines/PolyBezierSpline.h b/Engine/Source/Runtime/GeometryCore/Public/Splines/PolyBezierSpline.h
--- a/Engine/Source/Runtime/GeometryCore/Public/Splines/PolyBezierSpline.h
+++ b/Engine/Source/Runtime/GeometryCore/Public/Splines/PolyBezierSpline.h
@@ -264,7 +264,6 @@ public:
         Clone->Knots = this->Knots;
         Clone->bIsClosedLoop = this->bIsClosedLoop;
         Clone->bClampEnds = this->bClampEnds;
-        Clone->bEnableKnotUpdates = this->bEnableKnotUpdates;
         
         // Copy infinity modes
         Clone->PreInfinityMode = this->PreInfinityMode;
@@ -726,11 +725,6 @@ public:
 	    }
 	    
 	    SegmentIndex = FMath::Clamp(SegmentIndex, 0, NumSegments - 1);
-
-		// we clamp the LocalT to not be exactly 0 or 1, as this would cause
-		// the new point to be exactly the same as the start or end of the segment
-		// and creating possible degenerate segments with zero lengths.
-	    LocalT = FMath::Clamp(LocalT, UE_KINDA_SMALL_NUMBER, 1.0f - UE_KINDA_SMALL_NUMBER);
 	    
 	    // Get the segment's parameter range
 	    FInterval1f SegmentRange = GetSegmentParameterRange(SegmentIndex);
@@ -768,6 +762,7 @@ public:
 		int32 SegmentIndex;
 		float LocalT;
 		MapGlobalParameterToLocalSegment(Parameter, SegmentIndex, LocalT);
+
 		const int32 BaseIndex = SegmentIndex * 4;
 		ValueType P0 = Base::GetValue(BaseIndex);
 		
@@ -788,6 +783,8 @@ public:
         ValueType P3 = Base::GetValue(BaseIndex + 3);
  
         // Use De Casteljau algorithm to split the curve
+		float MinDelta = UE_KINDA_SMALL_NUMBER;
+		LocalT = FMathf::Clamp(LocalT, MinDelta, 1.f - MinDelta);
         const float t = LocalT;
         const float mt = 1.0f - t;
         
@@ -805,14 +802,9 @@ public:
         FInterval1f SegmentRange = GetSegmentParameterRange(SegmentIndex);
 		float OriginalStartParam = SegmentRange.Min;
 		float OriginalEndParam = SegmentRange.Max;
-		float TargetSplitParam = SegmentRange.Interpolate(LocalT);//OriginalStartParam + LocalT * (OriginalEndParam - OriginalStartParam);
-    
-		// Instead of removing the original segment, let's directly insert a knot with multiplicity 3
-		// at the split point, which is the proper structure for an internal Bezier junction
-		//this->InsertKnot(TargetSplitParam, 3);
-		Base::SetEnableKnotUpdates(false);
-		// Now we can safely remove the control points
-		// We're not touching the knot structure anymore - just the control points
+		float TargetSplitParam = SegmentRange.Interpolate(LocalT);
+		TargetSplitParam = FMath::Clamp(TargetSplitParam, OriginalStartParam + MinDelta, OriginalEndParam - MinDelta);
+		
 		for (int32 i = 0; i < 4; ++i)
 		{
 			Base::RemoveValue(BaseIndex);
@@ -823,7 +815,6 @@ public:
         AddBezierSegmentInternal(SegmentIndex, Q0, Q1, Q2, Q3, ParameterizationPolicy, OriginalStartParam, TargetSplitParam);
 		int32 IndexToReturn = AddBezierSegmentInternal(SegmentIndex + 1, R0, R1, R2, R3, ParameterizationPolicy, TargetSplitParam, OriginalEndParam);
 
-		Base::SetEnableKnotUpdates(true);
 		return IndexToReturn;
     }
 
@@ -901,11 +892,8 @@ public:
     	// Calculate storage indices
     	const int32 BaseIndex = SegmentIndex * 4;
     	const bool bIsFirstSegment = (SegmentIndex == 0);
-    	const bool bIsLastSegment = /*!this->IsClosedLoop() && */(SegmentIndex == NumSegments - 1);
+    	const bool bIsLastSegment = (SegmentIndex == NumSegments - 1);
 
-		// Disable automatic knot updates during removal, we manage the knots together ourselves
-		Base::SetEnableKnotUpdates(false);
-		
     	// Get the Knots at segment boundaries
     	FInterval1f SegmentRange = GetSegmentParameterRange(SegmentIndex);
     	float StartParam = SegmentRange.Min;
@@ -986,7 +974,6 @@ public:
 	    	Base::RemoveKnotsCompletely(StartParam, ExpectedKnots);
 	    	
 	    }
-		Base::SetEnableKnotUpdates(true);
 	    return bSuccess;
     }
 
@@ -1734,7 +1721,8 @@ private:
 		// For a Bezier spline, the control points don't all lie on the curve
 		// So we need to compute chord lengths between segment endpoints only
 		int NumSegments = Points.Num() / 4;
-    	
+		constexpr float MinSegmentLength = UE_KINDA_SMALL_NUMBER;
+		
 		TArray<float> SegmentLengths;
 		SegmentLengths.Reserve(NumSegments);
 		float TotalLength = 0.0f;
@@ -1759,9 +1747,9 @@ private:
 					// Actual distance for chord length mode
 					Length = Math::Distance(Points[EndIndex], Points[StartIndex]);
 				}
-                
-				SegmentLengths.Add(static_cast<float>(Length));
-				TotalLength += static_cast<float>(Length);
+				float ClampedLength = FMath::Max(static_cast<float>(Length), MinSegmentLength);
+				SegmentLengths.Add(static_cast<float>(ClampedLength));
+				TotalLength += static_cast<float>(ClampedLength);
 			}
 		}
         
