Knot vector updates

diff --git a/Engine/Source/Runtime/GeometryCore/Public/Splines/BSpline.h b/Engine/Source/Runtime/GeometryCore/Public/Splines/BSpline.h
--- a/Engine/Source/Runtime/GeometryCore/Public/Splines/BSpline.h
+++ b/Engine/Source/Runtime/GeometryCore/Public/Splines/BSpline.h
@@ -25,6 +25,28 @@ enum class EParameterizationPolicy
 	Centripetal,
 };
 
+struct FKnot
+{
+	float Value;
+	uint32 Multiplicity;
+
+	FKnot() : Value(0.0f), Multiplicity(1) {}
+	FKnot(float InValue, uint32 InMultiplicity = 1) 
+		: Value(InValue), Multiplicity(InMultiplicity) {}
+
+	bool operator==(const FKnot& Other) const
+	{
+		return FMath::IsNearlyEqual(Value, Other.Value) && Multiplicity == Other.Multiplicity;
+	}
+
+	friend FArchive& operator<<(FArchive& Ar, FKnot& Knot)
+	{
+		Ar << Knot.Value;
+		Ar << Knot.Multiplicity;
+		return Ar;
+	}
+};
+	
 // Generate compile-time type ID for this BSpline specialization
 // Using degree in the implementation name for uniqueness
 template<int Degree>
@@ -66,7 +88,7 @@ public:
 	virtual void Clear() override
 	{
 		Values.Empty();
-		Knots.Empty();
+		PairKnots.Empty();
 	}
 
 	virtual bool IsEqual(const ISplineInterface* OtherSpline) const override
@@ -89,16 +111,16 @@ public:
 		}
     	
 		Ar << Values;
-		Ar << Knots;
+		Ar << PairKnots;
 		Ar << bIsClosedLoop;
 		Ar << bClampEnds;
-		
+		MarkFlatKnotsCacheDirty();
 		return true;
 	}
 	bool operator==(const TBSpline& Other) const
 	{
 		return	Values == Other.Values &&
-				Knots == Other.Knots &&
+				PairKnots == Other.PairKnots &&
 				bIsClosedLoop == Other.bIsClosedLoop &&
 				bClampEnds == Other.bClampEnds;
 	}
@@ -117,7 +139,7 @@ public:
     
 		// Copy all member variables
 		Clone->Values = this->Values;
-		Clone->Knots = this->Knots;
+		Clone->PairKnots = this->PairKnots;
 		Clone->bIsClosedLoop = this->bIsClosedLoop;
 		Clone->bClampEnds = this->bClampEnds;
     
@@ -185,13 +207,14 @@ public:
 	// Returns parameter value for a control point index using Greville Abscissae
     virtual float GetParameter(int32 Index) const
 	{
-		if (Index < 0 || Index + Degree >= Knots.Num() - 1)
+		EnsureFlatKnotsCache();
+		if (Index < 0 || Index + Degree >= FlatKnots.Num() - 1)
 			return 0.0f;
 
 		float Sum = 0.0f;
 		for (int j = 1; j <= Degree; ++j)
 		{
-			Sum += Knots[Index + j];
+			Sum += FlatKnots[Index + j];
 		}
 		return Sum / static_cast<float>(Degree);
 	}
@@ -199,6 +222,9 @@ public:
 	// Find the control point index and local parameter for global parameter
 	virtual int32 FindIndexForParameter(float Parameter, float& OutLocalParam) const
 	{
+		EnsureFlatKnotsCache();
+        const TArray<float>& Knots = FlatKnots;
+		
 		if (Knots.IsEmpty())
 			return INDEX_NONE;
         
@@ -251,32 +277,29 @@ public:
 		return Index;
 	}
 	// Knot vector management
-	const TArray<float>& GetKnotVector() const { return Knots; }
-	void ResetKnotVector() { Knots.Reset(); }
+	const TArray<float>& GetKnotVector() const
+	{
+		EnsureFlatKnotsCache();
+		return FlatKnots;
+	}
+	// Get the pair-based knot vector
+    const TArray<FKnot>& GetPairKnots() const
+    {
+        return PairKnots;
+    }
+	void ResetKnotVector() { PairKnots.Reset(); FlatKnots.Reset(); }
 	/**
      * Gets the array of distinct knot values (without duplicates)
      * @return Array of distinct knot values in ascending order
      */
     TArray<float> GetDistinctKnotValues() const
     {
-        TArray<float> DistinctKnots;
-        
-        if (Knots.Num() > 0)
-        {
-            // Add first knot
-            DistinctKnots.Add(Knots[0]);
-            
-            // Add all subsequent knots that differ from the previous one
-            for (int32 i = 1; i < Knots.Num(); ++i)
-            {
-                if (!FMath::IsNearlyEqual(Knots[i], Knots[i-1]))
-                {
-                    DistinctKnots.Add(Knots[i]);
-                }
-            }
-        }
-        
-        return DistinctKnots;
+		TArray<float> Distinct;
+		for (const FKnot& Knot : PairKnots)
+		{
+			Distinct.Add(Knot.Value);
+		}
+		return Distinct;
     }
     
     /**
@@ -333,6 +356,8 @@ public:
      */
 	void GenerateUniformKnots(int32 KnotCount)
 	{
+		PairKnots.Reset();
+		
     	int32 NumPoints = NumKeys();
 		if (this->IsClosedLoop())
 		{
@@ -344,46 +369,47 @@ public:
 			const int32 d = Degree; 
 			const int32 NumKnots = n + 2*d + 1;
         
-			Knots.SetNum(NumKnots);
 
 			// For periodic B-splines, use uniform knot spacing
 			for (int32 i = 0; i < NumKnots; ++i)
 			{
 				// Map to [0,n] range for n segments
-				Knots[i] = static_cast<float>(i);
+				PairKnots.Add( FKnot(static_cast<float>(i), 1) );
 			}
 		}
 		else
 		{
-			Knots.SetNum(KnotCount);
+
 			for (int32 i = 0; i < KnotCount; ++i)
 			{
+				float Value;
 				if (bClampEnds)
 				{
 					if (i <= Degree)
 					{
 						// Multiple knots at start for endpoint interpolation
-						Knots[i] = 0.0f;
+						Value = 0.0f;
 					}
 					else if (i >= KnotCount - Degree - 1)
 					{
 						// Multiple knots at end for endpoint interpolation
-						Knots[i] = static_cast<float>(KnotCount - 2 * Degree);
+						Value = static_cast<float>(KnotCount - 2 * Degree);
 					}
 					else
 					{
 						// Uniform spacing for internal knots
-						Knots[i] = static_cast<float>(i - Degree);
+						Value = static_cast<float>(i - Degree);
 					}
 				}
 				else
 				{
 					// Simple uniform spacing when not clamped
-					Knots[i] = static_cast<float>(i);
+					Value = static_cast<float>(i);
 				}
+				PairKnots.Add(FKnot(Value, 1));
 			}
 		}
-
+		MarkFlatKnotsCacheDirty();
 	}
 
 	/**
@@ -402,7 +428,7 @@ public:
 	    const int32 n = Points.Num();
 	    const int32 d = Degree;
 	    const int32 NumKnots = this->IsClosedLoop() ? (n + 2*d + 1) : KnotCount;
-	    Knots.SetNum(NumKnots);
+		PairKnots.Reset();
 
 	    // Calculate chord lengths between consecutive points
 	    TArray<float> Chords;
@@ -429,7 +455,7 @@ public:
 	        // For closed loop, generate periodic knot vector based on chord lengths
 	        for (int32 i = 0; i < NumKnots; ++i)
 	        {
-	            Knots[i] = static_cast<float>(i) * TotalLength / static_cast<float>(n + d);
+	            PairKnots.Add( FKnot(static_cast<float>(i) * TotalLength / static_cast<float>(n + d)));
 	        }
 	    }
 	    else if (bClampEnds)
@@ -439,18 +465,9 @@ public:
 	        // - Last (Degree+1) knots equal to 1.0
 	        // - Middle knots based on chord lengths
 	        
-	        // First set all knots to 0
-	        for (int32 i = 0; i < KnotCount; ++i)
-	        {
-	            Knots[i] = 0.0f;
-	        }
-	        
-	        // Set the last (Degree+1) knots to 1.0
-	        for (int32 i = KnotCount - Degree - 1; i < KnotCount; ++i)
-	        {
-	            Knots[i] = 1.0f;
-	        }
-	        
+	        // First set starting knots to 0
+	    	PairKnots.Add(FKnot(0.0f, Degree+1));
+	    	
 	        // Calculate internal knots based on accumulated chord lengths
 	        // We have (KnotCount - 2*(Degree+1)) internal knots to set
 	        const int32 NumInternalKnots = KnotCount - 2 * (Degree + 1);
@@ -482,9 +499,12 @@ public:
 	                
 	                // Interpolate Knots
 	                const float KnotValue = FMath::Lerp(Params[LowIndex], Params[HighIndex], Alpha);
-	                Knots[Degree + 1 + i] = KnotValue;
+	            	PairKnots.Add(FKnot(KnotValue, 1));
 	            }
 	        }
+	    	// Set the last (Degree+1) knots to 1.0
+	    	PairKnots.Add(FKnot(1.0f, Degree+1));
+
 	    }
 	    else
 	    {
@@ -492,9 +512,10 @@ public:
 	        for (int32 i = 0; i < KnotCount; ++i)
 	        {
 	            const float Param = static_cast<float>(i) / static_cast<float>(KnotCount - 1);
-	            Knots[i] = Param * TotalLength;
+	            PairKnots.Add( FKnot(Param * TotalLength));
 	        }
 	    }
+		MarkFlatKnotsCacheDirty();
 	}
 
     /**
@@ -514,7 +535,7 @@ public:
 	    const int32 n = Points.Num();
 	    const int32 d = Degree;
 	    const int32 NumKnots = this->IsClosedLoop() ? (n + 2*d + 1) : KnotCount;
-	    Knots.SetNum(NumKnots);
+		PairKnots.Reset();
 	    
 	    TArray<float> Lengths;
 	    Lengths.Reserve(Points.Num() - 1);
@@ -541,25 +562,16 @@ public:
 	        // Generate periodic knot vector using centripetal parameterization
 	        for (int32 i = 0; i < NumKnots; ++i)
 	        {
-	            Knots[i] = static_cast<float>(i) * TotalLength / static_cast<float>(n + d);
+	            PairKnots.Add( FKnot(static_cast<float>(i) * TotalLength / static_cast<float>(n + d)));
 	        }
 	    }
 	    else if (bClampEnds)
 	    {
 	        // For clamped ends, same approach as chord length but with square root distances
 	        
-	        // First set all knots to 0
-	        for (int32 i = 0; i < KnotCount; ++i)
-	        {
-	            Knots[i] = 0.0f;
-	        }
-	        
-	        // Set the last (Degree+1) knots to 1.0
-	        for (int32 i = KnotCount - Degree - 1; i < KnotCount; ++i)
-	        {
-	            Knots[i] = 1.0f;
-	        }
-	        
+	    	// First set starting knots to 0
+	    	PairKnots.Add(FKnot(0.0f, Degree+1));
+	        	        
 	        // Calculate internal knots based on accumulated centripetal Knots
 	        const int32 NumInternalKnots = KnotCount - 2 * (Degree + 1);
 	        if (NumInternalKnots > 0)
@@ -590,9 +602,11 @@ public:
 	                
 	                // Interpolate Knots
 	                const float KnotValue = FMath::Lerp(Params[LowIndex], Params[HighIndex], Alpha);
-	                Knots[Degree + 1 + i] = KnotValue;
+	                PairKnots.Add(FKnot(KnotValue));
 	            }
 	        }
+	    	// Set the last (Degree+1) knots to 1.0
+	    	PairKnots.Add(FKnot(1.0f, Degree+1));
 	    }
 	    else
 	    {
@@ -600,55 +614,65 @@ public:
 	        for (int32 i = 0; i < KnotCount; ++i)
 	        {
 	            const float Param = static_cast<float>(i) / static_cast<float>(KnotCount - 1);
-	            Knots[i] = Param * TotalLength;
+	            PairKnots.Add(FKnot(Param * TotalLength));
 	        }
 	    }
+		MarkFlatKnotsCacheDirty();
 	}
 
 	/**
 	 * Inserts a knot at the specified parameter value
 	 * 
-	 * @param Parameter - Parameter value where the knot should be inserted
+	 * @param Value - Knot value where the knot should be inserted
 	 * @param Count - Number of times to insert the knot
 	 */
-	void InsertKnot(float Parameter, int32 Count = 1)
+	void InsertKnot(float Value, int32 Count = 1)
 	{
-		int32 InsertIndex = Knots.Num();  // Default to appending
-
-		for (int32 i = 0; i < Knots.Num(); ++i)
+		for (int32 i = 0; i < PairKnots.Num(); ++i)
 		{
-			if (Knots[i] >= Parameter)
+			if (PairKnots[i].Value == Value)  // exact match
 			{
-				// Ensure we insert *after* any existing repeated values
-				while (i < Knots.Num() && FMath::IsNearlyEqual(Knots[i], Parameter))
-				{
-					i++;
-				}
-				InsertIndex = i;
-				break;
+				PairKnots[i].Multiplicity += Count;
+				MarkFlatKnotsCacheDirty();
+				return;
+			}
+			else if (PairKnots[i].Value > Value)
+			{
+				PairKnots.Insert(FKnot(Value, Count), i);
+				MarkFlatKnotsCacheDirty();
+				return;
 			}
 		}
 
-		// Insert the parameter count times at the computed index
-		for (int32 i = 0; i < Count; ++i)
-		{
-			Knots.Insert(Parameter, InsertIndex);
-		}
+		PairKnots.Add(FKnot(Value, Count));
+		MarkFlatKnotsCacheDirty();
 	}
 
 	/**
 	 * Replaces all instances of a knot.
 	 * 
-	 * @param OldParameter - Parameter value to search for
-	 * @param NewParameter - Value to replace instances of OldParameter
+	 * @param OldValue - Knot value to search for
+	 * @param NewValue - Value to replace instances of OldValue
 	 */
-	void ReplaceKnot(float OldParameter, float NewParameter)
+	void ReplaceKnot(float OldValue, float NewValue)
 	{
-		for (int32 i = 0; i < Knots.Num(); ++i)
+		if (OldValue == NewValue)
 		{
-			if (FMath::IsNearlyEqual(Knots[i], OldParameter))
+			// Nothing to replace
+			return;
+		}
+
+		for (int32 i = 0; i < PairKnots.Num(); ++i)
+		{
+			if (PairKnots[i].Value == OldValue)
 			{
-				Knots[i] = NewParameter;
+				uint32 Count = PairKnots[i].Multiplicity;
+				PairKnots.RemoveAt(i);
+
+				// Insert the new knot with the same multiplicity
+				InsertKnot(NewValue, Count); // Will re-mark cache dirty
+
+				return;
 			}
 		}
 	}
@@ -661,194 +685,130 @@ public:
 	 */
 	void SwapKnots(float ParameterA, float ParameterB)
 	{
-		for (int32 i = 0; i < Knots.Num(); ++i)
+		for (int32 i = 0; i < PairKnots.Num(); ++i)
 		{
-			if (FMath::IsNearlyEqual(Knots[i], ParameterA))
+			if (PairKnots[i].Value == ParameterA)
 			{
-				Knots[i] = ParameterB;
+				PairKnots[i].Value = ParameterB;
 			}
-			else if (FMath::IsNearlyEqual(Knots[i], ParameterB))
+			if (PairKnots[i].Value == ParameterB)
 			{
-				Knots[i] = ParameterA;
+				PairKnots[i].Value = ParameterA;
 			}
 		}
 	}
 
 	/**
 	 * Returns how many times a given knot value appears.
-	 * @param Parameter - Knot value to check
+	 * @param Value - Knot value to check
 	 * @return Number of times the knot appears
 	 */
-	int32 GetKnotMultiplicity(float Parameter) const
+	int32 GetKnotMultiplicity(float Value) const
 	{
-	    int32 Count = 0;
-	    for (const float& Knot : Knots)
-	    {
-	        if (FMath::IsNearlyEqual(Knot, Parameter))
-	        {
-	            ++Count;
-	        }
-	    }
-	    return Count;
+		for (const FKnot& Knot : PairKnots)
+		{
+			if (Knot.Value == Value)
+			{
+				return static_cast<int32>(Knot.Multiplicity);
+			}
+		}
+		return 0;
 	}
 
-	// Given a parameter and an array of distinct knots (mult = 1),
-	// find the parameter closest to DesiredParameter which corresponds to no distinct knots.
-	static float GetNearestValidKnot(float DesiredParameter, const TArray<float>& Knots, float Epsilon = UE_KINDA_SMALL_NUMBER)
+	
+	
+	bool RemoveKnotsCompletely(float Value, int32 ExpectedKnotCount)
 	{
-		if (Knots.Num() == 0)
+		const int32 Removed = PairKnots.RemoveAll([&](const FKnot& K)
 		{
-			// If we have an empty knot vector, DesiredParameter is definitely valid.
-			return DesiredParameter;
-		}
+			return K.Value == Value;
+		});
 
-		if (Knots.Num() == 1)
+		if (Removed > 0)
 		{
-			// If we have 1 knot, we can special case to simplify logic in the loop below.
-			return DesiredParameter < Knots[0]
-				? FMath::Min(DesiredParameter, Knots[0] - Epsilon)
-				: FMath::Max(DesiredParameter, Knots[0] + Epsilon);
+			// fix end knots if an end knot is removed
+			FixEndKnots(ExpectedKnotCount);
+			MarkFlatKnotsCacheDirty();
+			return true;
 		}
 
-		// A candidate is valid if it is not too close to any existing knots (including the knot which was used to generate the candidate)
-		auto IsCandidateValid = [&Knots, Epsilon](const float Candidate) -> bool
+		return false;
+
+
+	}
+	/**
+	 * Removes a single occurrence of a knot value.
+	 * For clamped endpoints, we avoid removing knots if doing so would drop the multiplicity
+	 * below the required (Degree+1) at the endpoints.
+	 * @param Value - Knot value to remove
+	 * @param Count - Number of times to remove the knot
+	 * @return true if the knot was removed
+	 */
+	bool RemoveKnot(float Value, int32 Count = 1)
+	{
+		const TArray<FKnot> OriginalKnots = PairKnots;
+
+		for (int32 i = 0; i < PairKnots.Num(); ++i)
 		{
-			for (const float Knot : Knots)
+			if (PairKnots[i].Value == Value)
 			{
-				if (FMath::Abs(Candidate - Knot) < Epsilon)
+				bool bIsEndpoint =
+					(Value == GetKnotRange().Min && bClampEnds) ||
+					(Value == GetKnotRange().Max && bClampEnds);
+
+				uint32 MinRequiredMultiplicity = bIsEndpoint ? Degree + 1 : 1;
+				uint32 Current = PairKnots[i].Multiplicity;
+
+				if (Current <= MinRequiredMultiplicity)
 				{
+					UE_LOG(LogSpline, Warning, TEXT("Cannot remove knot %f below required multiplicity %u"), Value, MinRequiredMultiplicity);
 					return false;
 				}
-			}
 
-			return true;
-		};
-		
-		// Candidates are closest valid parameters to existing knots which are too close
-		// to DesiredParameter to allow DesiredParameter to be considered a valid knot.
-		TArray<float> Candidates;
-		bool bDegenerate = false;
-		for (int Idx = 0; Idx < Knots.Num(); ++Idx)
-		{
-			const float Knot = Knots[Idx];
-
-			if (FMath::Abs(DesiredParameter - Knot) < Epsilon)
-			{
-				bDegenerate = true;
-				UE_LOG(LogSpline, Warning, TEXT("Parameter %f is too close to existing knot value %f."), DesiredParameter, Knot);
-
-				if (float Candidate = std::nextafter(Knot - Epsilon, -INFINITY); IsCandidateValid(Candidate))
+				uint32 NewCount = Current - Count;
+				if (NewCount < MinRequiredMultiplicity)
+				{
+					PairKnots.RemoveAt(i);
+				}
+				else
 				{
-					Candidates.Add(Candidate);
+					PairKnots[i].Multiplicity = NewCount;
 				}
 
-				if (float Candidate = std::nextafter(Knot + Epsilon, +INFINITY); IsCandidateValid(Candidate))
+				// Validate
+				TArray<float> Flat = ConvertPairToFlatKnots(PairKnots);
+				if (!Math::FSplineValidation::IsValidKnotVector(Flat, Degree))
 				{
-					Candidates.Add(Candidate);
+					UE_LOG(LogSpline, Warning, TEXT("Knot removal resulted in invalid knot vector. Reverting."));
+					PairKnots = OriginalKnots;
+					return false;
 				}
-			}
-		}
 
-		if (bDegenerate && Candidates.Num() == 0)
-		{
-			UE_LOG(LogSpline, Warning, TEXT("Parameter %f is invalid and no good alternatives were found, expect a degenerate segment."), DesiredParameter);
-		}
-		
-		float BestCandidate = DesiredParameter;
-		float BestCandidateDistanceToDesiredParameter = INFINITY;	// Set to infinity so that we take a candidate if one exists.
-		
-		for (const float Candidate : Candidates)
-		{
-			if (const float CandidateDistanceToDesiredParameter = FMath::Abs(DesiredParameter - Candidate);
-				CandidateDistanceToDesiredParameter < BestCandidateDistanceToDesiredParameter)
-			{
-				BestCandidate = Candidate;
-				BestCandidateDistanceToDesiredParameter = CandidateDistanceToDesiredParameter;
+				MarkFlatKnotsCacheDirty();
+				return true;
 			}
 		}
 
-		return BestCandidate;
+		return false;
 	}
 
-	void RemoveKnotsCompletely(float Parameter, int32 ExpectedKnotCount)
-	{
-		// Remove all occurrences of the knot without checking multiplicity
-		Knots.Remove(Parameter);
-
-		// fix end knots if an end knot is removed
-		FixEndKnots(ExpectedKnotCount);
-
-	}
 	/**
-	 * Removes a single occurrence of a knot value.
-	 * For clamped endpoints, we avoid removing knots if doing so would drop the multiplicity
-	 * below the required (Degree+1) at the endpoints.
-	 * @param Parameter - Knot value to remove
-	 * @param Count - Number of times to remove the knot
-	 * @return true if the knot was removed
+	 * Converts a pair-based knot vector to a flat knot vector
+	 * @param PairKnots - The pair-based knot vector
+	 * @return Flat knot vector with multiplicities expanded
 	 */
-	bool RemoveKnot(float Parameter, int32 Count = 1)
+	static TArray<float> ConvertPairToFlatKnots(const TArray<FKnot>& PairKnots)
 	{
-		TArray OrigKnotVector = Knots;
-		
-		// First check if this is an endpoint knot
-		FInterval1f KnotRange = GetKnotRange();
-		bool bIsEndpoint = FMath::IsNearlyEqual(Parameter, KnotRange.Min) || 
-						   FMath::IsNearlyEqual(Parameter, KnotRange.Max);
-    
-		bool bSuccess = false;
-				
-		// Find all knots with this parameter value
-		TArray<int32> KnotIndices;
-		for (int32 i = 0; i < Knots.Num(); ++i)
-		{
-			if (FMath::IsNearlyEqual(Knots[i], Parameter))
-			{
-				KnotIndices.Add(i);
-			}
-		}
-		// Check current multiplicity
-		int32 CurrentMultiplicity = KnotIndices.Num();
-		if (CurrentMultiplicity == 0)
+		TArray<float> FlatKnots;
+		for (const FKnot& Knot : PairKnots)
 		{
-			UE_LOG(LogSpline, Warning, TEXT("No knots found at parameter %f"), Parameter);
-			return false;
-		}
-
-		// Calculate how many knots we can safely remove
-		int32 MinRequiredMultiplicity = bIsEndpoint && bClampEnds ? (Degree + 1) : 1;
-		int32 MaxRemovableKnots = FMath::Max(0, CurrentMultiplicity - MinRequiredMultiplicity);
-		int32 KnotsToRemove = FMath::Min(Count, MaxRemovableKnots);
-
-		if (KnotsToRemove <= 0)
-		{
-			if (bIsEndpoint && bClampEnds)
+			for (uint32 i = 0; i < Knot.Multiplicity; ++i)
 			{
-				UE_LOG(LogSpline, Warning, TEXT("Cannot remove endpoint knot below required multiplicity: %f (current: %d, required: %d)"), 
-					Parameter, CurrentMultiplicity, MinRequiredMultiplicity);
+				FlatKnots.Add(Knot.Value);
 			}
-			return false;
-		}
-
-		// Remove knots in reverse order to maintain valid indices
-		
-		for (int32 i = 0; i < KnotsToRemove; ++i)
-		{
-			// Always remove from the last found index to avoid shifting issues
-			Knots.RemoveAt(KnotIndices[KnotIndices.Num() - 1 - i]);
-			bSuccess = true;
-		}
-		
-		// validate if removing completely adheres to knot vector requirements
-		if (!UE::Geometry::Spline::Math::FSplineValidation::IsValidKnotVector(Knots, Degree))
-		{
-			UE_LOG(LogSpline, Warning, TEXT("Invalid knot vector after removing knot(s) %f"), Parameter);
-			Knots = OrigKnotVector;
-			return false;
 		}
-		return bSuccess;
+		return FlatKnots;
 	}
-	
 	/**
 	 * Sets a custom knot vector 
 	 * @param NewKnots - Array of knot values. Must be non-decreasing and have correct length (n + p + 1)
@@ -878,14 +838,44 @@ public:
 			}
 		}
 
-		Knots = NewKnots;
-		
+		// Convert flat knots to pair representation
+		PairKnots.Reset();
+    
+		if (NewKnots.Num() > 0)
+		{
+			float CurrentValue = NewKnots[0];
+			uint32 Multiplicity = 1;
+        
+			for (int32 i = 1; i < NewKnots.Num(); ++i)
+			{
+				if (FMath::IsNearlyEqual(NewKnots[i], CurrentValue))
+				{
+					// Same value, increase multiplicity
+					Multiplicity++;
+				}
+				else
+				{
+					// Different value, add previous knot and start new one
+					PairKnots.Add(FKnot(CurrentValue, Multiplicity));
+					CurrentValue = NewKnots[i];
+					Multiplicity = 1;
+				}
+			}
+        
+			// Add the last knot
+			PairKnots.Add(FKnot(CurrentValue, Multiplicity));
+		}
+    
+		// Mark the flat knots cache as dirty
+		MarkFlatKnotsCacheDirty();
+    
 		UE_LOG(LogSpline, Verbose, TEXT("Set Custom knot vector - Points: %d, Degree: %d, Knots: %d"),
-			   NumValues, Degree, Knots.Num());
-		for (int32 i = 0; i < Knots.Num(); ++i)
+			   NumValues, Degree, NewKnots.Num());
+		for (int32 i = 0; i < NewKnots.Num(); ++i)
 		{
-			UE_LOG(LogSpline, Verbose, TEXT("  Knot[%d] = %f"), i, Knots[i]);
+			UE_LOG(LogSpline, Verbose, TEXT("  Knot[%d] = %f"), i, NewKnots[i]);
 		}
+    
 		
 		return true;
 	}
@@ -909,7 +899,7 @@ public:
 		}
 
 		const int32 KnotCount = NumValues + Degree + 1;
-		Knots.SetNum(KnotCount);
+		PairKnots.Reset();
 
 		switch (ParameterizationPolicy)
 		{
@@ -924,12 +914,11 @@ public:
 			break;
 		}
 
+		// Mark the flat knots cache as dirty
+		MarkFlatKnotsCacheDirty();
+
 		UE_LOG(LogSpline, Verbose, TEXT("Updated knot vector - Points: %d, Degree: %d, Knots: %d"),
-			   Points.Num(), Degree, KnotCount);
-		for (int32 i = 0; i < KnotCount; ++i)
-		{
-			UE_LOG(LogSpline, Verbose, TEXT("  Knot[%d] = %f"), i, Knots[i]);
-		}
+			   Points.Num(), Degree, PairKnots.Num());
 	}
 	
 	/** 
@@ -949,6 +938,7 @@ public:
 	// Get the valid parameter range in knot space
 	FInterval1f GetKnotRange() const
 	{
+		TArray<float> Knots = GetKnotVector();
 		if (Knots.Num() - Degree - 1 < 0)
 		{
 			return FInterval1f::Empty();
@@ -964,41 +954,201 @@ public:
 			return FInterval1f(Knots[0], Knots.Last());
 		}
 	}
-	
+
+	// Ensure we have the correct number of knots at the endpoints
 	void FixEndKnots(const int32 ExpectedKnots)
 	{
-		// Calculate how many knots to add to reach desired multiplicity
-		int32 Diff = FMath::Abs(Knots.Num() - ExpectedKnots);
-		if (Knots.Num() > ExpectedKnots)
+		if (PairKnots.Num() == 0)
+			return;
+        
+		// Calculate current flat knot count
+		int32 CurrentFlatCount = 0;
+		for (const FKnot& Knot : PairKnots)
 		{
-			Knots.SetNum(ExpectedKnots);
+			CurrentFlatCount += static_cast<int32>(Knot.Multiplicity);
 		}
-		else if (Knots.Num() < ExpectedKnots)
+    
+		// Calculate how many knots to add or remove
+		int32 Diff = ExpectedKnots - CurrentFlatCount;
+    
+		if (Diff == 0)
+		{
+			// Nothing to do
+			return;
+		}
+		else if (Diff > 0)
 		{
-			float ClampedKnot = Knots.Last();
-			for (int32 i = 0; i < Diff; ++i)
+			// Need to add knots - add them to the end with the same value as the last knot
+			FKnot& LastKnot = PairKnots.Last();
+			LastKnot.Multiplicity += static_cast<uint32>(Diff);
+		}
+		else // Diff < 0
+		{
+			// Need to remove knots - remove them from the end
+			uint32 KnotsToRemove = static_cast<uint32>(-Diff);
+        
+			// Start removing from the last knot
+			while (KnotsToRemove > 0 && PairKnots.Num() > 0)
 			{
-				Knots.Add(ClampedKnot);
+				FKnot& LastKnot = PairKnots.Last();
+            
+				if (LastKnot.Multiplicity <= KnotsToRemove)
+				{
+					// Remove the entire knot
+					KnotsToRemove -= LastKnot.Multiplicity;
+					PairKnots.RemoveAt(PairKnots.Num() - 1);
+				}
+				else
+				{
+					// Reduce multiplicity
+					LastKnot.Multiplicity -= KnotsToRemove;
+					KnotsToRemove = 0;
+				}
 			}
 		}
+    
+		MarkFlatKnotsCacheDirty();
 	}
+	// Fill the end of the knot vector with the specified knot value
 	void FillEndKnots(float NewEndKnot, const int32 ExpectedKnots)
 	{
-		// Calculate how many knots to add to reach desired multiplicity
-		int32 DesiredMultiplicity = ExpectedKnots - Knots.Num(); 
-		int32 KnotsToAdd = FMath::Max(0, DesiredMultiplicity);
-        
-		UE_LOG(LogSpline, Log, TEXT("Adding %d knots at new end parameter %f to reach multiplicity %d"),
-			   KnotsToAdd, NewEndKnot, DesiredMultiplicity);
-        
-		// Insert knots to reach desired multiplicity
-		for (int32 i = 0; i < KnotsToAdd; ++i)
+		// Calculate current flat knot count
+		int32 CurrentFlatCount = 0;
+		for (const FKnot& Knot : PairKnots)
+		{
+			CurrentFlatCount += static_cast<int32>(Knot.Multiplicity);
+		}
+    
+		// Calculate how many knots to add
+		int32 KnotsToAdd = FMath::Max(0, ExpectedKnots - CurrentFlatCount);
+    
+		if (KnotsToAdd == 0)
+		{
+			return;
+		}
+    
+		UE_LOG(LogSpline, Log, TEXT("Adding %d knots at new end parameter %f"),
+			   KnotsToAdd, NewEndKnot);
+    
+		// Check if the last knot has the same value
+		if (PairKnots.Num() > 0 && FMath::IsNearlyEqual(PairKnots.Last().Value, NewEndKnot))
+		{
+			// If the last knot has the same value, just increase its multiplicity
+			PairKnots.Last().Multiplicity += static_cast<uint32>(KnotsToAdd);
+		}
+		else
 		{
-			Knots.Add(NewEndKnot);
+			// Otherwise, add a new knot with the desired multiplicity
+			PairKnots.Add(FKnot(NewEndKnot, static_cast<uint32>(KnotsToAdd)));
 		}
+    
+		MarkFlatKnotsCacheDirty();
 	}
 	
 protected:
+	// Find a valid parameter value that doesn't conflict with existing knots
+    float GetNearestValidKnot(float DesiredParameter, float ExcludeValue) const
+    {
+        if (PairKnots.Num() == 0)
+        {
+            return DesiredParameter;
+        }
+
+        // Get valid parameter space
+        FInterval1f ParamSpace = GetParameterSpace();
+        
+        // First ensure the desired parameter is within valid range
+        DesiredParameter = FMath::Clamp(DesiredParameter, ParamSpace.Min, ParamSpace.Max);
+
+        // Check if DesiredParameter conflicts with any existing knot (excluding ExcludeValue)
+        bool bConflict = false;
+        for (const FKnot& Knot : PairKnots)
+        {
+            if (Knot.Value != ExcludeValue && Knot.Value == DesiredParameter)
+            {
+                bConflict = true;
+                break;
+            }
+        }
+
+        // If no conflict found, parameter is already valid
+        if (!bConflict)
+        {
+            return DesiredParameter;
+        }
+
+        // Find the closest valid parameter by walking outward alternately
+        float Left = DesiredParameter;
+        float Right = DesiredParameter;
+        bool bSearchLeft = true;
+        bool bSearchRight = true;
+        
+        // Continue as long as we have valid search directions
+        while (bSearchLeft || bSearchRight)
+        {
+            // Try moving left if still in valid range
+            if (bSearchLeft)
+            {
+                Left = std::nextafter(Left, -INFINITY);
+                
+                // Stop if we've moved outside the parameter space
+                if (Left < ParamSpace.Min)
+                {
+                    bSearchLeft = false;
+                    continue; // Skip to right search
+                }
+                
+                // Check if this value conflicts with any knot (excluding ExcludeValue)
+                bool bLeftValid = true;
+                for (const FKnot& Knot : PairKnots)
+                {
+                    if (Knot.Value != ExcludeValue && Knot.Value == Left)
+                    {
+                        bLeftValid = false;
+                        break;
+                    }
+                }
+                
+                if (bLeftValid)
+                {
+                    return Left; // Found valid parameter to the left
+                }
+            }
+            
+            // Try moving right if still in valid range
+            if (bSearchRight)
+            {
+                Right = std::nextafter(Right, INFINITY);
+                
+                // Stop if we've moved outside the parameter space
+                if (Right > ParamSpace.Max)
+                {
+                    bSearchRight = false;
+                    continue; // Skip to next iteration
+                }
+                
+                // Check if this value conflicts with any knot (excluding ExcludeValue)
+                bool bRightValid = true;
+                for (const FKnot& Knot : PairKnots)
+                {
+                    if (Knot.Value != ExcludeValue && Knot.Value == Right)
+                    {
+                        bRightValid = false;
+                        break;
+                    }
+                }
+                
+                if (bRightValid)
+                {
+                    return Right; // Found valid parameter to the right
+                }
+            }
+        }
+        
+        // If we couldn't find a valid parameter within range, return the closest boundary
+        return (FMath::Abs(ParamSpace.Min - DesiredParameter) < FMath::Abs(ParamSpace.Max - DesiredParameter)) 
+            ? ParamSpace.Min : ParamSpace.Max;
+    }
 	
 	virtual int32 GetExpectedNumKnots() const
 	{
@@ -1011,6 +1161,8 @@ private:
 	/** Find knot span using binary search */
 	int32 FindKnotSpan(float Parameter) const
 	{
+		EnsureFlatKnotsCache();
+		const TArray<float>& Knots = FlatKnots;
 		// First, check that the knot vector is populated.
 		if (Knots.Num() == 0)
 		{
@@ -1106,7 +1258,7 @@ private:
     {
     	TArray<float> Basis;
     	    
-    	if (Knots.IsEmpty())
+    	if (PairKnots.IsEmpty())
     		return ValueType();
 
     	TArray<float> KnotVector = GetKnotVector();
@@ -1117,8 +1269,44 @@ private:
     }
 	
 protected:
+	// Ensure the flat knot cache is up to date
+	void EnsureFlatKnotsCache() const
+	{
+		if (bFlatKnotsCacheDirty)
+		{
+			RebuildFlatKnotsCache();
+			bFlatKnotsCacheDirty = false;
+		}
+	}
+    
+	// Rebuild the flat knot cache from the pair representation
+	void RebuildFlatKnotsCache() const
+	{
+		FlatKnots.Reset();
+		for (const FKnot& Knot : PairKnots)
+		{
+			for (uint32 i = 0; i < Knot.Multiplicity; ++i)
+			{
+				FlatKnots.Add(Knot.Value);
+			}
+		}
+	}
+    
+	// Mark the flat knot cache as dirty
+	void MarkFlatKnotsCacheDirty() const
+	{
+		bFlatKnotsCacheDirty = true;
+	}
+	
     TArray<ValueType> Values;
-    TArray<float> Knots;
+
+	// Knot representation in pair form having multiplicity
+	TArray<FKnot> PairKnots;
+    
+	// Cached flattened knot vector for compatible algorithms
+	mutable TArray<float> FlatKnots;
+	mutable bool bFlatKnotsCacheDirty;
+	
 	bool bIsClosedLoop = false;
 	
 	/** Whether to clamp the endpoints for interpolation of first/last control points */
diff --git a/Engine/Source/Runtime/GeometryCore/Public/Splines/PolyBezierSpline.h b/Engine/Source/Runtime/GeometryCore/Public/Splines/PolyBezierSpline.h
--- a/Engine/Source/Runtime/GeometryCore/Public/Splines/PolyBezierSpline.h
+++ b/Engine/Source/Runtime/GeometryCore/Public/Splines/PolyBezierSpline.h
@@ -261,7 +261,7 @@ public:
         
         // Copy base class members
         Clone->Values = this->Values;
-        Clone->Knots = this->Knots;
+        Clone->PairKnots = this->PairKnots;
         Clone->bIsClosedLoop = this->bIsClosedLoop;
         Clone->bClampEnds = this->bClampEnds;
         
@@ -1307,10 +1307,10 @@ public:
     	
     	
 		UE_LOG(LogSpline, Verbose, TEXT("Set Custom knot vector - Points: %d, Knots: %d, Mode: %d"),
-					   NumValues,  this->Knots.Num(), static_cast<int32>(Mode));
-		for (int32 i = 0; i < this->Knots.Num(); ++i)
+					   NumValues,  this->PairKnots.Num(), static_cast<int32>(Mode));
+		for (int32 i = 0; i < this->PairKnots.Num(); ++i)
 		{
-			UE_LOG(LogSpline, Verbose, TEXT("  Knot[%d] = %f"), i, this->Knots[i]);
+			UE_LOG(LogSpline, Verbose, TEXT("  Knot[%d] = %f x %d"), i, this->PairKnots[i].Value, this->PairKnots[i].Multiplicity);
 		}
 	}
 	virtual FInterval1f GetParameterSpace() const override
@@ -1351,17 +1351,23 @@ public:
 		const int32 LocalPointIndex = Index % PointsPerSegment;
 		int32 NumSegments = GetNumSegments();
 
+		// Get the original parameter value
+		float OldParameter = GetParameter(Index);
+    
+		// Early exit if parameter isn't changing
+		if (OldParameter == NewParameter)
+		{
+			return SegmentIndex * PointsPerSegment + LocalPointIndex;
+		}
+		
 		// Prevent segment degeneration.
-		TArray<float> DistinctKnots = Base::GetDistinctKnotValues();
-		DistinctKnots.Remove(GetParameter(Index));
-		NewParameter = Base::GetNearestValidKnot(NewParameter, DistinctKnots);
+		NewParameter = this->GetNearestValidKnot(NewParameter, OldParameter);
 		
 		if (LocalPointIndex == 0) // P0
 		{
 			int32 LeftSegmentIndex = SegmentIndex - 1;
 			int32 RightSegmentIndex = SegmentIndex;
 			
-			float OldParameter = GetSegmentParameterRange(SegmentIndex).Min;
 			Base::ReplaceKnot(OldParameter, NewParameter);
 
 			int32 Segment = SegmentIndex;
@@ -1393,7 +1399,6 @@ public:
 			int32 LeftSegmentIndex = SegmentIndex;
 			int32 RightSegmentIndex = SegmentIndex + 1;
 			
-			float OldParameter = GetSegmentParameterRange(SegmentIndex).Max;
 			Base::ReplaceKnot(OldParameter, NewParameter);
 
 			int32 Segment = SegmentIndex;
@@ -1730,7 +1735,6 @@ private:
 	
 	void GenerateDistanceKnotsForBezier(EParameterizationPolicy Mode)
 	{
-		const int32 NumKnots = this->GetExpectedNumKnots();
 		TArray<ValueType> Points = this->Values;
 
 		// For a Bezier spline, the control points don't all lie on the curve
@@ -1772,16 +1776,13 @@ private:
 		// 1. Bezier segment structure (knot multiplicity)
 		// 2. Proportional chord/centripetal lengths
         
-		TArray<float> DistanceKnots;
-		DistanceKnots.SetNum(NumKnots);
+		// Reset pair knots
+		this->PairKnots.Reset();
         
 		int32 CurrentKnot = 0;
 
-		// Start knots with multiplicity 4
-		for (int32 i = 0; i < 4; ++i)
-		{
-			DistanceKnots[CurrentKnot++] = 0.0f;
-		}
+		// Add start knot with multiplicity 4
+		this->PairKnots.Add(FKnot(0.0f, 4));
         
 		// Calculate accumulated distances for internal knots
 		float AccumLength = 0.0f;
@@ -1795,22 +1796,16 @@ private:
 				const float t = AccumLength;
 				
 				// Interior knots at segment boundaries with multiplicity 3
-				for (int32 j = 0; j < 3; ++j)
-				{
-					DistanceKnots[CurrentKnot++] = t;
-				}
+				this->PairKnots.Add(FKnot(t, 3));
 			}
 		}
 
-		// End knots with multiplicity 4
+		// End knot with multiplicity 4
 		const float EndKnot = AccumLength + (SegmentLengths.Num() > 0 ? SegmentLengths.Last() : 0.0f);
-		while (CurrentKnot < NumKnots)
-		{
-			DistanceKnots[CurrentKnot++] = EndKnot;
-		}
+		this->PairKnots.Add(FKnot(EndKnot, 4));
 
-		// Set the knot vector
-		this->Knots = DistanceKnots;
+		// Mark flat knots cache as dirty
+		this->MarkFlatKnotsCacheDirty();
 	}
 
 	/**
@@ -1827,16 +1822,12 @@ private:
 		const int32 NumKnots = this->GetExpectedNumKnots();             // normal case: n + p + 1 where p = 3
         
 		// Set up knot vector with multiple knots at segment boundaries
-		TArray<float> BezierKnots;
-		BezierKnots.SetNum(NumKnots);
+		this->PairKnots.Reset();
         
 		int32 CurrentKnot = 0;
 
 		// Start knots with multiplicity 4
-		for (int32 i = 0; i < 4; ++i)
-		{
-			BezierKnots[CurrentKnot++] = 0.0f;
-		}
+		this->PairKnots.Add( FKnot(0.0f, 4));
 
 		// For closed loop, add an extra segment to connect back to start
 		const int32 NumSegmentsToProcess = this->IsClosedLoop() ? (NumSegments + 1) : NumSegments;
@@ -1845,21 +1836,15 @@ private:
 		for (int32 i = 1; i < NumSegmentsToProcess; ++i)
 		{
 			const float t = static_cast<float>(i);
-			for (int32 j = 0; j < 3; ++j)
-			{
-				BezierKnots[CurrentKnot++] = t;
-			}
+			this->PairKnots.Add(FKnot(t, 3));
 		}
 
-		// End knots with multiplicity 4
+		// End knot with multiplicity 4
 		const float EndKnot = static_cast<float>(NumSegmentsToProcess);
-		while (CurrentKnot < NumKnots)
-		{
-			BezierKnots[CurrentKnot++] = EndKnot;
-		}
+		this->PairKnots.Add(FKnot(EndKnot, 4));
 
-		// Set the knot vector
-		this->Knots = BezierKnots;
+		// Mark flat knots cache as dirty
+		this->MarkFlatKnotsCacheDirty();
 	}
 
     virtual TArray<const ValueType*> FindWindow(float Parameter) const override
