Adds GetAllocatedSize to TIndexMap.

diff --git a/Engine/Source/Runtime/GeometryCore/Public/GeometryTypes.h b/Engine/Source/Runtime/GeometryCore/Public/GeometryTypes.h
--- a/Engine/Source/Runtime/GeometryCore/Public/GeometryTypes.h
+++ b/Engine/Source/Runtime/GeometryCore/Public/GeometryTypes.h
@@ -194,6 +194,11 @@ public:
 			ReverseMap.Reserve(NumElements);
 		}
 	}
+	
+	SIZE_T GetAllocatedSize() const
+	{
+		return ForwardMap.GetAllocatedSize() + ReverseMap.GetAllocatedSize();
+	}
 };
 
 typedef TIndexMap<int> FIndexMapi;
