Adds dllexport to FDynamicMesh3::CanCollapseEdgeInternal. This is missing to be able to subclass FDynamicMesh from external modules.

diff --git a/Engine/Source/Runtime/GeometryCore/Public/DynamicMesh/DynamicMesh3.h b/Engine/Source/Runtime/GeometryCore/Public/DynamicMesh/DynamicMesh3.h
--- a/Engine/Source/Runtime/GeometryCore/Public/DynamicMesh/DynamicMesh3.h
+++ b/Engine/Source/Runtime/GeometryCore/Public/DynamicMesh/DynamicMesh3.h
@@ -1798,9 +1798,9 @@ protected:
 	void SerializeInternal(FArchive& Ar, void* Ptr);
 
 	/* We keep this version of CanCollapseEdge internal because the CollapseInfo struct may only be partially filled out by the function */
-	virtual EMeshResult CanCollapseEdgeInternal(int vKeep, int vRemove, double collapse_t, FEdgeCollapseInfo* OutCollapseInfo) const;
+	GEOMETRYCORE_API virtual EMeshResult CanCollapseEdgeInternal(int vKeep, int vRemove, double collapse_t, FEdgeCollapseInfo* OutCollapseInfo) const;
 private:
-	virtual EMeshResult CanCollapseEdgeInternal(int vKeep, int vRemove, double collapse_t,
+	GEOMETRYCORE_API virtual EMeshResult CanCollapseEdgeInternal(int vKeep, int vRemove, double collapse_t,
 		const FCollapseEdgeOptions& Options, FEdgeCollapseInfo* OutCollapseInfo) const;
 
 	friend class FDynamicMeshAttributeSet;
