SplineCore: Change some knot vector mutating functions to address knots by index instead of attempting to search by value. SetKnot no longer reorders the knot vector, it just naively sets values. GetNearestValidKnot no longer clamps to the existing parameter space.

diff --git a/Engine/Source/Runtime/GeometryCore/Public/Splines/BSpline.h b/Engine/Source/Runtime/GeometryCore/Public/Splines/BSpline.h
--- a/Engine/Source/Runtime/GeometryCore/Public/Splines/BSpline.h
+++ b/Engine/Source/Runtime/GeometryCore/Public/Splines/BSpline.h
@@ -651,51 +651,36 @@ public:
 	/**
 	 * Replaces all instances of a knot.
 	 * 
-	 * @param OldValue - Knot value to search for
+	 * @param KnotIdx - Index of the knot to replace
 	 * @param NewValue - Value to replace instances of OldValue
 	 */
-	void ReplaceKnot(float OldValue, float NewValue)
+	void SetKnot(int32 KnotIdx, float NewValue)
 	{
-		if (OldValue == NewValue)
+		if (KnotIdx < 0 || KnotIdx >= PairKnots.Num() || PairKnots[KnotIdx].Value == NewValue)
 		{
 			// Nothing to replace
 			return;
 		}
 
-		for (int32 i = 0; i < PairKnots.Num(); ++i)
-		{
-			if (PairKnots[i].Value == OldValue)
-			{
-				uint32 Count = PairKnots[i].Multiplicity;
-				PairKnots.RemoveAt(i);
-
-				// Insert the new knot with the same multiplicity
-				InsertKnot(NewValue, Count); // Will re-mark cache dirty
-
-				return;
-			}
-		}
+		PairKnots[KnotIdx].Value = NewValue;
 	}
 
 	/**
 	 * Swaps all instances of two knots.
 	 * 
-	 * @param ParameterA - Values which should be replaced with ParameterB
-	 * @param ParameterB - Values which should be replaced with ParameterA
+	 * @param KnotIdxA - Index of first knot to swap.
+	 * @param KnotIdxB - Index of second knot to swap.
 	 */
-	void SwapKnots(float ParameterA, float ParameterB)
+	void SwapKnots(int32 KnotIdxA, int32 KnotIdxB)
 	{
-		for (int32 i = 0; i < PairKnots.Num(); ++i)
+		if (KnotIdxA < 0 || KnotIdxB < 0 || KnotIdxA >= PairKnots.Num() || KnotIdxB >= PairKnots.Num())
 		{
-			if (PairKnots[i].Value == ParameterA)
-			{
-				PairKnots[i].Value = ParameterB;
-			}
-			if (PairKnots[i].Value == ParameterB)
-			{
-				PairKnots[i].Value = ParameterA;
-			}
+			return;
 		}
+		
+		float OldA = PairKnots[KnotIdxA].Value;
+		PairKnots[KnotIdxA].Value = PairKnots[KnotIdxB].Value;
+		PairKnots[KnotIdxB].Value = OldA;
 	}
 
 	/**
@@ -1056,9 +1041,6 @@ protected:
 
         // Get valid parameter space
         FInterval1f ParamSpace = GetParameterSpace();
-        
-        // First ensure the desired parameter is within valid range
-        DesiredParameter = FMath::Clamp(DesiredParameter, ParamSpace.Min, ParamSpace.Max);
 
         // Check if DesiredParameter conflicts with any existing knot (excluding ExcludeValue)
         bool bConflict = false;
diff --git a/Engine/Source/Runtime/GeometryCore/Public/Splines/PolyBezierSpline.h b/Engine/Source/Runtime/GeometryCore/Public/Splines/PolyBezierSpline.h
--- a/Engine/Source/Runtime/GeometryCore/Public/Splines/PolyBezierSpline.h
+++ b/Engine/Source/Runtime/GeometryCore/Public/Splines/PolyBezierSpline.h
@@ -1324,17 +1324,17 @@ public:
 
 	virtual float GetParameter(int32 Index) const override
 	{
-		const TArray<float> Knots = Base::GetDistinctKnotValues();
-		if (Knots.IsEmpty())
-			return 0;
-        int32 NumSegments = Knots.Num() - 1;
+		if (Base::PairKnots.IsEmpty())
+			return 0.f;
+		
+        int32 NumSegments = Base::PairKnots.Num() - 1;
         int32 SegmentIndex = Index / 4;
         int32 LocalPointIndex = Index % 4;
     
         if (SegmentIndex >= NumSegments)
-            return Knots.Last();
+            return Base::PairKnots.Last().Value;
     
-        FInterval1f Range(Knots[SegmentIndex], Knots[SegmentIndex + 1]);
+        FInterval1f Range(Base::PairKnots[SegmentIndex].Value, Base::PairKnots[SegmentIndex + 1].Value);
         const float T = static_cast<float>(LocalPointIndex) / 3.0f;
         return Range.Interpolate(T);
 	}
@@ -1352,7 +1352,8 @@ public:
 		int32 NumSegments = GetNumSegments();
 
 		// Get the original parameter value
-		float OldParameter = GetParameter(Index);
+		int32 OldParameterIndex = LocalPointIndex == 0 ? SegmentIndex : SegmentIndex + 1;
+		float OldParameter = Base::PairKnots[OldParameterIndex].Value;
     
 		// Early exit if parameter isn't changing
 		if (OldParameter == NewParameter)
@@ -1368,7 +1369,7 @@ public:
 			int32 LeftSegmentIndex = SegmentIndex - 1;
 			int32 RightSegmentIndex = SegmentIndex;
 			
-			Base::ReplaceKnot(OldParameter, NewParameter);
+			Base::SetKnot(OldParameterIndex, NewParameter);
 
 			int32 Segment = SegmentIndex;
 			if (NewParameter < OldParameter)
@@ -1399,7 +1400,7 @@ public:
 			int32 LeftSegmentIndex = SegmentIndex;
 			int32 RightSegmentIndex = SegmentIndex + 1;
 			
-			Base::ReplaceKnot(OldParameter, NewParameter);
+			Base::SetKnot(OldParameterIndex, NewParameter);
 
 			int32 Segment = SegmentIndex;
 			if (NewParameter < OldParameter)
@@ -1453,8 +1454,7 @@ public:
 		Base::Values.Swap(P1, P2);
 		
 		// Fix up the knot vector
-		FInterval1f SegmentParamRange = GetSegmentParameterRange(Segment);
-		Base::SwapKnots(SegmentParamRange.Min, SegmentParamRange.Max);
+		Base::SwapKnots(Segment, Segment + 1);
 	}
 	
 	virtual int32 FindIndexForParameter(float Parameter, float& OutLocalParam) const override
