Add Polyline helper to get a point by walking a distance on the line from the first or last vertex

diff --git a/Engine/Source/Runtime/GeometryCore/Public/Polyline.h b/Engine/Source/Runtime/GeometryCore/Public/Polyline.h
--- a/Engine/Source/Runtime/GeometryCore/Public/Polyline.h
+++ b/Engine/Source/Runtime/GeometryCore/Public/Polyline.h
@@ -438,7 +438,55 @@ public:
 		return avg / (T)(N-1);
 	}
 
+	/**
+	 * Compute a point on the polyline based on its distance from the first vertex 
+	 */
+	VectorType GetPointFromFirstVertex(const T InDistance) const
+	{
+		if (InDistance <= 0)
+		{
+			return Vertices[0];
+		}
+
+		T RemainingDistance = InDistance;
+		const int N = Vertices.Num();
+		for (int i = 1; i < N; ++i)
+		{
+			const T SegmentLength = Distance(Vertices[i], Vertices[i - 1]);
+			if (SegmentLength > 0 && SegmentLength > RemainingDistance)
+			{
+				return Lerp(Vertices[i - 1], Vertices[i], RemainingDistance / SegmentLength);
+			}
+			RemainingDistance -= SegmentLength;
+		}
+
+		return Vertices.Last();
+	}
 
+	/**
+	* Compute a point on the polyline based on its distance from the last vertex 
+	*/
+	VectorType GetPointFromLastVertex(const T InDistance) const
+	{
+		if (InDistance <= 0)
+		{
+			return Vertices.Last();
+		}
+
+		T RemainingDistance = InDistance;
+		const int N = Vertices.Num();
+		for (int i = N-1; i > 0; --i)
+		{
+			const T SegmentLength = Distance(Vertices[i], Vertices[i - 1]);
+			if (SegmentLength > 0 && SegmentLength > RemainingDistance)
+			{
+				return Lerp(Vertices[i], Vertices[i-1], RemainingDistance / SegmentLength);
+			}
+			RemainingDistance -= SegmentLength;
+		}
+
+		return Vertices[0];
+	}
 
 	/**
 	 * Produce a new polyline that is smoother than this one
