Fixed averaging of vertices (crack-free correction)

diff --git a/Engine/Source/Runtime/GeometryCore/Public/Tessellation/AdaptiveTessellator.h b/Engine/Source/Runtime/GeometryCore/Public/Tessellation/AdaptiveTessellator.h
--- a/Engine/Source/Runtime/GeometryCore/Public/Tessellation/AdaptiveTessellator.h
+++ b/Engine/Source/Runtime/GeometryCore/Public/Tessellation/AdaptiveTessellator.h
@@ -269,7 +269,7 @@ private:
 					uint32 Hash = static_cast<uint32>(HashPosition( Mesh.GetVertexPosition(VID) ));
 					for( uint32 OtherIndex = HashTable.First( Hash ); HashTable.IsValid( OtherIndex ); OtherIndex = HashTable.Next( OtherIndex ) )
 					{
-						if( Mesh.GetVertexPosition(VID) == Mesh.GetVertexPosition(VID) )
+						if( Mesh.GetVertexPosition(VID) == Mesh.GetVertexPosition(OtherIndex) )
 						{
 							Average += OldDisplacements[ OtherIndex ];
 							Count++;
