Add missing typename (Android compile error)

diff --git a/Engine/Source/Runtime/GeometryCore/Public/Tessellation/AdaptiveTessellator.h b/Engine/Source/Runtime/GeometryCore/Public/Tessellation/AdaptiveTessellator.h
--- a/Engine/Source/Runtime/GeometryCore/Public/Tessellation/AdaptiveTessellator.h
+++ b/Engine/Source/Runtime/GeometryCore/Public/Tessellation/AdaptiveTessellator.h
@@ -48,11 +48,11 @@ struct FFlipEdgeInfo
 };
 
 
-// make edge-flip in nanite code produce the same ordering as DM3, so that results can be compared
+// make edge-flip in nanite code produce the same ordering as FDynamicMesh3, so that results can be compared
 #define UE_DM3_NANITE_COMPATIBILITY 
 
-// Configured by MeshT and DisplacementPolicyT, which provide mesh topology funtionality and
-// evalution/bounds of the displacement function.
+// Configured by MeshT and DisplacementPolicyT, which provide mesh topology functionality and
+// evaluation/bounds of the displacement function.
 // 
 // See MinimalMeshTopology for an example topology implementation.
 // 
@@ -155,10 +155,10 @@ private:
 
   public:
 
-	using RealType = MeshT::RealType; // the field related to the vector-space
-	using VecType  = MeshT::VecType;  // for positions and normals
+	using RealType = typename MeshT::RealType; // the field related to the vector-space
+	using VecType  = typename MeshT::VecType;  // for positions and normals
 
-    // Initialize and run adaptive aptive tesselation
+    // Initialize and run adaptive tessellation
 	FAdaptiveTessellator(
 		MeshT& InMesh,                                //< mesh interface
 		DisplacementPolicyT& InDisplacementPolicy,    //< displacement and error interface
@@ -166,7 +166,7 @@ private:
 		RealType	InSampleRate,                     //< one-dimensional rate (edge-length) at which features are sampled. triangles with edge-length smaller than
 		                                              //< this value will not be refined
 		bool		bCrackFree,                       //< vertices with same positions are considered to be identical and stay at the same position after displacement
-		bool        bFinalDisplace)                   //< apply displacement to the tesselated mesh
+		bool        bFinalDisplace)                   //< apply displacement to the tessellated mesh
 		: Mesh(InMesh)
 		, DisplacementPolicy(InDisplacementPolicy)
 		, TargetError(InTargetError)
