FSpline: Make attribute channels use Constant out of bounds handling by default.

This also makes the out of bounds handling mode accessors virtual so that the spline wrapper class can propagate this information to the internal spline. This actually might not matter, when we evaluate a spline wrapper, we call that spline wrapper's out of bounds handling functions which respects the spline wrappers mode. But, if someone fetches the internal spline and evaluates it directly, we want it to be consistent with what the wrapped spline is doing.

diff --git a/Engine/Source/Runtime/Engine/Private/Curves/Spline.cpp b/Engine/Source/Runtime/Engine/Private/Curves/Spline.cpp
--- a/Engine/Source/Runtime/Engine/Private/Curves/Spline.cpp
+++ b/Engine/Source/Runtime/Engine/Private/Curves/Spline.cpp
@@ -1189,7 +1189,12 @@ bool FSpline::CreateAttributeChannel(FName AttributeName) const
 {
 	if (NewData && SupportsAttributes())
 	{
-		return NewData->CreateAttributeChannel<TTangentBezierSpline<AttrType>>(AttributeName) != nullptr;
+		TTangentBezierSpline<AttrType>* NewAttrChannel = NewData->CreateAttributeChannel<TTangentBezierSpline<AttrType>>(AttributeName);
+		
+		NewAttrChannel->SetPreInfinityMode(EOutOfBoundsHandlingMode::Constant);
+		NewAttrChannel->SetPostInfinityMode(EOutOfBoundsHandlingMode::Constant);
+		
+		return NewAttrChannel != nullptr;
 	}
 	
 	return false;
diff --git a/Engine/Source/Runtime/GeometryCore/Public/Splines/Spline.h b/Engine/Source/Runtime/GeometryCore/Public/Splines/Spline.h
--- a/Engine/Source/Runtime/GeometryCore/Public/Splines/Spline.h
+++ b/Engine/Source/Runtime/GeometryCore/Public/Splines/Spline.h
@@ -44,9 +44,26 @@ public:
     virtual FInterval1f GetParameterSpace() const override { return InternalSpline.GetParameterSpace(); }
     virtual void SetClosedLoop(bool bClosed) override { InternalSpline.SetClosedLoop(bClosed); }
     virtual bool IsClosedLoop() const override { return InternalSpline.IsClosedLoop(); }
-    //virtual bool IsEqual(const ISplineInterface& Other) const override { return InternalSpline.IsEqual(Other); }
 
     // TSplineInterface Forwarding
+
+	// TSplineInterface<ValueType>::PreInfinityMode is the one that actually matters, but we propagate to InternalSpline for clarity
+	virtual void SetPreInfinityMode(EOutOfBoundsHandlingMode InMode) override
+    {
+    	InternalSpline.SetPreInfinityMode(InMode);
+	    TSplineInterface<ValueType>::PreInfinityMode = InMode;
+    }
+
+	// TSplineInterface<ValueType>::PostInfinityMode is the one that actually matters, but we propagate to InternalSpline for clarity
+	virtual void SetPostInfinityMode(EOutOfBoundsHandlingMode InMode) override
+    {
+    	InternalSpline.SetPostInfinityMode(InMode);
+	    TSplineInterface<ValueType>::PostInfinityMode = InMode;
+    }
+	
+	virtual EOutOfBoundsHandlingMode GetPreInfinityMode() const override { return InternalSpline.GetPreInfinityMode(); }
+	virtual EOutOfBoundsHandlingMode GetPostInfinityMode() const override { return InternalSpline.GetPostInfinityMode(); }
+	
     virtual ValueType EvaluateImpl(float Parameter) const override { return InternalSpline.EvaluateImpl(Parameter); }
     virtual float FindNearest(const ValueType& Point, float& OutSquaredDistance) const override
         { return InternalSpline.FindNearest(Point, OutSquaredDistance); }
diff --git a/Engine/Source/Runtime/GeometryCore/Public/Splines/SplineInterfaces.h b/Engine/Source/Runtime/GeometryCore/Public/Splines/SplineInterfaces.h
--- a/Engine/Source/Runtime/GeometryCore/Public/Splines/SplineInterfaces.h
+++ b/Engine/Source/Runtime/GeometryCore/Public/Splines/SplineInterfaces.h
@@ -171,16 +171,16 @@ public:
     virtual FName GetValueTypeName() const override { return TypeName; }
 
     /** Set Pre-Infinity handling mode when out of bounds */ 
-    void SetPreInfinityMode(EOutOfBoundsHandlingMode InMode){ PreInfinityMode = InMode; }
+    virtual void SetPreInfinityMode(EOutOfBoundsHandlingMode InMode){ PreInfinityMode = InMode; }
 
     /** Set Post-Infinity handling mode when out of bounds */
-    void SetPostInfinityMode(EOutOfBoundsHandlingMode InMode){ PostInfinityMode = InMode; }
+    virtual void SetPostInfinityMode(EOutOfBoundsHandlingMode InMode){ PostInfinityMode = InMode; }
 
-    /** Get the current post-infinity mode */
-    EOutOfBoundsHandlingMode GetPreInfinityMode() const{ return PreInfinityMode; }
+    /** Get the current Pre-Infinity mode */
+    virtual EOutOfBoundsHandlingMode GetPreInfinityMode() const{ return PreInfinityMode; }
 
-    /** Get the current post-infinity mode */
-    EOutOfBoundsHandlingMode GetPostInfinityMode() const{ return PostInfinityMode; }
+    /** Get the current Post-Infinity mode */
+    virtual EOutOfBoundsHandlingMode GetPostInfinityMode() const{ return PostInfinityMode; }
 
 protected:
     /** How to handle parameters outside the mapped range */
