SplineCore: Fix issues with Values array when flipping poly bezier segments.

diff --git a/Engine/Source/Runtime/GeometryCore/Public/Splines/BSpline.h b/Engine/Source/Runtime/GeometryCore/Public/Splines/BSpline.h
--- a/Engine/Source/Runtime/GeometryCore/Public/Splines/BSpline.h
+++ b/Engine/Source/Runtime/GeometryCore/Public/Splines/BSpline.h
@@ -84,6 +84,25 @@ public:
 	TBSpline() = default;
 	virtual ~TBSpline() override = default;
 
+	void Dump() const
+	{
+		if constexpr (std::is_same_v<ValueType, float>)
+		{
+			UE_LOG(LogSpline, Verbose, TEXT("Values:"));
+			for (int32 Index = 0; Index < Values.Num(); ++Index)
+			{
+				UE_LOG(LogSpline, Verbose, TEXT("  [%d] %f"), Index, Values[Index]);
+			}
+		}
+
+		UE_LOG(LogSpline, Verbose, TEXT("Knots:"));
+		for (int32 Index = 0; Index < PairKnots.Num(); ++Index)
+		{
+			const FKnot& Knot = PairKnots[Index];
+			UE_LOG(LogSpline, Verbose, TEXT("  [%d] Value: %f, Multiplicity: %u"), Index, Knot.Value, Knot.Multiplicity);
+		}
+	}
+	
 	// ISplineInterface Implementation
 	virtual void Clear() override
 	{
diff --git a/Engine/Source/Runtime/GeometryCore/Public/Splines/PolyBezierSpline.h b/Engine/Source/Runtime/GeometryCore/Public/Splines/PolyBezierSpline.h
--- a/Engine/Source/Runtime/GeometryCore/Public/Splines/PolyBezierSpline.h
+++ b/Engine/Source/Runtime/GeometryCore/Public/Splines/PolyBezierSpline.h
@@ -1361,8 +1361,13 @@ public:
 			return SegmentIndex * PointsPerSegment + LocalPointIndex;
 		}
 		
-		// Prevent segment degeneration.
+		UE_LOG(LogSpline, Verbose, TEXT("\t"));
+		UE_LOG(LogSpline, Verbose, TEXT("Before SetParameter(%d, %f):"), Index, NewParameter);
+		Base::Dump();
+		
+		// Prevent reuse of an existing knot.
 		NewParameter = this->GetNearestValidKnot(NewParameter, OldParameter);
+		int32 ReturnIndex = INDEX_NONE;
 		
 		if (LocalPointIndex == 0) // P0
 		{
@@ -1393,7 +1398,7 @@ public:
 				}
 			}
 			
-			return Segment * PointsPerSegment + LocalPointIndex;
+			ReturnIndex = Segment * PointsPerSegment + LocalPointIndex;
 		}
 		else if (LocalPointIndex == 3) // P3
 		{
@@ -1424,17 +1429,20 @@ public:
 				}
 			}
 			
-			return Segment * PointsPerSegment + LocalPointIndex;
+			ReturnIndex = Segment * PointsPerSegment + LocalPointIndex;
 		}
 
-		return INDEX_NONE;
+		UE_LOG(LogSpline, Verbose, TEXT("\t"));
+		UE_LOG(LogSpline, Verbose, TEXT("After SetParameter(%d, %f):"), Index, NewParameter);
+		Base::Dump();
+		
+		return ReturnIndex;
 	}
 	
 	void FlipSegment(int32 Segment)
 	{
 		static constexpr int32 Degree = 3; // Cubic
 		static constexpr int32 ControlPointsPerSegment = 4;
-		static constexpr int32 KnotsPerSegment = Degree; // Not degree+1 because knot overlap between segments
 		
 		const int32 NumSegments = GetNumSegments();
 		if (Segment < 0 || Segment >= NumSegments)
@@ -1442,19 +1450,55 @@ public:
 			return;
 		}
 
-		// Swap P0 && P3
-		// Swap P1 && P2
+		UE_LOG(LogSpline, Verbose, TEXT("\t"));
+		UE_LOG(LogSpline, Verbose, TEXT("Before FlipSegment(%d):"), Segment);
+		Base::Dump();
 		
-		int32 P0 = Segment * ControlPointsPerSegment + 0;
-		int32 P1 = Segment * ControlPointsPerSegment + 1;
-		int32 P2 = Segment * ControlPointsPerSegment + 2;
-		int32 P3 = Segment * ControlPointsPerSegment + 3;
+		// Flip the segment
+		{
+			int32 P0 = Segment * ControlPointsPerSegment + 0;
+			int32 P1 = Segment * ControlPointsPerSegment + 1;
+			int32 P2 = Segment * ControlPointsPerSegment + 2;
+			int32 P3 = Segment * ControlPointsPerSegment + 3;
 
-		Base::Values.Swap(P0, P3);
-		Base::Values.Swap(P1, P2);
+			Base::Values.Swap(P0, P3);
+			Base::Values.Swap(P1, P2);
+		}
+
+		// Fix up neighbor segment by keeping P3 and P0 coincident and preserving the out-tangent of the end of the neighbor segment
+		if (Segment != 0)
+		{
+			// Set P3 of Segment - 1 to P0 of Segment
+			// Need to preserve P2 relative to P3 on Segment - 1
+			int32 P2 = (Segment - 1) * ControlPointsPerSegment + 2;
+			int32 P3 = (Segment - 1) * ControlPointsPerSegment + 3;
+			int32 P0 = Segment * ControlPointsPerSegment + 0;
+
+			ValueType Delta = Base::Values[P3] - Base::Values[P2];
+			Base::Values[P3] = Base::Values[P0];
+			Base::Values[P2] = Base::Values[P3] - Delta;
+		}
+
+		// Fix up neighbor segment by keeping P0 and P3 coincident and preserving the in-tangent of the beginning of the neighbor segment
+		if (Segment != NumSegments - 1)
+		{
+			// Set P0 of Segment + 1 to P3 of Segment
+			// Need to preserve P1 relative to P0 on Segment + 1
+			int32 P0 = (Segment + 1) * ControlPointsPerSegment + 0;
+			int32 P1 = (Segment + 1) * ControlPointsPerSegment + 1;
+			int32 P3 = Segment * ControlPointsPerSegment + 3;
+
+			ValueType Delta = Base::Values[P0] - Base::Values[P1];
+			Base::Values[P0] = Base::Values[P3];
+			Base::Values[P1] = Base::Values[P0] - Delta;
+		}
 		
 		// Fix up the knot vector
 		Base::SwapKnots(Segment, Segment + 1);
+
+		UE_LOG(LogSpline, Verbose, TEXT("\t"));
+		UE_LOG(LogSpline, Verbose, TEXT("After FlipSegment(%d):"), Segment);
+		Base::Dump();
 	}
 	
 	virtual int32 FindIndexForParameter(float Parameter, float& OutLocalParam) const override
diff --git a/Engine/Source/Runtime/GeometryCore/Public/Splines/TangentBezierSpline.h b/Engine/Source/Runtime/GeometryCore/Public/Splines/TangentBezierSpline.h
--- a/Engine/Source/Runtime/GeometryCore/Public/Splines/TangentBezierSpline.h
+++ b/Engine/Source/Runtime/GeometryCore/Public/Splines/TangentBezierSpline.h
@@ -1357,7 +1357,19 @@ public:
 	int32 SetParameter(int32 Index, float NewParameter)
     {
 		const int32 ControlPointIndex = (Index == 0) ? 0 : ((Index - 1) * 4 + 3);
-		return (InternalSpline.SetParameter(ControlPointIndex, NewParameter) + 1) / 4;
+		const int32 NewIndex = (InternalSpline.SetParameter(ControlPointIndex, NewParameter) + 1) / 4;
+
+		// We likely need to do better in PolyBezier's SetParameter to account for user tangents, but this works for now.
+		if (NewIndex != Index)
+		{
+			UpdateTangents();
+
+			UE_LOG(LogSpline, Verbose, TEXT("\t"));
+			UE_LOG(LogSpline, Verbose, TEXT("After Tangent Recompute:"));
+			InternalSpline.Dump();
+		}
+
+		return NewIndex;
     }
 
 	virtual FInterval1f GetParameterSpace() const override
