fix crash in implicit morphology offset operation due to offset grid dimensions overflowing int32

diff --git a/Engine/Source/Runtime/GeometryCore/Public/Implicit/Morphology.h b/Engine/Source/Runtime/GeometryCore/Public/Implicit/Morphology.h
--- a/Engine/Source/Runtime/GeometryCore/Public/Implicit/Morphology.h
+++ b/Engine/Source/Runtime/GeometryCore/Public/Implicit/Morphology.h
@@ -73,13 +73,15 @@ public:
 	// Set input grid cell size to hit the target voxel counts along the max dimension of the bounds
 	void SetGridCellSize(FAxisAlignedBox3d Bounds, double DistanceIn, int TargetInputVoxelCount)
 	{
-		GridCellSize = (Bounds.MaxDim() + (!bUseCustomBounds ? DistanceIn * 2.0 : 0.0)) / double(TargetInputVoxelCount);
+		int UseTargetVoxelCount = FMath::Min(MaxTargetVoxelCount, TargetInputVoxelCount);
+		GridCellSize = (Bounds.MaxDim() + (!bUseCustomBounds ? DistanceIn * 2.0 : 0.0)) / double(UseTargetVoxelCount);
 	}
 
 	// Set output meshing cell size to hit the target voxel counts along the max dimension of the bounds
 	void SetMeshCellSize(FAxisAlignedBox3d Bounds, double DistanceIn, int TargetInputVoxelCount)
 	{
-		MeshCellSize = (Bounds.MaxDim() + (!bUseCustomBounds ? DistanceIn * 2.0 : 0.0)) / double(TargetInputVoxelCount);
+		int UseTargetVoxelCount = FMath::Min(MaxTargetVoxelCount, TargetInputVoxelCount);
+		MeshCellSize = (Bounds.MaxDim() + (!bUseCustomBounds ? DistanceIn * 2.0 : 0.0)) / double(UseTargetVoxelCount);
 	}
 
 	/** if this function returns true, we should abort calculation */
@@ -149,8 +151,8 @@ protected:
 		ComputedSDF.Spatial = SourceSpatial;
 		ComputedSDF.ComputeMode = MeshSDFType::EComputeModes::NarrowBand_SpatialFloodFill;
 
-
-		ComputedSDF.CellSize = (float)GridCellSize;
+		double UseGridCellSize = GetSafeCellSize(2*UnsignedOffset + SourceSpatial->GetBoundingBox().MaxDim(), GridCellSize, 2);
+		ComputedSDF.CellSize = (float)UseGridCellSize;
 		NarrowBandMaxDistance = UnsignedOffset + ComputedSDF.CellSize;
 		ComputedSDF.NarrowBandMaxDistance = NarrowBandMaxDistance;
 		ComputedSDF.ExactBandWidth = FMath::CeilToInt32(ComputedSDF.NarrowBandMaxDistance / ComputedSDF.CellSize);
@@ -185,7 +187,7 @@ protected:
 			}
 		}
 		MarchingCubes.RootMode = ERootfindingModes::SingleLerp;
-		MarchingCubes.CubeSize = MeshCellSize;
+		MarchingCubes.CubeSize = GetSafeCellSize(MarchingCubes.Bounds.MaxDim(), MeshCellSize, 1);
 
 		MarchingCubes.CancelF = CancelF;
 
@@ -225,17 +227,22 @@ protected:
 		TIndexVectorMeshArrayAdapter<FIndex3i, double, FVector3d> MCAdapter(&MarchingCubes.Vertices, &MarchingCubes.Triangles);
 		TMeshAABBTree3<TIndexVectorMeshArrayAdapter<FIndex3i, double, FVector3d>> SecondSpatial(&MCAdapter, false);
 
-		MeshSDFType SecondSDF;
-		SecondSDF.Mesh = &MCAdapter;
-		SecondSDF.CellSize = (float)GridCellSize;
-		SecondSDF.Spatial = nullptr;
-
 		FAxisAlignedBox3d Bounds = MarchingCubes.Bounds;
 		if (!bUseCustomBounds)
 		{
 			Bounds.Expand(MeshCellSize); // (because mesh may spill one cell over bounds)
 		}
 
+		MeshSDFType SecondSDF;
+		SecondSDF.Mesh = &MCAdapter;
+
+		// Adjust cell size to not overflow w/ the added UnsignedOffset
+		double UseGridCellSize = GetSafeCellSize(2*UnsignedOffset + Bounds.MaxDim(), GridCellSize, 2);
+
+		SecondSDF.CellSize = (float)UseGridCellSize;
+		SecondSDF.Spatial = nullptr;
+
+
 		SecondSDF.NarrowBandMaxDistance = UnsignedOffset + SecondSDF.CellSize;
 		SecondSDF.ExactBandWidth = FMath::CeilToInt32(SecondSDF.NarrowBandMaxDistance / SecondSDF.CellSize);
 
@@ -272,11 +279,13 @@ protected:
 		MarchingCubes.Bounds = Bounds;
 		if (!bUseCustomBounds)
 		{
-			MarchingCubes.Bounds.Expand(GridCellSize);
+			MarchingCubes.Bounds.Expand(UseGridCellSize);
 			if (MarchingCubes.IsoValue < 0)
 			{
 				MarchingCubes.Bounds.Expand(NarrowBandMaxDistance);
 			}
+			// Make sure the CubeSize is still safe after expanding the bounds
+			MarchingCubes.CubeSize = GetSafeCellSize(MarchingCubes.Bounds.MaxDim(), MarchingCubes.CubeSize, 1);
 		}
 
 
@@ -300,6 +309,25 @@ protected:
 		/// ..... then remove this line
 		MarchingCubes.Implicit = nullptr;
 	}
+
+private:
+	// Set a max target voxel count s.t. the VoxelCount^3 does not overflow int32 linearized grid indices
+	// (and to reduce the chance of running out of memory for a very dense grid, generally)
+	static constexpr int32 MaxTargetVoxelCount = 1200;
+
+	// Adjust cell size so that a cell count based on (BoundsWidth/InitialCellSize + ExtraCellCount) should not
+	// (too far) exceed MaxTargetVoxelCount.  (Assumes ExtraCellCount is smaller MaxTargetVoxelCount)
+	static double GetSafeCellSize(double BoundsWidth, double InitialCellSize, int32 ExtraCellCount)
+	{
+		if (BoundsWidth + (double)ExtraCellCount * InitialCellSize > InitialCellSize * (double)MaxTargetVoxelCount)
+		{
+			return BoundsWidth / (double)(MaxTargetVoxelCount - ExtraCellCount);
+		}
+		else
+		{
+			return InitialCellSize;
+		}
+	}
 };
 
 
