Fix Android build fallout from CL  42399822

diff --git a/Engine/Source/Runtime/GeometryCore/Public/Solvers/MatrixInterfaces.h b/Engine/Source/Runtime/GeometryCore/Public/Solvers/MatrixInterfaces.h
--- a/Engine/Source/Runtime/GeometryCore/Public/Solvers/MatrixInterfaces.h
+++ b/Engine/Source/Runtime/GeometryCore/Public/Solvers/MatrixInterfaces.h
@@ -28,6 +28,11 @@ namespace UE
 
 			struct FTupleData
 			{
+				FTupleData(int32 IIn, int32 JIn, RealType ValueIn)
+					: I(IIn)
+					, J(JIn)
+					, Value(ValueIn)
+				{}
 				int32 I;
 				int32 J;
 				RealType Value;
