GeometryProcessing: Make newly added TSparseMatrixAssembler::AddEntriesFunc route to AddEntryFunc by default so that subclasses still work with code that now uses AddEntriesFunc.

diff --git a/Engine/Source/Runtime/GeometryCore/Public/Solvers/MatrixInterfaces.h b/Engine/Source/Runtime/GeometryCore/Public/Solvers/MatrixInterfaces.h
--- a/Engine/Source/Runtime/GeometryCore/Public/Solvers/MatrixInterfaces.h
+++ b/Engine/Source/Runtime/GeometryCore/Public/Solvers/MatrixInterfaces.h
@@ -39,7 +39,18 @@ namespace UE
 			};
 
 			/** Add matrix entries (i,j,Value) and size of tuples */
-			TUniqueFunction<void(const FTupleData*, int32)> AddEntriesFunc;
+			TUniqueFunction<void(const FTupleData*, int32)> AddEntriesFunc = [this](const FTupleData* DataStart, int32 Num)
+			{
+				if (!AddEntryFunc)
+				{
+					return;
+				}
+				for (int32 i = 0; i < Num; ++i)
+				{
+					const FTupleData& Tuple = DataStart[Num];
+					AddEntryFunc(Tuple.I, Tuple.J, Tuple.Value);
+				}
+			};
 		};
 
 		typedef TSparseMatrixAssembler<double> FSparseMatrixAssemblerd;
