Promote SameAsOptions::Epsilon to double precision. Change epsilon comparison to use <= instead of <.

diff --git a/Engine/Plugins/Runtime/GeometryScripting/Source/GeometryScriptingCore/Private/MeshComparisonFunctions.cpp b/Engine/Plugins/Runtime/GeometryScripting/Source/GeometryScriptingCore/Private/MeshComparisonFunctions.cpp
--- a/Engine/Plugins/Runtime/GeometryScripting/Source/GeometryScriptingCore/Private/MeshComparisonFunctions.cpp
+++ b/Engine/Plugins/Runtime/GeometryScripting/Source/GeometryScriptingCore/Private/MeshComparisonFunctions.cpp
@@ -51,7 +51,8 @@ UDynamicMesh* UGeometryScriptLibrary_MeshComparisonFunctions::IsSameMeshAs(
 			CompareOptions.bCheckUVs = Options.bCheckUVs;
 			CompareOptions.bCheckGroups = Options.bCheckGroups;
 			CompareOptions.bCheckAttributes = Options.bCheckAttributes;
-			CompareOptions.Epsilon = Options.Epsilon;
+			CompareOptions.EpsilonFloat = Options.Epsilon;
+			CompareOptions.EpsilonDouble = static_cast<double>(Options.Epsilon);
 			bIsSameMesh = ReadMesh.IsSameAs(OtherMesh, CompareOptions, Difference);
 			auto ConvertIDType = [](FDynamicMesh3::FMeshDifferenceInfo::EIDType IDType)
 			{
diff --git a/Engine/Source/Runtime/GeometryCore/Private/DynamicMesh/DynamicMesh3.cpp b/Engine/Source/Runtime/GeometryCore/Private/DynamicMesh/DynamicMesh3.cpp
--- a/Engine/Source/Runtime/GeometryCore/Private/DynamicMesh/DynamicMesh3.cpp
+++ b/Engine/Source/Runtime/GeometryCore/Private/DynamicMesh/DynamicMesh3.cpp
@@ -835,7 +835,7 @@ bool FDynamicMesh3::IsSameAs_Helper(const FDynamicMesh3& m2, const FSameAsOption
 
 	auto SameVertex = [this, &m2, &Options](const int Vid, const int VidM2)
 	{
-		return VectorUtil::EpsilonEqual(GetVertexRef(Vid), m2.GetVertexRef(VidM2), static_cast<double>(Options.Epsilon));
+		return VectorUtil::EpsilonEqual(GetVertexRef(Vid), m2.GetVertexRef(VidM2), Options.EpsilonDouble);
 	};
 
 	if (VertexCount() != m2.VertexCount())
@@ -1079,7 +1079,7 @@ bool FDynamicMesh3::IsSameAs_Helper(const FDynamicMesh3& m2, const FSameAsOption
 		{
 			for (int vid : VertexIndicesItr())
 			{
-				if (VectorUtil::EpsilonEqual(GetVertexNormal(vid), m2.GetVertexNormal(VidMapping ? (*VidMapping)[vid] : vid), Options.Epsilon) == false)
+				if (VectorUtil::EpsilonEqual(GetVertexNormal(vid), m2.GetVertexNormal(VidMapping ? (*VidMapping)[vid] : vid), Options.EpsilonFloat) == false)
 				{
 					if (OutMeshDifferenceInfo)
 					{
@@ -1106,7 +1106,7 @@ bool FDynamicMesh3::IsSameAs_Helper(const FDynamicMesh3& m2, const FSameAsOption
 		{
 			for (int vid : VertexIndicesItr())
 			{
-				if (VectorUtil::EpsilonEqual(GetVertexColor(vid), m2.GetVertexColor(VidMapping ? (*VidMapping)[vid] : vid), Options.Epsilon) == false)
+				if (VectorUtil::EpsilonEqual(GetVertexColor(vid), m2.GetVertexColor(VidMapping ? (*VidMapping)[vid] : vid), Options.EpsilonFloat) == false)
 				{
 					if (OutMeshDifferenceInfo)
 					{
@@ -1133,7 +1133,7 @@ bool FDynamicMesh3::IsSameAs_Helper(const FDynamicMesh3& m2, const FSameAsOption
 		{
 			for (int vid : VertexIndicesItr())
 			{
-				if (VectorUtil::EpsilonEqual(GetVertexUV(vid), m2.GetVertexUV(VidMapping ? (*VidMapping)[vid] : vid), Options.Epsilon) == false)
+				if (VectorUtil::EpsilonEqual(GetVertexUV(vid), m2.GetVertexUV(VidMapping ? (*VidMapping)[vid] : vid), Options.EpsilonFloat) == false)
 				{
 					if (OutMeshDifferenceInfo)
 					{
diff --git a/Engine/Source/Runtime/GeometryCore/Public/DynamicMesh/DynamicMesh3.h b/Engine/Source/Runtime/GeometryCore/Public/DynamicMesh/DynamicMesh3.h
--- a/Engine/Source/Runtime/GeometryCore/Public/DynamicMesh/DynamicMesh3.h
+++ b/Engine/Source/Runtime/GeometryCore/Public/DynamicMesh/DynamicMesh3.h
@@ -1521,7 +1521,11 @@ public:
 		// Ignore gaps and padding in the data layout, i.e. the exact index is irrelevant as long as valid entries are equal in value and in the same order.
 		bool bIgnoreDataLayout = false;
 
-		float Epsilon = TMathUtil<float>::Epsilon;
+		// comparisons check for smaller or equal than this value. used on on 32-bit floating point checks
+		float EpsilonFloat { TMathUtil<float>::Epsilon };
+
+		// comparisons check for smaller or equal than this value. used for 64-bit comparisons, such as on points.
+		double EpsilonDouble{ static_cast<double>(TMathUtil<float>::Epsilon) };
 	};
 
 	struct FMeshDifferenceInfo
diff --git a/Engine/Source/Runtime/GeometryCore/Public/VectorUtil.h b/Engine/Source/Runtime/GeometryCore/Public/VectorUtil.h
--- a/Engine/Source/Runtime/GeometryCore/Public/VectorUtil.h
+++ b/Engine/Source/Runtime/GeometryCore/Public/VectorUtil.h
@@ -75,7 +75,6 @@ namespace VectorUtil
 		Normalize(edge2);
 		// Unreal has Left-Hand Coordinate System so we need to reverse this cross-product to get proper triangle normal
 		TVector<RealType> vCross(edge2.Cross(edge1));
-		//TVector<RealType> vCross(edge1.Cross(edge2));
 		return Normalized(vCross);
 	}
 
@@ -87,7 +86,6 @@ namespace VectorUtil
 	{
 		// Unreal has Left-Hand Coordinate System so we need to reverse this cross-product to get proper triangle normal
 		return (V2 - V0).Cross(V1 - V0);
-		//return (V1 - V0).Cross(V2 - V0);
 	}
 
 	/**
@@ -146,26 +144,25 @@ namespace VectorUtil
 		TVector<RealType> edge2(V2 - V0);
 		// Unreal has Left-Hand Coordinate System so we need to reverse this cross-product to get proper triangle normal
 		FVector3d vCross = edge2.Cross(edge1);
-		//FVector3d vCross = edge1.Cross(edge2);
 		AreaOut = RealType(0.5) * Normalize(vCross);
 		return vCross;
 	}
 
-	/** @return true if Abs(A-B) is less than Epsilon */
+	/** @return true if Abs(A-B) is less or equal than Epsilon */
 	template <typename RealType>
 	inline bool EpsilonEqual(RealType A, RealType B, RealType Epsilon)
 	{
-		return TMathUtil<RealType>::Abs(A - B) < Epsilon;
+		return TMathUtil<RealType>::Abs(A - B) <= Epsilon;
 	}
 
-	/** @return true if all coordinates of V0 and V1 are within Epsilon of eachother */
+	/** @return true if all coordinates of V0 and V1 are within Epsilon of each other */
 	template <typename RealType>
 	inline bool EpsilonEqual(const TVector2<RealType>& V0, const TVector2<RealType>& V1, RealType Epsilon)
 	{
 		return EpsilonEqual(V0.X, V1.X, Epsilon) && EpsilonEqual(V0.Y, V1.Y, Epsilon);
 	}
 
-	/** @return true if all coordinates of V0 and V1 are within Epsilon of eachother */
+	/** @return true if all coordinates of V0 and V1 are within Epsilon of each other */
 	template <typename RealType>
 	inline bool EpsilonEqual(const TVector<RealType>& V0, const TVector<RealType>& V1, RealType Epsilon)
 	{
