SplineCore: Fix FQuat basis interpolation

The Result local should be initialized to 0, not the identity quat. By initializing to FQuat::Identity, we get a result quat that has a much greater real part than expected. Normalizing the quaternion results in the rotation becoming incorrect.

diff --git a/Engine/Source/Runtime/GeometryCore/Public/Splines/InterpolationPolicies.h b/Engine/Source/Runtime/GeometryCore/Public/Splines/InterpolationPolicies.h
--- a/Engine/Source/Runtime/GeometryCore/Public/Splines/InterpolationPolicies.h
+++ b/Engine/Source/Runtime/GeometryCore/Public/Splines/InterpolationPolicies.h
@@ -205,7 +205,7 @@ public:
 	// Window-based interpolation for spline implementations
 	static FQuat InterpolateWithBasis(TArrayView<const FQuat* const> Window, TArrayView<const float> Basis)
 	{
-		FQuat Result = FQuat::Identity;
+		FQuat Result(0.f, 0.f, 0.f, 0.f);
 		float TotalWeight = 0.0f;
 
 		for (int32 i = 0; i < FMath::Min(Window.Num(), Basis.Num()); ++i)
