Spline: Fix reparameterization of final segment.

This bug is due to the fact that we are incorrectly computing the length of the final segment because it relies on the out tangent of the final point. This makes it so that all points have in and out tangents. for start and end points, out and in tangents must be the same (as we have no way of uniquely representing the in tangent for first point or the out tangent for the last point).

The additional dirtying of the reparam table is not necessary for BP splines because we are always going to dirty on reconstruct, but it would result in similar issues when using native actor splines so I have fixed in this change as well.

diff --git a/Engine/Source/Runtime/Engine/Private/Curves/Spline.cpp b/Engine/Source/Runtime/Engine/Private/Curves/Spline.cpp
--- a/Engine/Source/Runtime/Engine/Private/Curves/Spline.cpp
+++ b/Engine/Source/Runtime/Engine/Private/Curves/Spline.cpp
@@ -2077,6 +2077,8 @@ void FNewSpline::SetLocation(int32 Index, const FVector& InLocation)
     }
 	
 	GetSpline().ModifyPoint(Index, {InLocation, GetSpline().GetTangentIn(Index), GetSpline().GetTangentOut(Index), GetSpline().GetTangentMode(Index)});
+
+	MarkReparamTableDirty();
     RebuildLegacyCurves();
 }
 
@@ -2099,6 +2101,8 @@ void FNewSpline::SetInTangent(const int32 Index, const FVector& InTangent)
 
 	GetSpline().SetPointTangentMode(Index, ETangentMode::User);
     GetSpline().SetTangentIn(Index, InTangent);
+
+	MarkReparamTableDirty();
     RebuildLegacyCurves();
 }
 
@@ -2121,6 +2125,8 @@ void FNewSpline::SetOutTangent(const int32 Index, const FVector& OutTangent)
 	
 	GetSpline().SetPointTangentMode(Index, ETangentMode::User);
 	GetSpline().SetTangentOut(Index, OutTangent);
+
+	MarkReparamTableDirty();
     RebuildLegacyCurves();
 }
 
@@ -2190,7 +2196,8 @@ void FNewSpline::SetSplinePointType(int32 Index, EInterpCurveMode Type)
     }
 	
 	GetSpline().SetPointTangentMode(Index, ConvertInterpCurveModeToTangentMode(Type));
-    
+
+	MarkReparamTableDirty();
     RebuildLegacyCurves();
 }
 
diff --git a/Engine/Source/Runtime/GeometryCore/Public/Splines/TangentBezierSpline.h b/Engine/Source/Runtime/GeometryCore/Public/Splines/TangentBezierSpline.h
--- a/Engine/Source/Runtime/GeometryCore/Public/Splines/TangentBezierSpline.h
+++ b/Engine/Source/Runtime/GeometryCore/Public/Splines/TangentBezierSpline.h
@@ -1199,18 +1199,25 @@ public:
     /** Gets incoming tangent for the specified point */
     ValueType GetTangentIn(int32 Index) const
     {
+		const int32 NumPoints = GetNumPoints();
+
+		if (NumPoints <= 1)
+		{
+			return ValueType();
+		}
+		
         if (Index == 0)
         {
             // Special case for first point
-            if (InternalSpline.IsClosedLoop() && InternalSpline.GetNumSegments() > 0)
+            if (IsClosedLoop())
             {
-                // For closed loop, look at last segment
-                const int32 LastSegment = InternalSpline.GetNumSegments() - 1;
-                ValueType P3 = InternalSpline.GetValue(LastSegment * 4 + 3);
-                ValueType P2 = InternalSpline.GetValue(LastSegment * 4 + 2);
-                return (P3 - P2) * 3.0f;
+                return GetTangentIn(NumPoints - 1);
+            }
+            else
+            {
+            	// If spline is unclosed, in tangent == out tangent at first point.
+				return GetTangentOut(Index);
             }
-            return ValueType();
         }
         
         // Normal case - get from previous segment's control points
@@ -1224,25 +1231,39 @@ public:
         return ValueType();
     }
     
-    /** Gets outgoing tangent for the specified point */
-    ValueType GetTangentOut(int32 Index) const
-    {
-        const int32 NumPoints = GetNumPoints();
-        if (Index == NumPoints - 1 && !IsClosedLoop())
+	/** Gets outgoing tangent for the specified point */
+	ValueType GetTangentOut(int32 Index) const
+	{
+		const int32 NumPoints = GetNumPoints();
+
+		if (NumPoints <= 1)
 		{
 			return ValueType();
 		}
-        
-        // Normal case - get from next segment's control points
-        const int32 BezierIndex = Index * 4;
-        if (BezierIndex + 1 < InternalSpline.NumKeys())
-        {
-            ValueType P0 = InternalSpline.GetValue(BezierIndex);
-            ValueType P1 = InternalSpline.GetValue(BezierIndex + 1);
-            return (P1 - P0) * 3.0f;
-        }
-        return ValueType();
-    }
+		
+		if (Index == NumPoints - 1)
+		{
+			if (IsClosedLoop())
+			{
+				return GetTangentOut(0);
+			}
+			else
+			{
+				// If spline is unclosed, out tangent == in tangent at final point.
+				return GetTangentIn(Index);
+			}
+		}
+		
+		// Normal case - get from next segment's control points
+		const int32 BezierIndex = Index * 4;
+		if (BezierIndex + 1 < InternalSpline.NumKeys())
+		{
+			ValueType P0 = InternalSpline.GetValue(BezierIndex);
+			ValueType P1 = InternalSpline.GetValue(BezierIndex + 1);
+			return (P1 - P0) * 3.0f;
+		}
+		return ValueType();
+	}
 
     /** Sets tangent in for the specified point */
     void SetTangentIn(int32 Index, const ValueType& NewTangent)
