make mesh surface point sampling more robust to inputs w/ very small bounds

diff --git a/Engine/Source/Runtime/GeometryCore/Private/Sampling/MeshSurfacePointSampling.cpp b/Engine/Source/Runtime/GeometryCore/Private/Sampling/MeshSurfacePointSampling.cpp
--- a/Engine/Source/Runtime/GeometryCore/Private/Sampling/MeshSurfacePointSampling.cpp
+++ b/Engine/Source/Runtime/GeometryCore/Private/Sampling/MeshSurfacePointSampling.cpp
@@ -288,6 +288,9 @@
 	MaxSamples = (MaxSamples == 0) ? TNumericLimits<int>::Max() : MaxSamples;
 	const bool bShuffle = (MaxSamples < TNumericLimits<int>::Max());
 
+	// Guard against requesting extremely small sample radii
+	SampleRadius = FMath::Max(UE_DOUBLE_SMALL_NUMBER, SampleRadius);
+
 	//
 	// Step 1: generate dense random point sampling of mesh surface. 
 	// 
@@ -302,25 +305,9 @@
 
 	if (Progress && Progress->Cancelled() ) return;
 
-	//
-	// Step 2: store dense point sampling in octree
 	// 
-	FSparseDynamicPointOctree3 Octree;
-	Octree.ConfigureFromPointCountEstimate(DensePointSet.Bounds.MaxDim(), NumDensePoints);
-	{
-		TRACE_CPUPROFILER_EVENT_SCOPE(BuildDenseOctree);
-		Octree.ParallelInsertDensePointSet(NumDensePoints, [&](int32 VertexID) { return DensePointSet.GetVertex(VertexID); } );
-	}
+	// Generate a random point ordering for sampling 
+	//
 
-	if (Progress && Progress->Cancelled() ) return;
-
-	//
-	// Step 3: dart throwing. Draw "random" dense point, add it to output set, and then
-	// remove all other dense points within radius of this point. 
-	// *NOTE* that in this implementation, we are not necessarily drawing random points.
-	// We are drawing from a random sampling on the triangles, but the per-triangle points
-	// are added to the PointSet in triangle index order. This introduces some bias
-	// but makes the algorithm quite a bit faster...
-
 	// currently only generating a semi-random point ordering (via shuffling) 
 	// if a subset of points is requested.  This likely does create some bias, 
@@ -344,6 +331,40 @@
 	}
 	int32 CurOrderingIndex = 0;
 
+	// if the bounds are small enough that we'd have at most a single point (i.e., sample radius sphere from any point in bounds will cover bounds), early out here
+	// this helps avoid the case of an octree covering very small bounds, which FSparseDynamicPointOctree3 currently doesn't handle well
+	if (DensePointSet.Bounds.MaxDim() * FMathd::Sqrt3 < SampleRadius)
+	{
+		if (NumDensePoints > 0)
+		{
+			int32 UseVertexID = PointOrdering[0];
+			FVector3d SamplePoint = DensePointSet.DensePoints[UseVertexID];
+			EmitSampleFunc(SamplePoint, DensePointSet.Triangles[UseVertexID], SampleRadius);
+		}
+		return;
+	}
+
+	//
+	// Step 2: store dense point sampling in octree
+	// 
+	FSparseDynamicPointOctree3 Octree;
+	Octree.ConfigureFromPointCountEstimate(DensePointSet.Bounds.MaxDim(), NumDensePoints);
+	{
+		TRACE_CPUPROFILER_EVENT_SCOPE(BuildDenseOctree);
+		Octree.ParallelInsertDensePointSet(NumDensePoints, [&](int32 VertexID) { return DensePointSet.GetVertex(VertexID); } );
+	}
+
+	if (Progress && Progress->Cancelled() ) return;
+
+	//
+	// Step 3: dart throwing. Draw "random" dense point, add it to output set, and then
+	// remove all other dense points within radius of this point. 
+	// *NOTE* that in this implementation, we are not necessarily drawing random points.
+	// We are drawing from a random sampling on the triangles, but the per-triangle points
+	// are added to the PointSet in triangle index order. This introduces some bias
+	// but makes the algorithm quite a bit faster...
+
+
 	double QueryRadiusSqr = 4 * SampleRadius * SampleRadius;
 	TArray<bool> IsValidPoint;
 	IsValidPoint.Init(true, DensePointSet.MaxVertexID());
@@ -426,6 +447,10 @@
 
 	MaxSamples = (MaxSamples == 0) ? TNumericLimits<int>::Max() : MaxSamples;
 
+	// Guard against requesting extremely small sample radii
+	MinSampleRadius = FMath::Max(UE_DOUBLE_SMALL_NUMBER, MinSampleRadius);
+	MaxSampleRadius = FMath::Max(UE_DOUBLE_SMALL_NUMBER, MaxSampleRadius);
+
 	//
 	// Step 1: generate dense random point sampling of mesh surface. 
 	// 
@@ -440,6 +465,38 @@
 
 	if (Progress && Progress->Cancelled() ) return;
 
+	// 
+	// Generate a random point ordering for sampling 
 	//
+
+	// currently only generating a semi-random point ordering (via shuffling) 
+	// if a subset of points is requested.  This likely does create some bias, 
+	TArray<int32> PointOrdering;
+	PointOrdering.Reserve(NumDensePoints);
+	FModuloIteration Iter(NumDensePoints);
+	{
+		uint32 NextIndex;
+		while (Iter.GetNextIndex(NextIndex))
+		{
+			PointOrdering.Add(NextIndex);
+		}
+	}
+	int32 CurOrderingIndex = 0;
+
+	// if the bounds are small enough that we'd have at most a single point (i.e., min radius sphere from any point in bounds will cover bounds), early out here
+	// this helps avoid the case of an octree covering very small bounds, which FSparseDynamicPointOctree3 currently doesn't handle well
+	if (DensePointSet.Bounds.MaxDim() * FMathd::Sqrt3 < MinSampleRadius)
+	{
+		if (NumDensePoints > 0)
+		{
+			int32 UseVertexID = PointOrdering[0];
+			FVector3d SamplePoint = DensePointSet.DensePoints[UseVertexID];
+			EmitSampleFunc(SamplePoint, DensePointSet.Triangles[UseVertexID], MinSampleRadius);
+		}
+		return;
+	}
+
+
+	//
 	// Step 2: store dense point sampling in octree
 	// 
@@ -460,20 +517,6 @@
 	// are added to the PointSet in triangle index order. This introduces some bias
 	// but makes the algorithm quite a bit faster...
 
-	// currently only generating a semi-random point ordering (via shuffling) 
-	// if a subset of points is requested.  This likely does create some bias, 
-	TArray<int32> PointOrdering;
-	PointOrdering.Reserve(NumDensePoints);
-	FModuloIteration Iter(NumDensePoints);
-	{
-		uint32 NextIndex;
-		while (Iter.GetNextIndex(NextIndex))
-		{
-			PointOrdering.Add(NextIndex);
-		}
-	}
-	int32 CurOrderingIndex = 0;
-
 	TArray<bool> IsValidPoint;
 	IsValidPoint.Init(true, DensePointSet.MaxVertexID());
 	TArray<const FSparsePointOctreeCell*> QueryTempBuffer;
