Quaternion Interpolation fixes

diff --git a/Engine/Source/Runtime/GeometryCore/Public/Splines/InterpolationPolicies.h b/Engine/Source/Runtime/GeometryCore/Public/Splines/InterpolationPolicies.h
--- a/Engine/Source/Runtime/GeometryCore/Public/Splines/InterpolationPolicies.h
+++ b/Engine/Source/Runtime/GeometryCore/Public/Splines/InterpolationPolicies.h
@@ -5,6 +5,15 @@
 #include "SplineMath.h"
 #include "ParameterizedTypes.h"
 
+namespace UE
+{
+namespace Math
+{
+	template<typename T>
+	struct TQuat;
+}
+}
+
 namespace UE
 {
 namespace Geometry
@@ -187,49 +196,164 @@ public:
 	}
 };
 
-// FQuat interpolation policy
-template<>
-class TSplineInterpolationPolicy<FQuat>
+// Partial specialization for any quaternion type (TQuat)
+template <typename RealType>
+class TSplineInterpolationPolicy<UE::Math::TQuat<RealType>>
 {
+	using TQuat = UE::Math::TQuat<RealType>;
 public:
-	// Existing parameter-based interpolation
-	static FQuat Interpolate(TArrayView<const FQuat* const> Window, float Parameter)
+    // Parameter-based interpolation using Slerp
+	static TQuat Interpolate(TArrayView<const TQuat* const> Window, float Parameter)
 	{
 		if (Window.Num() < 2)
 		{
-			return Window.Num() == 1 ? *Window[0] : FQuat::Identity;
+			return Window.Num() == 1 ? *Window[0] : TQuat::Identity;
 		}
-		return FQuat::Slerp(*Window[0], *Window[1], Parameter);
+        return TQuat::Slerp(*Window[0], *Window[1], static_cast<RealType>(Parameter));
 	}
 
-	// Window-based interpolation for spline implementations
-	static FQuat InterpolateWithBasis(TArrayView<const FQuat* const> Window, TArrayView<const float> Basis)
+	// Weighted quaternion interpolation 
+	// based on eigenvector approach from Markley et al.,
+	// which handles antipodal quaternions and degenerate cases effectively.
+	static TQuat InterpolateWithBasis(TArrayView<const TQuat* const> Window, TArrayView<const float> Basis)
 	{
-		FQuat Result(0.f, 0.f, 0.f, 0.f);
-		float TotalWeight = 0.0f;
+		const int32 Count = FMath::Min(Window.Num(), Basis.Num());
+		if (Count == 0)
+		{
+			return TQuat::Identity;
+		}
 
-		for (int32 i = 0; i < FMath::Min(Window.Num(), Basis.Num()); ++i)
+		// Optimization: single quaternion with full weight
+		for (int32 Index = 0; Index < Count; ++Index)
 		{
-			const float Weight = Basis[i];
-			if (Weight > UE_SMALL_NUMBER)
+			if (FMath::IsNearlyEqual(Basis[Index], 1.0f))
+			{
+				return *Window[Index];
+			}
+		}
+
+		// Use first quaternion as reference for hemisphere checks
+		const TQuat& ReferenceQuat = *Window[0];
+
+        // Build 4x4 correlation matrix
+		RealType M[4][4] = {};
+        RealType TotalWeight = 0;
+
+		for (int32 Index = 0; Index < Count; ++Index)
+		{
+            const RealType Weight = static_cast<RealType>(Basis[Index]);
+            if (Weight > static_cast<RealType>(UE_KINDA_SMALL_NUMBER))
 			{
-				const FQuat& Q = *Window[i];
-				const double Dot = Result.W * Q.W + Result.X * Q.X + Result.Y * Q.Y + Result.Z * Q.Z;
-				Result += (Dot >= 0.0f ? Q : -Q) * Weight;
+				TQuat Quat = *Window[Index];
+
+                // Ensure consistent hemisphere to avoid cancellation
+				if ((Quat | ReferenceQuat) < static_cast<RealType>(0))
+				{
+					Quat *= static_cast<RealType>(-1);
+				}
+
+				// Accumulate outer product contribution
+				const RealType Q[4] = {Quat.X, Quat.Y, Quat.Z, Quat.W};
+				for (int32 Row = 0; Row < 4; ++Row)
+				{
+					for (int32 Col = 0; Col < 4; ++Col)
+					{
+						M[Row][Col] += Weight * Q[Row] * Q[Col];
+					}
+				}
+
 				TotalWeight += Weight;
 			}
 		}
 
-		if (TotalWeight > UE_SMALL_NUMBER)
+		if (FMath::IsNearlyZero(TotalWeight))
 		{
-			Result.Normalize();
+			return TQuat::Identity;
 		}
-		return Result;
+
+		// Precompute inverse for efficiency
+		double InvTotalWeight = 1.0 / TotalWeight;
+
+        // Normalize matrix for numerical stability
+		for (int32 Row = 0; Row < 4; ++Row)
+		{
+			for (int32 Col = 0; Col < 4; ++Col)
+			{
+				M[Row][Col] *= InvTotalWeight;
+			}
+		}
+
+		// Find dominant eigenvector using power iteration
+        RealType EigenVec[4] = {0, 0, 0, 1};
+        RealType PrevEigenValue = 0;
+
+		constexpr int32 MaxIterations = 32;
+		for (int32 Iter = 0; Iter < MaxIterations; ++Iter)
+		{
+			// Matrix-vector multiplication: v' = M*v
+            RealType NewVec[4] = {0};
+			for (int32 Row = 0; Row < 4; ++Row)
+			{
+				for (int32 Col = 0; Col < 4; ++Col)
+				{
+					NewVec[Row] += M[Row][Col] * EigenVec[Col];
+				}
+			}
+
+            // Normalize vector
+			RealType LengthSq = 
+				NewVec[0] * NewVec[0] + 
+				NewVec[1] * NewVec[1] +
+				NewVec[2] * NewVec[2] + 
+				NewVec[3] * NewVec[3];
+
+			// Possible with certain opposing quaternion configurations
+			if (FMath::IsNearlyZero(LengthSq))
+			{
+				return TQuat::Identity;
+			}
+
+			RealType InvLength = FMath::InvSqrt(LengthSq);
+			for (int32 i = 0; i < 4; ++i)
+			{
+				EigenVec[i] = NewVec[i] * InvLength;
+			}
+
+			// Calculate Rayleigh quotient for convergence check
+            RealType EigenValue = 0;
+			for (int32 Row = 0; Row < 4; ++Row)
+			{
+				for (int32 Col = 0; Col < 4; ++Col)
+				{
+					EigenValue += EigenVec[Row] * M[Row][Col] * EigenVec[Col];
+				}
+			}
+
+			// Check for sufficient convergence
+			if (FMath::IsNearlyEqual(EigenValue, PrevEigenValue))
+			{
+				break;
+			}
+
+			PrevEigenValue = EigenValue;
+		}
+
+        // Convert eigenvector to TQuat
+		TQuat ResultQuat(
+            EigenVec[0],
+            EigenVec[1],
+            EigenVec[2],
+            EigenVec[3]
+		);
+
+		ResultQuat.Normalize();
+		return ResultQuat;
 	}
 
+
 	// Template-based n-th derivative calculation
 	template<int32 Order>
-	static FQuat EvaluateDerivative(TArrayView<const FQuat* const> Window, float Parameter)
+	static TQuat EvaluateDerivative(TArrayView<const TQuat* const> Window, float Parameter)
 	{
 		if (Order == 0)
 		{
@@ -239,38 +363,44 @@ public:
 		// First derivative - angular velocity
 		if (Order == 1)
 		{
-			// Adaptive h based on the difference between quaternions at Parameter and Parameter + base_h
-			float base_h = 0.01f; // Start with a reasonable base value
-			FQuat Q0 = Interpolate(Window, Parameter);
-			FQuat Q1base = Interpolate(Window, Parameter + base_h);
-			float diff = static_cast<float>(Q0.AngularDistance(Q1base)); // Use a suitable distance metric
-			float h = base_h * FMath::Clamp(diff, 0.001f, 0.1f); // Adjust h based on the difference
-
-			FQuat Q1 = Interpolate(Window, Parameter + h);
-			return (Q1 * Q0.Inverse()).Log() * (1.0f / h);
+            // Adaptive step size based on quaternion difference
+            RealType base_h = static_cast<RealType>(0.01);
+			TQuat Q0 = Interpolate(Window, Parameter);
+            TQuat Q1base = Interpolate(Window, Parameter + static_cast<float>(base_h));
+            RealType diff = Q0.AngularDistance(Q1base);
+            RealType h = base_h * FMath::Clamp(
+                diff, 
+                static_cast<RealType>(0.001),
+                static_cast<RealType>(0.1)
+            );
+
+            TQuat Q1 = Interpolate(Window, Parameter + static_cast<float>(h));
+            return (Q1 * Q0.Inverse()).Log() * (static_cast<RealType>(1) / h);
 		}
 
 		// Second derivative - angular acceleration
 		if (Order == 2)
 		{
-			// Use the same adaptive h as the first derivative, centered around the parameter
-			float base_h = 0.01f;
-			FQuat Q_m1 = Interpolate(Window, Parameter - base_h);
-			FQuat Q_p1 = Interpolate(Window, Parameter + base_h);
-			float diff = static_cast<float>(Q_m1.AngularDistance(Q_p1));
-			float h = base_h * FMath::Clamp(diff / 2.0f, 0.001f, 0.1f);  // Divide diff by 2 for centering
-
-			FQuat V0 = EvaluateDerivative<1>(Window, Parameter - h);
-			FQuat V1 = EvaluateDerivative<1>(Window, Parameter + h);
-			return (V1 * V0.Inverse()).Log() * (1.0f / (2.0f * h));
+            RealType base_h = static_cast<RealType>(0.01);
+            TQuat Q_m1 = Interpolate(Window, Parameter - static_cast<float>(base_h));
+            TQuat Q_p1 = Interpolate(Window, Parameter + static_cast<float>(base_h));
+            RealType diff = Q_m1.AngularDistance(Q_p1);
+            RealType h = base_h * FMath::Clamp(
+                diff / static_cast<RealType>(2),
+                static_cast<RealType>(0.001),
+                static_cast<RealType>(0.1)
+            );
+
+            TQuat V0 = EvaluateDerivative<1>(Window, Parameter - static_cast<float>(h));
+            TQuat V1 = EvaluateDerivative<1>(Window, Parameter + static_cast<float>(h));
+            return (V1 * V0.Inverse()).Log() * (static_cast<RealType>(1) / (static_cast<RealType>(2) * h));
 		}
 
 		// Higher orders generally not meaningful for rotations
-		return FQuat::Identity;
+		return TQuat::Identity;
 	}
-
 };
-
+	
 // FTransform interpolation policy
 template<>
 class TSplineInterpolationPolicy<FTransform>
