remove extra normalize calls from vectorutil tri normal method (to match NormalArea)

diff --git a/Engine/Source/Runtime/GeometryCore/Public/VectorUtil.h b/Engine/Source/Runtime/GeometryCore/Public/VectorUtil.h
--- a/Engine/Source/Runtime/GeometryCore/Public/VectorUtil.h
+++ b/Engine/Source/Runtime/GeometryCore/Public/VectorUtil.h
@@ -71,8 +71,6 @@ namespace VectorUtil
 	{
 		TVector<RealType> edge1(V1 - V0);
 		TVector<RealType> edge2(V2 - V0);
-		Normalize(edge1);
-		Normalize(edge2);
 		// Unreal has Left-Hand Coordinate System so we need to reverse this cross-product to get proper triangle normal
 		TVector<RealType> vCross(edge2.Cross(edge1));
 		return Normalized(vCross);
