Added tessellator with parallel topology refinement

FAdaptiveTessellator: add safety option for max triangles, custom refinement callback

Terrain:
Add options for minimum and maximum edgelength
Remove EWA from UI and demote it to CVar

diff --git a/Engine/Source/Runtime/GeometryCore/Public/Tessellation/AdaptiveTessellator.h b/Engine/Source/Runtime/GeometryCore/Public/Tessellation/AdaptiveTessellator.h
--- a/Engine/Source/Runtime/GeometryCore/Public/Tessellation/AdaptiveTessellator.h
+++ b/Engine/Source/Runtime/GeometryCore/Public/Tessellation/AdaptiveTessellator.h
@@ -6,6 +6,7 @@
 #include "Math/Vector.h"
 #include "Math/Vector2D.h"
 
+#include "Async/ParallelFor.h"
 #include "Templates/Function.h"
 #include "Templates/Tuple.h"
 #include "Templates/IsMemberPointer.h"
@@ -151,22 +152,45 @@ private:
 
   public:
 
+	enum class ERefinementMethod
+	{
+		HierarchicalSplit,
+		
+		Custom
+	};
+
 	using RealType = typename MeshT::RealType; // the field related to the vector-space
 	using VecType  = typename MeshT::VecType;  // for positions and normals
 
+	struct FOptions
+	{
+		// target error to achieve (with respect to error defined by displacement policy)
+		RealType TargetError; 
+
+		// one-dimensional rate (edge-length) at which features are sampled. triangles with edg	e-length smaller than
+		// this value will not be refined
+		RealType SampleRate; 
+
+		// vertices with same positions are considered to be identical and stay at the same position after displacement
+		bool bCrackFree { true };
+
+		// apply displacement to the tessellated mesh
+		bool bFinalDisplace { true };
+
+		// upper bound on the number of triangles of the tessellated mesh
+		uint32 MaxTriangles { 10'000'000u };
+
+		ERefinementMethod RefinementMethod { ERefinementMethod::HierarchicalSplit };
+	};
+
     // Initialize and run adaptive tessellation
 	FAdaptiveTessellator(
 		MeshT& InMesh,                                //< mesh interface
 		DisplacementPolicyT& InDisplacementPolicy,    //< displacement and error interface
-		RealType	InTargetError,                    //< target error to achieve (with respect to error defined by displacement policy)
-		RealType	InSampleRate,                     //< one-dimensional rate (edge-length) at which features are sampled. triangles with edge-length smaller than
-		                                              //< this value will not be refined
-		bool		bCrackFree,                       //< vertices with same positions are considered to be identical and stay at the same position after displacement
-		bool        bFinalDisplace)                   //< apply displacement to the tessellated mesh
+		const FOptions& InOptions)                   
 		: Mesh(InMesh)
 		, DisplacementPolicy(InDisplacementPolicy)
-		, TargetError(InTargetError)
-		, SampleRate(InSampleRate)
+		, Options(InOptions)
 	{
 		TriSplitRecords.AddUninitialized( Mesh.MaxTriID() );
 		Displacements.Init( VecType( std::numeric_limits<RealType>::signaling_NaN(), 
@@ -194,54 +218,74 @@ private:
 			}
 		}
 
-		SplitRequests.SetNum(TriSplitRecords.Num() );
+		SplitRequests.SetNum( TriSplitRecords.Num() );
 		NumSplits = 0;
 		
-		ParallelFor( TEXT("FAdaptiveTessellator2.FindSplitBVH.PF"), TriSplitRecords.Num(), 32,
+		NumNewTriangles.store(0);
+
+		if ( static_cast<uint32>(Mesh.MaxTriID()) >= Options.MaxTriangles )
+		{
+			return;
+		}
+		
+		ParallelFor( TEXT("FAdaptiveTessellator.FindSplitBVH.PF"), TriSplitRecords.Num(), 32,
 		[&]( uint32 TriIndex )
 		{
-			FindSplitBVH( TriIndex );
+			CheckSplitAndEnqueue( TriIndex, 0 );
 		} );
-				
+	
+		int32 Iter = 0;
 		while( NumSplits )
 		{
 			// Size to atomic count and sort for deterministic order
 			{
-				TRACE_CPUPROFILER_EVENT_SCOPE_TEXT("FAdaptiveTessellator2.SortSplits.SF");
+				TRACE_CPUPROFILER_EVENT_SCOPE_TEXT("FAdaptiveTessellator.SortSplits.SF");
 				SplitRequests.SetNum( NumSplits, EAllowShrinking::No );
 				SplitRequests.Sort();
 			}
 
 			for( int32 i = 0; i < SplitRequests.Num(); i++ )
+			{
 				TriSplitRecords[ SplitRequests[i] ].RequestIndex = i;
+			}
 	
 			{
-				TRACE_CPUPROFILER_EVENT_SCOPE_TEXT("FAdaptiveTessellator2.SplitTriangles.SF");
-				while( SplitRequests.Num() )
+				TRACE_CPUPROFILER_EVENT_SCOPE_TEXT("FAdaptiveTessellator.SplitTriangles.SF");
+				while( SplitRequests.Num() ) 
+				{
 					SplitTriangle( SplitRequests.Pop() );
+				}
+			}
+
+			if ( static_cast<uint32>(Mesh.MaxTriID()) >= Options.MaxTriangles)
+			{
+				break;
 			}
 
 			SplitRequests.SetNum( TriSplitRecords.Num() );
 			NumSplits = 0;
 
-			ParallelFor( TEXT("FAdaptiveTessellator2.FindSplitBVH.PF"), FindRequests.Num(), 32,
+			NumNewTriangles.store(0);
+
+			ParallelFor( TEXT("FAdaptiveTessellator.FindSplitBVH.PF"), FindRequests.Num(), 32,
 				[&]( uint32 i )
 				{
-					uint32 TriIndex = FindRequests[i];
-					FindSplitBVH( TriIndex );
- 				} );
+					CheckSplitAndEnqueue( FindRequests[i], Iter );
+				} );
 				
 			FindRequests.Reset();
+
+			Iter++;
 		}
 
-		if( bCrackFree && Mesh.IsTriangleSoup())
+		if( Options.bCrackFree && Mesh.IsTriangleSoup())
 		{
 			TArray<VecType> OldDisplacements;
 			OldDisplacements.AddUninitialized( Displacements.Num() );
 			Swap( Displacements, OldDisplacements );
 
 			FHashTable HashTable( 1 << FMath::FloorLog2( Mesh.MaxVertexID() ), Mesh.MaxVertexID() );
-			ParallelFor( TEXT("FAdaptiveTessellator2.HashVerts.PF"), Mesh.MaxVertexID(), 4096,
+			ParallelFor( TEXT("FAdaptiveTessellator.HashVerts.PF"), Mesh.MaxVertexID(), 4096,
 				[&]( int32 i )
 				{
 					if (Mesh.IsValidVertex(i))
@@ -251,7 +295,7 @@ private:
 					}
 				} );
 
-			ParallelFor( TEXT("FAdaptiveTessellator2.HashVerts.PF"), Mesh.MaxVertexID(), 4096,
+			ParallelFor( TEXT("FAdaptiveTessellator.HashVerts.PF"), Mesh.MaxVertexID(), 4096,
 				[&]( int32 VID )
 				{
 					if (!Mesh.IsValidVertex(VID))
@@ -276,9 +320,9 @@ private:
 				} );
 		}
 
-		if (bFinalDisplace)
+		if (Options.bFinalDisplace)
 		{
-			ParallelFor(TEXT("FAdaptiveTessellator2.Displace.PF"), Mesh.MaxVertexID(), 4096,
+			ParallelFor(TEXT("FAdaptiveTessellator.Displace.PF"), Mesh.MaxVertexID(), 4096,
 				[&](int32 VID)
 				{
 					if (Mesh.IsValidVertex(VID))
@@ -356,15 +400,17 @@ private:
 
 	void AddFindRequest(int32 TriIndex);
 
+	void CheckSplitAndEnqueue( int32 TriIndex, int32 Level ); // thread-safe
+
+	void EnqueueSplitRequest( int32 TriIndex, VecType Barycentrics ); // thread-safe
+
 	void RemoveSplitRequest(int32 TriIndex);
 
-	void TryDelaunayFlip(int32 TriIndex, int TriEdgeIndex);
+	void TryDelaunayFlip( int32 TriIndex, int TriEdgeIndex );
 
 	MeshT&               Mesh;
 	DisplacementPolicyT& DisplacementPolicy;
-
-	const RealType       TargetError;
-	const RealType       SampleRate;
+	const FOptions&      Options;
 
 	TArray<VecType>      Displacements;
 	TArray<FTriSplit>	 TriSplitRecords;
@@ -372,12 +418,12 @@ private:
 	TArray<uint32>		 SplitRequests; // triangle indices
 
 	std::atomic<uint32>	 NumSplits;
+	std::atomic<uint32>  NumNewTriangles;
 };
 
 template <typename MeshT, typename DisplacementPolicyT>
 inline void FAdaptiveTessellator<MeshT, DisplacementPolicyT>::FindSplitBVH( uint32 TriIndex )
 {
-	TriSplitRecords[ TriIndex ].RequestIndex = -1;
 	const VecType P0 = Mesh.GetVertexPosition(TriIndex, 0);
 	const VecType P1 = Mesh.GetVertexPosition(TriIndex, 1);
 	const VecType P2 = Mesh.GetVertexPosition(TriIndex, 2);
@@ -388,9 +434,9 @@ inline void FAdaptiveTessellator<MeshT, DisplacementPolicyT>::FindSplitBVH( uint
 
 	const VecType EdgeLengths( Edge01.Length(), Edge12.Length(), Edge20.Length() );
 
-	if( EdgeLengths[0] < SampleRate &&
-		EdgeLengths[1] < SampleRate &&
-		EdgeLengths[2] < SampleRate )
+	if( EdgeLengths[0] < Options.SampleRate &&
+		EdgeLengths[1] < Options.SampleRate &&
+		EdgeLengths[2] < Options.SampleRate )
 	{
 		return;
 	}
@@ -398,7 +444,7 @@ inline void FAdaptiveTessellator<MeshT, DisplacementPolicyT>::FindSplitBVH( uint
 	const float TriArea = 0.5f * ( Edge01 ^ Edge20 ).Size();
 
 	// area of equilateral triangle with edgelength = SampleRate. We want to have 1 sample per triangle of this size
-	const float SampleArea = EquilateralArea( SampleRate );
+	const float SampleArea = EquilateralArea( Options.SampleRate );
 	
 	const FIndex3i Triangle = Mesh.GetTriangle(TriIndex);
 
@@ -452,10 +498,10 @@ inline void FAdaptiveTessellator<MeshT, DisplacementPolicyT>::FindSplitBVH( uint
 		Node.ErrorMax = ErrorBounds.Y;
 	}
 
-	RealType ErrorMinimum = TargetError;
+	RealType ErrorMinimum = Options.TargetError;
 
 	while(	Node.ErrorMax >= ErrorMinimum &&
-			Node.ErrorMax > BestError * 1.25f + TargetError )
+			Node.ErrorMax > BestError * 1.25f + Options.TargetError )
 	{
 		for( uint32 ChildIndex = 0; ChildIndex < 4; ChildIndex++ )
 		{
@@ -511,11 +557,11 @@ inline void FAdaptiveTessellator<MeshT, DisplacementPolicyT>::FindSplitBVH( uint
 			if( ErrorMinimum < ChildNode.ErrorMin )
 				ErrorMinimum = ChildNode.ErrorMin;
 
-			if( ChildNode.ErrorMax > BestError * 1.25f + TargetError )
+			if( ChildNode.ErrorMax > BestError * 1.25f + Options.TargetError )
 			{
 				int32 NumSamples = 64;
 
-				bool bStopTraversal = ChildNode.Level == 12 || ChildNode.ErrorMax - ChildNode.ErrorMin < TargetError;
+				bool bStopTraversal = ChildNode.Level == 12 || ChildNode.ErrorMax - ChildNode.ErrorMin < Options.TargetError;
 				if( !bStopTraversal )
 				{
 					RealType ChildArea = SubtriangleArea<RealType>( Barycentrics[0], Barycentrics[1], Barycentrics[2], TriArea );
@@ -543,7 +589,6 @@ inline void FAdaptiveTessellator<MeshT, DisplacementPolicyT>::FindSplitBVH( uint
 						UE::Math::TVector2<RealType> Random;
 
 						// todo see UniformSampleTrianglePoint
-
 						Random.X = RealType( i + 1 ) / RealType( NumSamples + 1 );
 						Random.Y = RealType( Reverse >> 8 ) * 0x1p-24f;
 
@@ -576,15 +621,15 @@ inline void FAdaptiveTessellator<MeshT, DisplacementPolicyT>::FindSplitBVH( uint
 						CA_ASSUME(e1 <= 2);
 						CA_ASSUME(e2 <= 2);
 
-						bool bTooCloseToEdge = DistToEdge[ e0 ] < 0.5f * SampleRate;
+						bool bTooCloseToEdge = DistToEdge[ e0 ] < 0.5f * Options.SampleRate;
 						if( bTooCloseToEdge && !bCouldFlipEdge[ e0 ] )
 						{
 							Split[ e0 ] = Split[ e0 ] / ( Split[ e0 ] + Split[ e1 ] );
 							Split[ e1 ] = 1.0f - Split[ e0 ];
 							Split[ e2 ] = 0.0f;
 
-							bool bTooCloseToEdge1 = !bCouldFlipEdge[ e1 ] && DistanceToEdge<RealType>( Split[ e0 ], EdgeLengths[ e1 ], TriArea ) < 0.5f * SampleRate;
-							bool bTooCloseToEdge2 = !bCouldFlipEdge[ e2 ] && DistanceToEdge<RealType>( Split[ e1 ], EdgeLengths[ e2 ], TriArea ) < 0.5f * SampleRate;
+							bool bTooCloseToEdge1 = !bCouldFlipEdge[ e1 ] && DistanceToEdge<RealType>( Split[ e0 ], EdgeLengths[ e1 ], TriArea ) < 0.5f * Options.SampleRate;
+							bool bTooCloseToEdge2 = !bCouldFlipEdge[ e2 ] && DistanceToEdge<RealType>( Split[ e1 ], EdgeLengths[ e2 ], TriArea ) < 0.5f * Options.SampleRate;
 							bTooCloseToEdge = bTooCloseToEdge1 || bTooCloseToEdge2;
 						}
 
@@ -629,11 +674,9 @@ inline void FAdaptiveTessellator<MeshT, DisplacementPolicyT>::FindSplitBVH( uint
 			}, EAllowShrinking::No );
 	}
 
-	if( BestError > TargetError )
+	if( BestError > Options.TargetError )
 	{
-		TriSplitRecords[ TriIndex ].SplitBarycentrics = BestSplit;
-		TriSplitRecords[ TriIndex ].RequestIndex = NumSplits++;
-		SplitRequests[ TriSplitRecords[ TriIndex ].RequestIndex ] = TriIndex;
+		EnqueueSplitRequest( TriIndex, BestSplit );
 	}
 }
 
@@ -642,11 +685,11 @@ inline void FAdaptiveTessellator<MeshT, DisplacementPolicyT>::SplitTriangle( uin
 {
 	const VecType Barycentrics = TriSplitRecords[ TriIndex ].SplitBarycentrics;
 	TriSplitRecords[ TriIndex ].RequestIndex = -1;
-	
+
 	// first check whether we are splitting along an edge
 	for( uint32 EdgeIndex = 0; EdgeIndex < 3; EdgeIndex++ )
 	{
-		if( Barycentrics[ ( EdgeIndex + 2 ) % 3 ] == 0.0f )
+		if( Barycentrics[ ( EdgeIndex + 2 ) % 3 ] == RealType(0) )
 		{
 			if (!Mesh.AllowEdgeSplit(TriIndex, EdgeIndex))
 			{
@@ -705,7 +748,6 @@ inline void FAdaptiveTessellator<MeshT, DisplacementPolicyT>::SplitTriangle( uin
 	const VecType NewDisplacement = DisplacementPolicy.GetDisplacement( Barycentrics, TriIndex );
 	const FPokeInfo pokeInfo = Mesh.PokeTriangle(TriIndex, Barycentrics);
 
-	// also true for DynamicMesh
 	check(pokeInfo.NewTriIndex[0] == TriIndex);
 	
 	GrowTriRecords(FMath::Max(pokeInfo.NewTriIndex[1], pokeInfo.NewTriIndex[2])+1);
@@ -734,6 +776,48 @@ inline void FAdaptiveTessellator<MeshT, DisplacementPolicyT>::AddFindRequest( in
 	}
 }
 
+template <typename MeshT, typename DisplacementPolicyT>
+inline void FAdaptiveTessellator<MeshT, DisplacementPolicyT>::CheckSplitAndEnqueue( const int32 TriIndex, const int32 Level )
+{
+	TriSplitRecords[ TriIndex ].RequestIndex = -1;
+
+	if ( Options.RefinementMethod == ERefinementMethod::HierarchicalSplit )
+	{
+		FindSplitBVH( TriIndex );
+	}
+	else
+	{
+		VecType SplitVertexBary;
+
+		const FIndex3i Triangle = Mesh.GetTriangle( TriIndex );
+		const VecType& Displacement0 = Displacements[ Triangle.A ];
+		const VecType& Displacement1 = Displacements[ Triangle.B ];
+		const VecType& Displacement2 = Displacements[ Triangle.C ];
+
+		if (DisplacementPolicy.ShouldRefine( TriIndex, Displacement0, Displacement1, Displacement2, SplitVertexBary, Level ))
+		{
+			EnqueueSplitRequest( TriIndex, SplitVertexBary );
+		}
+	}
+}
+
+template <typename MeshT, typename DisplacementPolicyT>
+inline void FAdaptiveTessellator<MeshT, DisplacementPolicyT>::EnqueueSplitRequest( int32 TriIndex, VecType Barycentrics )
+{
+	uint32 AddTrianglesPerRefine = 3;
+	if (Barycentrics[0] == RealType(0) || Barycentrics[1] == RealType(0) || Barycentrics[2] == RealType(0) )
+	{
+		AddTrianglesPerRefine = 2;
+	}
+
+	if ( NumNewTriangles.fetch_add(AddTrianglesPerRefine) < Options.MaxTriangles - Mesh.MaxTriID() )
+	{
+		TriSplitRecords[ TriIndex ].SplitBarycentrics = Barycentrics;
+		TriSplitRecords[ TriIndex ].RequestIndex = NumSplits++;
+		SplitRequests[ TriSplitRecords[ TriIndex ].RequestIndex ] = TriIndex;
+	}
+}
+
 template <typename MeshT, typename DisplacementPolicyT>
 inline void FAdaptiveTessellator<MeshT, DisplacementPolicyT>::RemoveSplitRequest( int32 TriIndex )
 {
diff --git a/Engine/Source/Runtime/GeometryCore/Public/Tessellation/ParallelAdaptiveRefinement.h b/Engine/Source/Runtime/GeometryCore/Public/Tessellation/ParallelAdaptiveRefinement.h
new file mode 100644
--- /dev/null
+++ b/Engine/Source/Runtime/GeometryCore/Public/Tessellation/ParallelAdaptiveRefinement.h
@@ -0,0 +1,692 @@
+// Copyright Epic Games, Inc. All Rights Reserved.
+
+#pragma once
+
+#include "HAL/Platform.h"
+#include "Misc/AssertionMacros.h"
+#include "IndexTypes.h"
+#include "VectorUtil.h"
+#include "Async/ParallelFor.h"
+#include "ProfilingDebugging/CpuProfilerTrace.h"
+
+#include <atomic>
+
+namespace UE {
+namespace Geometry {
+
+template <typename MeshType, typename RefinementPolicyType>
+class TParallelAdaptiveRefinement : public FNoncopyable
+{
+private:
+
+	struct FTriSplitRecord
+	{
+		uint16 SplitMask                 { 0 }; //< bitmask: per-edge refinement tagging [0..3)
+		std::atomic<uint16> GreenAdjMask { 0 }; //< bitmask: encode which of the neighbors [0..3) are green refinement 
+		std::atomic_flag bClaimed = ATOMIC_FLAG_INIT;
+
+		FORCEINLINE void Reset()
+		{
+			SplitMask = 0;
+			GreenAdjMask.store(0);
+			bClaimed.clear();
+		}
+
+		// split one edge (green refinement)
+		bool IsSplit1() const
+		{
+			return SplitMask == 1 || SplitMask == 2 || SplitMask == 4;
+		}
+
+		// split two edges
+		bool IsSplit2() const
+		{
+			return SplitMask == 3 || SplitMask == 5 || SplitMask == 6;
+		}
+
+		// split all three edges (red refinement)
+		bool IsSplit3() const
+		{
+			return SplitMask == 7;
+		}
+	};
+
+	struct FEdgeSplitRecord
+	{
+		void Reset()
+		{
+			SplitVertex = IndexConstants::InvalidID;
+			HalfEdges[0] = FIndex2i(IndexConstants::InvalidID, IndexConstants::InvalidID);
+			HalfEdges[1] = FIndex2i(IndexConstants::InvalidID, IndexConstants::InvalidID);
+			NewEdge = IndexConstants::InvalidID;
+
+			bTagged.clear();
+			EntryCountPreWrite.store(0);
+			EntryCountPostWrite.store(0);
+		}
+
+		int32 SplitVertex { IndexConstants::InvalidID };
+
+		// each entry pair corresponds to the half-edge indices of one side of the split
+		FIndex2i HalfEdges[2] { FIndex2i(IndexConstants::InvalidID, IndexConstants::InvalidID),
+								FIndex2i(IndexConstants::InvalidID, IndexConstants::InvalidID) };
+
+		int32 NewEdge { IndexConstants::InvalidID };
+
+		std::atomic_flag   bTagged = ATOMIC_FLAG_INIT; // tagged for refinement
+		std::atomic<int32> EntryCountPreWrite  { 0 };  // current number of entries, before write
+		std::atomic<int32> EntryCountPostWrite { 0 };  // number of half-edge pairs written
+	};
+
+	struct FConcurrencyState
+	{
+		TArray<FEdgeSplitRecord> EdgeSplitRecords;
+		TArray<FTriSplitRecord>  TriSplitRecords;
+		std::atomic<int32>       HalfEdgeCount;   // current number of halfedges, for parallel append
+		std::atomic<int32>       EdgeCount;       // current number of edges, for parallel append
+
+		void Resize(const int32 NumTris, const int32 NumEdges)
+		{
+			TRACE_CPUPROFILER_EVENT_SCOPE_TEXT("TParallelAdaptiveRefinement.FConcurrencyState.Resize");
+			if (NumTris > TriSplitRecords.Num())
+			{
+				TriSplitRecords.AddDefaulted(NumTris - TriSplitRecords.Num());
+			}
+			check(TriSplitRecords.Num() == NumTris);
+
+			for (FEdgeSplitRecord &Record : EdgeSplitRecords)
+			{
+				Record.Reset();
+			}
+
+			if (NumEdges > EdgeSplitRecords.Num())
+			{
+				EdgeSplitRecords.AddDefaulted(NumEdges - EdgeSplitRecords.Num());
+			}
+		}
+	};
+
+public:
+	using RealType = MeshType::RealType; // the field related to the vector-space
+	using VecType  = MeshType::VecType;  // for positions and normals
+
+	// Initialize and run adaptive aptive tesselation
+	TParallelAdaptiveRefinement(
+		MeshType& InMesh,                               //< mesh interface
+		RefinementPolicyType& InRefinementPolicy,       //< displacement and error interface
+		const int32 ParallelForBatchSize = 1024,     //<
+		const int32 ReserveVertices = 2'000'000 )
+		: Mesh(InMesh)
+		, RefinementPolicy(InRefinementPolicy)
+	{
+		const EParallelForFlags ParallelForFlags = ParallelForBatchSize > 0 ? EParallelForFlags::None : EParallelForFlags::ForceSingleThread;
+
+		const int32 ReserveTriangles = 2 * ReserveVertices;
+		const int32 ReserveEdges     = 3 * ReserveVertices;
+
+		FConcurrencyState ConcurrencyState;
+		TArray<int32> RefinementCand;      //< list of triangles to inspect for refinement
+		TArray<int32> TriSplitQueue;       //< list of triangle indices scheduled for refinement
+		TArray<int32> EdgeSplitRequests;   //< list of edge indices scheduled for refinement
+
+		RefinementCand.Reserve(ReserveTriangles);
+		TriSplitQueue.Reserve(ReserveTriangles);
+		EdgeSplitRequests.Reserve(ReserveEdges);
+
+		ConcurrencyState.TriSplitRecords.Reserve(ReserveTriangles);
+		ConcurrencyState.EdgeSplitRecords.Reserve(ReserveEdges);
+
+		Mesh.ReserveVertices(ReserveVertices);
+		Mesh.ReserveTriangles(ReserveTriangles);
+
+		RefinementCand.Init(0, Mesh.MaxTriID());
+		for (int32 TriangleIdx(0); TriangleIdx<Mesh.MaxTriID(); ++TriangleIdx)
+		{
+			RefinementCand[TriangleIdx] = TriangleIdx;
+		}
+
+		std::atomic<int32> NumRefinementCand(RefinementCand.Num());
+
+		for (int32 Level(0); Level < 32; ++Level)
+		{
+			if (NumRefinementCand.load() == 0) {
+				break;
+			}
+
+			{
+				TRACE_CPUPROFILER_EVENT_SCOPE_TEXT("TParallelAdaptiveRefinement.ResizeRequests");
+				// reserve sufficient space for tris and edges to be split
+				TriSplitQueue.SetNum(Mesh.MaxTriID());
+				EdgeSplitRequests.SetNum(Mesh.EdgeCount());
+				ConcurrencyState.Resize(Mesh.MaxTriID(), Mesh.EdgeCount());
+			}
+
+			check(ConcurrencyState.TriSplitRecords.Num() == Mesh.MaxTriID());
+
+			std::atomic<int32> TriSplitQueueSize(0);
+			std::atomic<int32> NumEdgesToSplit(0);
+
+			auto TagForRefine = [&](const int32 Idx)
+			{
+				const int32 TriIndex = RefinementCand[Idx];
+
+				if (!Mesh.IsValidTri(TriIndex))
+				{
+					return;
+				}
+
+				if (ConcurrencyState.TriSplitRecords[TriIndex].bClaimed.test())
+				{
+					return; // early out
+				}
+
+				if (ShouldRefine(TriIndex, Level))
+				{
+					if (!ConcurrencyState.TriSplitRecords[TriIndex].bClaimed.test_and_set())
+					{
+						TriSplitQueue[TriSplitQueueSize.fetch_add(1, std::memory_order_relaxed)] = TriIndex;
+						ConcurrencyState.TriSplitRecords[TriIndex].SplitMask = 7;
+
+						for (int32 LocalEdgeIdx=0; LocalEdgeIdx<3; ++LocalEdgeIdx)
+						{
+							const int32 EdgeIdx = Mesh.GetEdgeIndex(TriIndex, LocalEdgeIdx);
+							FEdgeSplitRecord& Edge = ConcurrencyState.EdgeSplitRecords[EdgeIdx];
+
+							// check whether this is the first time visiting this edge
+							if (!Edge.bTagged.test_and_set(std::memory_order_relaxed))
+							{
+								EdgeSplitRequests[ NumEdgesToSplit.fetch_add(1, std::memory_order_relaxed) ] = EdgeIdx;
+							}
+						}
+
+						// now the same for all the neighors
+						for (int NbIdx=0; NbIdx<3; ++NbIdx)
+						{
+							const int32 AdjTriIndex = Mesh.GetAdjTriangle(TriIndex, NbIdx);
+							if (AdjTriIndex != IndexConstants::InvalidID && !ConcurrencyState.TriSplitRecords[AdjTriIndex].bClaimed.test_and_set())
+							{
+								check(Mesh.IsValidTri(AdjTriIndex));
+
+								TriSplitQueue[TriSplitQueueSize.fetch_add(1, std::memory_order_relaxed)] = AdjTriIndex;
+								ConcurrencyState.TriSplitRecords[AdjTriIndex].SplitMask = 7;
+
+								for (int32 LocalEdgeIdx=0; LocalEdgeIdx<3; ++LocalEdgeIdx)
+								{
+									const int32 EdgeIdx = Mesh.GetEdgeIndex(AdjTriIndex, LocalEdgeIdx);
+									FEdgeSplitRecord& Edge = ConcurrencyState.EdgeSplitRecords[EdgeIdx];
+
+									// check whether this is the first time visiting this edge
+									if (!Edge.bTagged.test_and_set(std::memory_order_relaxed))
+									{
+										EdgeSplitRequests[ NumEdgesToSplit.fetch_add(1, std::memory_order_relaxed) ] = EdgeIdx;
+									}
+								}
+							}
+						}
+					}
+				}
+			};
+
+			// tag triangles for refinement;
+			ParallelForTemplate( TEXT("TParallelAdaptiveRefinement.TagForRefine.PF"), NumRefinementCand.load(), ParallelForBatchSize, TagForRefine, ParallelForFlags );
+
+			NumRefinementCand.store(0);
+			if (TriSplitQueueSize.load() == 0)
+			{
+				break;
+			}
+
+			std::atomic<int32> NumNewInteriorEdges(0);
+			std::atomic<int32> NumNewTriangles(0);
+
+			// identify neighboring triangles closure split masks (T-junctions)
+			ParallelFor( TEXT("TParallelAdaptiveRefinement.ResolveNeighbors.PF"), Mesh.MaxTriID(), ParallelForBatchSize, [&](const int32 TriIndex)
+			{
+				if (!Mesh.IsValidTri(TriIndex))
+				{
+					return;
+				}
+
+				int32 SplitMask = ConcurrencyState.TriSplitRecords[TriIndex].SplitMask;
+
+				if (SplitMask != 0)
+				{
+					check(SplitMask == 7);
+				}
+				else
+				{
+					for (int LocalEdgeIdx=0; LocalEdgeIdx<3; ++LocalEdgeIdx)
+					{
+						const int32 AdjTriIndex = Mesh.GetAdjTriangle(TriIndex, LocalEdgeIdx);
+						if (AdjTriIndex != IndexConstants::InvalidID && ConcurrencyState.TriSplitRecords[AdjTriIndex].IsSplit3())
+						{
+							SplitMask |= 1 << LocalEdgeIdx;
+						}
+					}
+				}
+
+				if (ConcurrencyState.TriSplitRecords[TriIndex].SplitMask == 0 && SplitMask != 0)
+				{
+					// record that this triangle needs to be split as a consequence as well
+					TriSplitQueue[TriSplitQueueSize.fetch_add(1, std::memory_order_relaxed)] = TriIndex;
+				}
+
+				ConcurrencyState.TriSplitRecords[TriIndex].SplitMask = SplitMask;
+				ConcurrencyState.TriSplitRecords[TriIndex].GreenAdjMask.store(uint16(0));
+
+				if (SplitMask == 7)
+				{
+					// full split
+					NumNewInteriorEdges.fetch_add(3, std::memory_order_relaxed);
+					NumNewTriangles.fetch_add(3, std::memory_order_relaxed);
+				}
+				else if (SplitMask == 1 || SplitMask == 2 || SplitMask == 4)
+				{
+					// green refinement
+					NumNewInteriorEdges.fetch_add(1, std::memory_order_relaxed);
+					NumNewTriangles.fetch_add(1, std::memory_order_relaxed);
+				}
+				else if (SplitMask != 0)
+				{
+					NumNewInteriorEdges.fetch_add(2, std::memory_order_relaxed);
+					NumNewTriangles.fetch_add(2, std::memory_order_relaxed);
+				}
+			}, ParallelForFlags );
+
+			
+			// for each green refined triangles, store the edge mask for each neighboring triangle that is marked for refinement
+			ParallelFor( TEXT("TParallelAdaptiveRefinement.IdentifyGreenNeighbors.PF"), TriSplitQueueSize.load(), ParallelForBatchSize, [&](int32 Idx)
+			{
+				const int32 TriIndex = TriSplitQueue[Idx];
+				check(ConcurrencyState.TriSplitRecords[TriIndex].SplitMask != 0);
+
+				if (ConcurrencyState.TriSplitRecords[TriIndex].IsSplit1())
+				{
+					// mark neighbors
+					for (int32 NbIdx=0; NbIdx<3; ++NbIdx)
+					{
+						const TPair<int32,int32> TriEdge = Mesh.GetAdjEdge(TriIndex, NbIdx);
+						if (TriEdge.Get<0>() != IndexConstants::InvalidID && ConcurrencyState.TriSplitRecords[TriEdge.Get<0>()].SplitMask != 0)
+						{
+							ConcurrencyState.TriSplitRecords[TriEdge.Get<0>()].GreenAdjMask.fetch_or(static_cast<uint16>(1u << TriEdge.Get<1>()));
+						}
+					}
+				}
+			});
+
+			// grow the arrays for edges and vertices
+			int32 PrevEdgeCount, PrevVtxCount;
+			{
+				TRACE_CPUPROFILER_EVENT_SCOPE_TEXT("TParallelAdaptiveRefinement.GrowVerts");
+
+				const int32 NewVertexCount = NumEdgesToSplit.load();
+				const int32 NewEdgeCount = NewVertexCount + NumNewInteriorEdges.load();
+
+				ConcurrencyState.EdgeCount.store(Mesh.EdgeCount() + NewVertexCount);
+
+				PrevEdgeCount = Mesh.AddEdges(NewEdgeCount);
+				PrevVtxCount = Mesh.AddVertices(NewVertexCount);
+
+				ConcurrencyState.EdgeSplitRecords.AddDefaulted( NewEdgeCount );
+				check(ConcurrencyState.EdgeSplitRecords.Num() == Mesh.EdgeCount());
+
+				RefinementPolicy.AllocateVertices(NewVertexCount);
+			}
+
+			// assign new edge-midpoints, assign displacement
+			ParallelForTemplate(TEXT("TParallelAdaptiveRefinement.MakeEdgeMidpoints.PF"), NumEdgesToSplit.load(), ParallelForBatchSize, [&](const int32 RequestIdx)
+			{
+				const int32 EdgeIndex = EdgeSplitRequests[RequestIdx];
+				FEdgeSplitRecord& Edge = ConcurrencyState.EdgeSplitRecords[EdgeIndex];
+
+				Edge.NewEdge = PrevEdgeCount + RequestIdx;
+				Edge.SplitVertex = PrevVtxCount + RequestIdx;
+
+				FIndex2i VertexIndices, Tris;
+
+				Mesh.GetEdge(EdgeIndex, VertexIndices, Tris);
+				Mesh.InterpolateVertex(0.5f, Edge.SplitVertex, EdgeIndex);
+				RefinementPolicy.VertexAdded(Edge.SplitVertex, Tris);
+
+			}, ParallelForFlags  );
+
+			ConcurrencyState.HalfEdgeCount.store(Mesh.HalfEdgeCount());
+			{
+				TRACE_CPUPROFILER_EVENT_SCOPE_TEXT("TParallelAdaptiveRefinement.GrowHalfEdges");
+
+				Mesh.AddTriangles(NumNewTriangles.load());
+				RefinementCand.AddDefaulted(NumNewTriangles.load());
+			}
+
+			// walk through refinement queue, apply refinement, enqueue triangles to check for refinement
+			ParallelFor( TEXT("TParallelAdaptiveRefinement.RefineTri.PF"), TriSplitQueueSize.load(), ParallelForBatchSize, [&](int32 Idx)
+			{
+				const int32 TriIndex = TriSplitQueue[Idx];
+				check(ConcurrencyState.TriSplitRecords[TriIndex].SplitMask != 0);
+
+				if (ConcurrencyState.TriSplitRecords[TriIndex].IsSplit3())
+				{
+					const int32 FirstNewTri = Split3(TriIndex, ConcurrencyState);
+
+					// enqueue all 4 children as refinement candidates
+					const int32 QueueIdx = NumRefinementCand.fetch_add(4, std::memory_order_relaxed);
+					RefinementCand[QueueIdx+0] = FirstNewTri+0;
+					RefinementCand[QueueIdx+1] = FirstNewTri+1;
+					RefinementCand[QueueIdx+2] = FirstNewTri+2;
+					RefinementCand[QueueIdx+3] = TriIndex;
+				}
+				else if (ConcurrencyState.TriSplitRecords[TriIndex].IsSplit2())
+				{
+					Split2(TriIndex, ConcurrencyState);
+				}
+				else if (ConcurrencyState.TriSplitRecords[TriIndex].IsSplit1())
+				{
+					Split1(TriIndex, ConcurrencyState);
+				}
+
+				// All triangle split records should be in clean state after this loop
+				ConcurrencyState.TriSplitRecords[TriIndex].Reset();
+			}, ParallelForFlags );
+
+			check(ConcurrencyState.HalfEdgeCount.load() == Mesh.HalfEdgeCount());
+
+			// Mesh.CheckConsistency();
+			
+			TriSplitQueueSize.store(0);
+		}
+	}
+
+	// split 1 edge, keep two edges (1 new triangle)
+	inline void Split1(int32 TriIndex, FConcurrencyState& ConcurrencyState)
+	{
+		const int32 SplitMask = ConcurrencyState.TriSplitRecords[TriIndex].SplitMask;
+		check(ConcurrencyState.TriSplitRecords[TriIndex].IsSplit1());
+
+		// determine the edge that is split
+		int32 LocSplitEdge = 0;
+		if (SplitMask == 2) LocSplitEdge = 1;
+		else if (SplitMask == 4) LocSplitEdge = 2;
+
+		const int32 NextLocEdge = (LocSplitEdge+1)%3;
+		const int32 PrevLocEdge = (LocSplitEdge+2)%3;
+
+		const int32 EdgeIndices[3] = { Mesh.GetEdgeIndex(TriIndex, 0),
+									   Mesh.GetEdgeIndex(TriIndex, 1),
+									   Mesh.GetEdgeIndex(TriIndex, 2) };
+
+		const FIndex3i BaseTri = Mesh.GetTriangle(TriIndex);
+
+		const int32 SplitEdgeIdx = Mesh.GetEdgeIndex(TriIndex, LocSplitEdge);
+		const int32 SplitVertexIdx = ConcurrencyState.EdgeSplitRecords[SplitEdgeIdx].SplitVertex;
+		check(SplitVertexIdx != IndexConstants::InvalidID);
+
+		const int32 AdjHalfEdges[3] = { Mesh.GetAdjHalfEdge(TriIndex, 0),
+										Mesh.GetAdjHalfEdge(TriIndex, 1),
+										Mesh.GetAdjHalfEdge(TriIndex, 2) };
+
+		check(AdjHalfEdges[LocSplitEdge] != IndexConstants::InvalidID);
+
+		// the adjacency for the outer edge doesn't change
+		FIndex3i NewTri = BaseTri;
+		NewTri[NextLocEdge] = SplitVertexIdx;
+
+		Mesh.SetTriangle(TriIndex, NewTri, TriIndex);
+
+		// allocate 3 new half-edges and indices. this is the first index to write to
+		const int32 HalfEdge = ConcurrencyState.HalfEdgeCount.fetch_add(3, std::memory_order_relaxed);
+		check(HalfEdge + 2 < Mesh.MaxTriID() * 3);
+
+		// first new triangle
+
+		const int32 NewTriIndex = HalfEdge / 3;
+		Mesh.SetTriangle(NewTriIndex, FIndex3i(SplitVertexIdx, BaseTri[NextLocEdge], BaseTri[PrevLocEdge]), TriIndex);
+
+		// for the interior edge
+		const int32 NewEdgeIndex = ConcurrencyState.EdgeCount.fetch_add(1, std::memory_order_relaxed);
+		check(NewEdgeIndex < Mesh.EdgeCount());
+
+		// interior edge
+		Mesh.LinkEdge(HalfEdge + 2, TriIndex * 3 + NextLocEdge, NewEdgeIndex, IndexConstants::InvalidID);
+
+		struct FLinkTask
+		{
+			int32 HalfEdge;     //< halfedge to relink
+			int32 LocalEdge;    //< local edge index
+			int32 AdjHalfEdge;  //< adjacent half-edge before the update
+			int32 Edge;         //<
+		};
+
+		const FLinkTask LinkTasks[2] = { {TriIndex * 3 + PrevLocEdge, PrevLocEdge, AdjHalfEdges[PrevLocEdge], EdgeIndices[PrevLocEdge] },
+										 {HalfEdge + 1              , NextLocEdge, AdjHalfEdges[NextLocEdge], EdgeIndices[NextLocEdge] } };
+
+
+		// update adjacency for the unsplit edges
+		for (const FLinkTask& LinkTask : LinkTasks )
+		{
+			// only if the neighbor is of type split-1, we need a safe concurrent update. split-2 is constructed such that the unsplit
+			// edge doesn't require an adjacency update
+			if (LinkTask.AdjHalfEdge == IndexConstants::InvalidID || !(ConcurrencyState.TriSplitRecords[TriIndex].GreenAdjMask.load() & static_cast<uint16>(1 << LinkTask.LocalEdge)))
+			{
+				Mesh.LinkEdge(LinkTask.HalfEdge, LinkTask.AdjHalfEdge, LinkTask.Edge, LinkTask.Edge);
+			}
+			else
+			{
+				FEdgeSplitRecord& Edge = ConcurrencyState.EdgeSplitRecords[LinkTask.Edge];
+				const int32 EntryCountPreWrite = Edge.EntryCountPreWrite.fetch_add(1, std::memory_order_relaxed);
+				check(EntryCountPreWrite < 2);
+
+				Edge.HalfEdges[EntryCountPreWrite] = FIndex2i( LinkTask.HalfEdge, IndexConstants::InvalidID );
+
+				if (1 == Edge.EntryCountPostWrite.fetch_add(1, std::memory_order_acq_rel))
+				{
+					check(Edge.HalfEdges[0][1] == IndexConstants::InvalidID);
+					check(Edge.HalfEdges[1][1] == IndexConstants::InvalidID);
+
+					// second entry is being written, so we can now link the edges
+					Mesh.LinkEdge(Edge.HalfEdges[0][0], Edge.HalfEdges[1][0], LinkTask.Edge, LinkTask.Edge);
+				}
+			}
+		}
+
+		// update adjacency for the split edge
+		FIndex2i BoundaryHalfEdges = { 3 * TriIndex + LocSplitEdge, HalfEdge + 0 };
+		{
+			const int32 EdgeIdx = EdgeIndices[LocSplitEdge];
+			FEdgeSplitRecord& Edge = ConcurrencyState.EdgeSplitRecords[EdgeIdx];
+
+			const int32 EntryCountPreWrite = Edge.EntryCountPreWrite.fetch_add(1, std::memory_order_relaxed);
+			check(EntryCountPreWrite < 2);
+
+			Edge.HalfEdges[EntryCountPreWrite] = BoundaryHalfEdges;
+
+			// second atomic necessary to make sure the memory written to in other thread in slot 0 is available
+			// todo: might be faster to use std::memory_order_relaxed and use a atomic_thread_fence
+			if (1 == Edge.EntryCountPostWrite.fetch_add(1, std::memory_order_acq_rel))
+			{
+				check(Edge.HalfEdges[0][0] != IndexConstants::InvalidID);
+				check(Edge.HalfEdges[0][1] != IndexConstants::InvalidID);
+				check(Edge.HalfEdges[1][0] != IndexConstants::InvalidID);
+				check(Edge.HalfEdges[1][1] != IndexConstants::InvalidID);
+
+				// second entry is being written, so we can now link the edges
+				Mesh.LinkEdge(Edge.HalfEdges[0][0], Edge.HalfEdges[1][1], EdgeIdx, EdgeIdx);
+				Mesh.LinkEdge(Edge.HalfEdges[0][1], Edge.HalfEdges[1][0], Edge.NewEdge, EdgeIdx);
+			}
+		}
+	}
+
+	// split 2 edges, keep one edge (2 new triangles)
+	inline void Split2(int32 TriIndex, FConcurrencyState& ConcurrencyState)
+	{
+		const FIndex3i BaseTri = Mesh.GetTriangle(TriIndex);
+
+		const int32 SplitMask = ConcurrencyState.TriSplitRecords[TriIndex].SplitMask;
+
+		int32 LocPreserveEdge = 0; // the one edge that is NOT split
+		if (SplitMask == 3) {
+			LocPreserveEdge = 2;
+		} else if (SplitMask == 5) {
+			LocPreserveEdge = 1;
+		} else {
+			check(SplitMask == 6);
+			LocPreserveEdge = 0;
+		}
+
+		// edges to reconnect
+		const int32 EdgeIndices[2] = { Mesh.GetEdgeIndex(TriIndex, (LocPreserveEdge+1)%3),
+									   Mesh.GetEdgeIndex(TriIndex, (LocPreserveEdge+2)%3) };
+
+		const int32 SplitVertexIdx0 = ConcurrencyState.EdgeSplitRecords[EdgeIndices[0]].SplitVertex;
+		const int32 SplitVertexIdx1 = ConcurrencyState.EdgeSplitRecords[EdgeIndices[1]].SplitVertex;
+
+		const int32 h0 = TriIndex * 3 + (LocPreserveEdge+0);
+		const int32 h1 = TriIndex * 3 + (LocPreserveEdge+1)%3;
+		const int32 h2 = TriIndex * 3 + (LocPreserveEdge+2)%3;
+
+		// the tri connected to the unsplit edge
+		FIndex3i NewTri = BaseTri;
+		NewTri[(LocPreserveEdge + 2) % 3] = SplitVertexIdx1;
+		Mesh.SetTriangle(TriIndex, NewTri, TriIndex);
+
+		// allocate 6 new half-edges and indices. this is the first index to write to
+		const int32 HalfEdge = ConcurrencyState.HalfEdgeCount.fetch_add(6, std::memory_order_relaxed);
+		check(HalfEdge + 5 < Mesh.MaxTriID() * 3);
+
+		// two new interior edges
+		const int32 NewEdgeIndex = ConcurrencyState.EdgeCount.fetch_add(2, std::memory_order_relaxed);
+
+
+		const int32 NewTriIndex = HalfEdge / 3;
+
+		// T0
+		Mesh.SetTriangle(NewTriIndex + 0, FIndex3i(BaseTri[(LocPreserveEdge + 2) % 3], SplitVertexIdx1, SplitVertexIdx0), TriIndex);
+
+		// T1
+		Mesh.SetTriangle(NewTriIndex + 1, FIndex3i(SplitVertexIdx0, SplitVertexIdx1, BaseTri[(LocPreserveEdge + 1) % 3]), TriIndex);
+
+		// interior edges
+		Mesh.LinkEdge(HalfEdge + 1, HalfEdge + 3, NewEdgeIndex  , IndexConstants::InvalidID);
+		Mesh.LinkEdge(HalfEdge + 4, h1          , NewEdgeIndex+1, IndexConstants::InvalidID);
+
+		const FIndex2i BoundaryHalfEdges[2] = { { HalfEdge + 5, HalfEdge + 2 },
+												{ HalfEdge + 0, h2           } };
+
+		for (int32 BoundaryIdx=0; BoundaryIdx<2; ++BoundaryIdx)
+		{
+			const int32 EdgeIdx = EdgeIndices[BoundaryIdx];
+			FEdgeSplitRecord& Edge = ConcurrencyState.EdgeSplitRecords[EdgeIdx];
+
+			const int32 EntryCountPreWrite = Edge.EntryCountPreWrite.fetch_add(1, std::memory_order_relaxed);
+			Edge.HalfEdges[EntryCountPreWrite] = BoundaryHalfEdges[BoundaryIdx];
+
+			// second atomic necessary to make sure the memory written to in other thread in slot 0 is available
+			// todo: might be faster to use std::memory_order_relaxed and use a atomic_thread_fence (not x86-64)
+			if (1 == Edge.EntryCountPostWrite.fetch_add(1, std::memory_order_acq_rel))
+			{
+				check( Edge.HalfEdges[0][0] != IndexConstants::InvalidID );
+				check( Edge.HalfEdges[0][1] != IndexConstants::InvalidID );
+				check( Edge.HalfEdges[1][0] != IndexConstants::InvalidID );
+				check( Edge.HalfEdges[1][1] != IndexConstants::InvalidID );
+
+				// second entry is being written, so we can now link the edges
+				Mesh.LinkEdge(Edge.HalfEdges[0][0], Edge.HalfEdges[1][1], EdgeIdx, EdgeIdx);
+				Mesh.LinkEdge(Edge.HalfEdges[0][1], Edge.HalfEdges[1][0], Edge.NewEdge, EdgeIdx);
+			}
+		}
+	}
+
+	// split 3 edges, triangle into 4 triangles (3 new triangles)
+	// return first index of newly added triangles
+	inline int32 Split3(int32 TriIndex, FConcurrencyState& ConcurrencyState)
+	{
+		const FIndex3i BaseTri = Mesh.GetTriangle(TriIndex);
+
+		const int32 BaseEdgeIndices[3] = { Mesh.GetEdgeIndex(TriIndex, 0),
+										   Mesh.GetEdgeIndex(TriIndex, 1),
+										   Mesh.GetEdgeIndex(TriIndex, 2) };
+
+		const int32 BaseAdj[3] = { Mesh.GetAdjTriangle(TriIndex, 0),
+								   Mesh.GetAdjTriangle(TriIndex, 1),
+								   Mesh.GetAdjTriangle(TriIndex, 2) };
+
+		const int32 SplitVertexIdx0 = ConcurrencyState.EdgeSplitRecords[BaseEdgeIndices[0]].SplitVertex;
+		const int32 SplitVertexIdx1 = ConcurrencyState.EdgeSplitRecords[BaseEdgeIndices[1]].SplitVertex;
+		const int32 SplitVertexIdx2 = ConcurrencyState.EdgeSplitRecords[BaseEdgeIndices[2]].SplitVertex;
+
+		// replace the base triangle with the one in the middle
+		Mesh.SetTriangle(TriIndex, FIndex3i(SplitVertexIdx0, SplitVertexIdx1, SplitVertexIdx2), TriIndex);
+
+		// allocate 9 new half-edges and indices. this is the first index to write to
+		const int32 HalfEdge = ConcurrencyState.HalfEdgeCount.fetch_add(9, std::memory_order_relaxed);
+		const int32 NewTriIndex = HalfEdge / 3;
+
+		// T0
+		Mesh.SetTriangle(NewTriIndex + 0, FIndex3i(BaseTri[0], SplitVertexIdx0, SplitVertexIdx2), TriIndex);
+
+		// T1
+		Mesh.SetTriangle(NewTriIndex + 1, FIndex3i(BaseTri[1], SplitVertexIdx1, SplitVertexIdx0), TriIndex);
+
+		// T2
+		Mesh.SetTriangle(NewTriIndex + 2, FIndex3i(BaseTri[2], SplitVertexIdx2, SplitVertexIdx1), TriIndex);
+
+		const int32 NewEdgeIndex = ConcurrencyState.EdgeCount.fetch_add(3, std::memory_order_relaxed);
+		check(NewEdgeIndex < Mesh.EdgeCount());
+
+		// interior edges
+		Mesh.LinkEdge( HalfEdge + 1, TriIndex * 3 + 2, NewEdgeIndex + 0, IndexConstants::InvalidID );
+		Mesh.LinkEdge( HalfEdge + 4, TriIndex * 3 + 0, NewEdgeIndex + 1, IndexConstants::InvalidID );
+		Mesh.LinkEdge( HalfEdge + 7, TriIndex * 3 + 1, NewEdgeIndex + 2, IndexConstants::InvalidID );
+
+		// pairs of halfedges corresponding to the split edge as seen from this triangle
+		const FIndex2i BoundaryHalfEdges[3] = { { HalfEdge + 0, HalfEdge + 5 },
+												{ HalfEdge + 3, HalfEdge + 8 },
+												{ HalfEdge + 6, HalfEdge + 2 } };
+
+		for (int32 LocalEdgeIdx=0; LocalEdgeIdx<3; ++LocalEdgeIdx)
+		{
+			const int32 EdgeIdx = BaseEdgeIndices[LocalEdgeIdx];
+			FEdgeSplitRecord& Edge = ConcurrencyState.EdgeSplitRecords[EdgeIdx];
+
+			if (BaseAdj[LocalEdgeIdx] == IndexConstants::InvalidID)
+			{
+				// boundary case: we can just update locally without needing data from the other side
+				Mesh.LinkEdge(BoundaryHalfEdges[LocalEdgeIdx][0], IndexConstants::InvalidID, EdgeIdx, EdgeIdx);
+				Mesh.LinkEdge(BoundaryHalfEdges[LocalEdgeIdx][1], IndexConstants::InvalidID, Edge.NewEdge, EdgeIdx);
+				continue;
+			}
+
+			const int32 EntryCountPreWrite = Edge.EntryCountPreWrite.fetch_add(1, std::memory_order_relaxed);
+			Edge.HalfEdges[EntryCountPreWrite] = BoundaryHalfEdges[LocalEdgeIdx];
+
+			// second atomic necessary to make sure the memory written to in other thread in slot 0 is available
+			// todo: might be faster to use std::memory_order_relaxed and use a atomic_thread_fence
+			if (1 == Edge.EntryCountPostWrite.fetch_add(1, std::memory_order_acq_rel))
+			{
+				check(Edge.HalfEdges[0][0] != IndexConstants::InvalidID);
+				check(Edge.HalfEdges[0][1] != IndexConstants::InvalidID);
+				check(Edge.HalfEdges[1][0] != IndexConstants::InvalidID);
+				check(Edge.HalfEdges[1][1] != IndexConstants::InvalidID);
+
+				// second entry is being written, so we can now link the edges
+				Mesh.LinkEdge(Edge.HalfEdges[0][0], Edge.HalfEdges[1][1], EdgeIdx, EdgeIdx);
+				Mesh.LinkEdge(Edge.HalfEdges[0][1], Edge.HalfEdges[1][0], Edge.NewEdge, EdgeIdx);
+			}
+		}
+		check(HalfEdge % 3 == 0);
+		return HalfEdge / 3;
+	}
+
+	[[nodiscard]] bool ShouldRefine(int32 TriIndex, int32 Iter) const
+	{
+		FVector3f SplitVertexBary;
+		return RefinementPolicy.ShouldRefine(TriIndex, SplitVertexBary, Iter);
+	}
+
+private:
+
+	MeshType&               Mesh;
+	RefinementPolicyType&   RefinementPolicy;
+};
+
+} // namespace Geometry
+} // namespace UE
\ No newline at end of file
