Buildhealth: add missing typename

#rnx

diff --git a/Engine/Source/Runtime/GeometryCore/Public/Tessellation/ParallelAdaptiveRefinement.h b/Engine/Source/Runtime/GeometryCore/Public/Tessellation/ParallelAdaptiveRefinement.h
--- a/Engine/Source/Runtime/GeometryCore/Public/Tessellation/ParallelAdaptiveRefinement.h
+++ b/Engine/Source/Runtime/GeometryCore/Public/Tessellation/ParallelAdaptiveRefinement.h
@@ -107,15 +107,15 @@ private:
 	};
 
 public:
-	using RealType = MeshType::RealType; // the field related to the vector-space
-	using VecType  = MeshType::VecType;  // for positions and normals
+	using RealType = typename MeshType::RealType; // the field related to the vector-space
+	using VecType  = typename MeshType::VecType;  // for positions and normals
 
 	// Initialize and run adaptive aptive tesselation
 	TParallelAdaptiveRefinement(
 		MeshType& InMesh,                               //< mesh interface
 		RefinementPolicyType& InRefinementPolicy,       //< displacement and error interface
-		const int32 ParallelForBatchSize = 1024,     //<
-		const int32 ReserveVertices = 2'000'000 )
+		const int32 ParallelForBatchSize = 1024,        //<
+		const int32 ReserveVertices = 2'000'000 )       //< pre-reserve enough space for growing up to this many vertices
 		: Mesh(InMesh)
 		, RefinementPolicy(InRefinementPolicy)
 	{
