Knot-vector API updates

diff --git a/Engine/Source/Runtime/Engine/Private/Curves/Spline.cpp b/Engine/Source/Runtime/Engine/Private/Curves/Spline.cpp
--- a/Engine/Source/Runtime/Engine/Private/Curves/Spline.cpp
+++ b/Engine/Source/Runtime/Engine/Private/Curves/Spline.cpp
@@ -452,11 +452,11 @@ int32 FNewSpline::SetAttributeParameter(const FName& Name, int32 Index, float Pa
 			const float KnotOffset = -1.f * ChannelSpace.Min;
 			const float KnotScaleFactor = 1.f / (ChannelSpace.Max - ChannelSpace.Min);
 
-			TArray<float> Knots = Channel->GetKnotVector();
-			for (float& Knot : Knots)
+			TArray<FKnot> Knots = Channel->GetKnotVector();
+			for (FKnot& Knot : Knots)
 			{
-				Knot += KnotOffset;
-				Knot *= KnotScaleFactor;
+				Knot.Value += KnotOffset;
+				Knot.Value *= KnotScaleFactor;
 			}
 			Channel->SetKnotVector(Knots);
 		};
diff --git a/Engine/Source/Runtime/GeometryCore/Public/Splines/BSpline.h b/Engine/Source/Runtime/GeometryCore/Public/Splines/BSpline.h
--- a/Engine/Source/Runtime/GeometryCore/Public/Splines/BSpline.h
+++ b/Engine/Source/Runtime/GeometryCore/Public/Splines/BSpline.h
@@ -35,6 +35,45 @@ struct FKnot
 	FKnot(float InValue, uint32 InMultiplicity = 1) 
 		: Value(InValue), Multiplicity(InMultiplicity) {}
 
+	/**
+	 * Converts a pair-based knot vector to a flat knot vector
+	 * @param PairKnots - The pair-based knot vector
+	 * @return Flat knot vector with multiplicities expanded
+	 */
+	static TArray<float> ConvertPairToFlatKnots(const TArray<FKnot>& PairKnots)
+	{
+		TArray<float> FlatKnots;
+		for (const FKnot& Knot : PairKnots)
+		{
+			for (uint32 i = 0; i < Knot.Multiplicity; ++i)
+			{
+				FlatKnots.Add(Knot.Value);
+			}
+		}
+		return FlatKnots;
+	}
+
+	static bool IsValidKnotVector(const TArray<FKnot>& InKnots, int32 Degree)
+	{
+		TArray<float> FlatKnots = ConvertPairToFlatKnots(InKnots);
+		// Need at least 2*(degree+1) knots
+		if (FlatKnots.Num() < 2*(Degree + 1))
+		{
+			return false;
+		}
+
+		// Check monotonicity
+		for (int32 i = 1; i < FlatKnots.Num(); ++i)
+		{
+			if (FlatKnots[i] < FlatKnots[i-1])
+			{
+				return false;
+			}
+		}
+
+		return true;
+	}
+
 	bool operator==(const FKnot& Other) const
 	{
 		return FMath::IsNearlyEqual(Value, Other.Value) && Multiplicity == Other.Multiplicity;
@@ -85,9 +124,15 @@ public:
 	TBSpline() = default;
 	virtual ~TBSpline() override = default;
 
+	template<typename T, typename = void>
+	struct THasToString : std::false_type {};
+
+	template<typename T>
+	struct THasToString<T, std::void_t<decltype(std::declval<T>().ToString())>> : std::true_type {};
+	
 	void Dump() const
 	{
-		if constexpr (std::is_same_v<ValueType, float>)
+		if constexpr (std::is_floating_point_v<ValueType>)
 		{
 			UE_LOG(LogSpline, Verbose, TEXT("Values:"));
 			for (int32 Index = 0; Index < Values.Num(); ++Index)
@@ -95,13 +140,18 @@ public:
 				UE_LOG(LogSpline, Verbose, TEXT("  [%d] %f"), Index, Values[Index]);
 			}
 		}
-
-		UE_LOG(LogSpline, Verbose, TEXT("Knots:"));
-		for (int32 Index = 0; Index < PairKnots.Num(); ++Index)
+		else if constexpr (THasToString<ValueType>::value)
 		{
-			const FKnot& Knot = PairKnots[Index];
-			UE_LOG(LogSpline, Verbose, TEXT("  [%d] Value: %f, Multiplicity: %u"), Index, Knot.Value, Knot.Multiplicity);
+			UE_LOG(LogSpline, Verbose, TEXT("Values:"));
+			for (int32 Index = 0; Index < Values.Num(); ++Index)
+			{
+				UE_LOG(LogSpline, Verbose, TEXT("  [%d] %s"), Index, *Values[Index].ToString());
+			}
 		}
+
+		UE_LOG(LogSpline, Verbose, TEXT("Knots:"));
+		PrintKnotVector();
+		
 	}
 	
 	// ISplineInterface Implementation
@@ -227,7 +277,7 @@ public:
 	// Returns parameter value for a control point index using Greville Abscissae
     virtual float GetParameter(int32 Index) const
 	{
-		EnsureFlatKnotsCache();
+		UpdateFlatKnotsCache();
 		if (Index < 0 || Index + Degree >= FlatKnots.Num() - 1)
 			return 0.0f;
 
@@ -242,7 +292,7 @@ public:
 	// Find the control point index and local parameter for global parameter
 	virtual int32 FindIndexForParameter(float Parameter, float& OutLocalParam) const
 	{
-		EnsureFlatKnotsCache();
+		UpdateFlatKnotsCache();
         const TArray<float>& Knots = FlatKnots;
 		
 		if (Knots.IsEmpty())
@@ -297,10 +347,9 @@ public:
 		return Index;
 	}
 	// Knot vector management
-	const TArray<float>& GetKnotVector() const
+	const TArray<FKnot>& GetKnotVector() const
 	{
-		EnsureFlatKnotsCache();
-		return FlatKnots;
+		return PairKnots;
 	}
 	// Get the pair-based knot vector
     const TArray<FKnot>& GetPairKnots() const
@@ -308,68 +357,285 @@ public:
         return PairKnots;
     }
 	void ResetKnotVector() { PairKnots.Reset(); FlatKnots.Reset(); }
+		
 	/**
-     * Gets the array of distinct knot values (without duplicates)
-     * @return Array of distinct knot values in ascending order
-     */
-    TArray<float> GetDistinctKnotValues() const
-    {
-		TArray<float> Distinct;
-		for (const FKnot& Knot : PairKnots)
+	 * Returns how many times a given knot value appears.
+	 * @param KnotIndex - Knot value to check
+	 * @return Number of times the knot appears
+	 */
+	int32 GetKnotMultiplicity(int32 KnotIndex) const
+	{
+		if (PairKnots.IsValidIndex(KnotIndex))
 		{
-			Distinct.Add(Knot.Value);
+			return PairKnots[KnotIndex].Multiplicity;
+		}
+		
+		return 0;
+	}
+
+	/**
+	 * Sets a custom knot vector using explicit knot/multiplicity pairs
+	 * @param NewKnots - Array of knot values. Must be non-decreasing and have correct length (see GetExpectedNumKnots)
+	 * @return True if the knot vector was valid and set successfully
+	 */
+	bool SetCustomKnots(const TArray<FKnot>& NewKnots)
+	{
+        const int32 RequiredKnots = GetExpectedNumKnots();
+		
+		// Validate knot multiplicities and calculate total expanded knots  
+		int32 TotalExpandedKnots = 0;
+		for (const FKnot& Knot : NewKnots)
+		{
+			if (Knot.Multiplicity == 0)
+			{
+				UE_LOG(LogSpline, Warning, TEXT("Invalid knot with zero multiplicity."));
+				return false;
+			}
+			TotalExpandedKnots += Knot.Multiplicity;
 		}
-		return Distinct;
-    }
     
-    /**
-     * Finds which distinct knot span contains a parameter
-     * @param Parameter - The parameter value to find
-     * @param OutStartKnot - Output parameter for span start value
-     * @param OutEndKnot - Output parameter for span end value
-     * @return True if a valid span was found
-     */
-    bool FindDistinctKnotSpan(float Parameter, float& OutStartKnot, float& OutEndKnot) const
+		if (TotalExpandedKnots < RequiredKnots)
+		{
+			UE_LOG(LogSpline, Warning, TEXT("Invalid knot vector size. Need %d knots, got %d."), 
+				RequiredKnots, TotalExpandedKnots);
+			return false;
+		}
+
+		// Verify knot values are non-decreasing (strict comparison, no epsilon)
+		for (int32 i = 1; i < NewKnots.Num(); ++i)
+		{
+			if (NewKnots[i].Value < NewKnots[i - 1].Value)
+			{
+				UE_LOG(LogSpline, Warning, TEXT("Invalid knot vector: knot values must be non-decreasing."));
+				return false;
+			}
+		}
+
+		// Directly assign the knots - no epsilon comparison needed
+		PairKnots = NewKnots;
+    
+		// Mark the flat knots cache as dirty
+		MarkFlatKnotsCacheDirty();
+    
+		UE_LOG(LogSpline, Verbose, TEXT("Set Custom knot vector - Points: %d, Degree: %d, Total Knots: %d, Unique Knots: %d"),
+			   NumKeys(), Degree, TotalExpandedKnots, NewKnots.Num());
+		PrintKnotVector();
+    
+		return true;
+	}
+
+	/**
+	 * Updates the knot vector based on current control points and settings
+	 * 
+	 * @param Points - Array of control points and their Knots
+	 * @param ParameterizationPolicy - Policy for generating knot vector
+	 */
+    virtual void Reparameterize(EParameterizationPolicy ParameterizationPolicy)
+	{
+    	const TArray<ValueType>& Points = Values;
+		int32 NumValues = Points.Num();
+		if (NumValues < Degree + 1)
+		{
+			UE_LOG(LogSpline, Warning, TEXT("Not enough points for B-spline: %d (need >= %d)"),
+				   NumValues, Degree + 1);
+			return;
+		}
+
+		const int32 KnotCount = NumValues + Degree + 1;
+		PairKnots.Reset();
+
+		switch (ParameterizationPolicy)
+		{
+		case EParameterizationPolicy::Uniform:
+			GenerateUniformKnots(KnotCount);
+			break;
+		case EParameterizationPolicy::ChordLength:
+			GenerateChordLengthKnots(KnotCount);
+			break;
+		case EParameterizationPolicy::Centripetal:
+			GenerateCentripetalKnots(KnotCount);
+			break;
+		}
+
+		// Mark the flat knots cache as dirty
+		MarkFlatKnotsCacheDirty();
+
+		UE_LOG(LogSpline, Verbose, TEXT("Updated knot vector - Points: %d, Degree: %d, Knots: %d"),
+			   Points.Num(), Degree, PairKnots.Num());
+	}
+	
+	/** 
+	 * Sets whether the spline should be clamped at endpoints
+	 * When true, the curve will interpolate the first and last control points
+	 */
+	void SetClampedEnds(bool bInClampEnds)
+	{
+		bClampEnds = bInClampEnds;
+	}
+
+	bool IsClampedEnds() const
+	{
+		return bClampEnds;
+	}
+
+	// Get the valid parameter range in knot space
+	FInterval1f GetKnotRange() const
+	{
+		UpdateFlatKnotsCache();
+		if (FlatKnots.Num() - Degree - 1 < 0)
+		{
+			return FInterval1f::Empty();
+		}
+		return FInterval1f(PairKnots[0].Value, PairKnots.Last().Value);
+	}
+	
+protected:
+	
+	/**
+	 *  Sets the value of a knot
+	 *  @param KnotIdx - Index of the knot to set
+	 *  @param NewValue - Value to set
+	**/
+	void SetKnot(int32 KnotIdx, float NewValue)
+	{
+		if (KnotIdx < 0 || KnotIdx >= PairKnots.Num() || PairKnots[KnotIdx].Value == NewValue)
+		{
+			// Nothing to replace
+			return;
+		}
+
+		PairKnots[KnotIdx].Value = NewValue;
+	}
+
+	/**
+	 * Removes a knot at the specified index.
+	 * 
+	 * @param KnotIdx - Index of the knot to remove
+	 * @return True if the knot was successfully removed, false if index was invalid
+	 */
+	bool RemoveKnot(int32 KnotIdx)
+	{
+		if (KnotIdx < 0 || KnotIdx >= PairKnots.Num())
+		{
+			return false;
+		}
+		PairKnots.RemoveAt(KnotIdx);
+		// Mark the flat knots cache as dirty
+		MarkFlatKnotsCacheDirty();
+		return true;
+	}
+
+	/**
+	 * Swaps two knots.
+	 * 
+	 * @param KnotIdxA - Index of first knot to swap.
+	 * @param KnotIdxB - Index of second knot to swap.
+	 */
+	void SwapKnots(int32 KnotIdxA, int32 KnotIdxB)
+	{
+		if (KnotIdxA < 0 || KnotIdxB < 0 || KnotIdxA >= PairKnots.Num() || KnotIdxB >= PairKnots.Num())
+		{
+			return;
+		}
+		
+		FKnot OldA = PairKnots[KnotIdxA];
+		PairKnots[KnotIdxA] = PairKnots[KnotIdxB];
+		PairKnots[KnotIdxB] = OldA;
+	}
+	
+	/**
+	 * Inserts a knot at the specified parameter value
+	 * 
+	 * @param InKnot - Knot to insert
+	 */
+	void InsertKnot(FKnot InKnot)
+	{
+		for (int32 i = 0; i < PairKnots.Num(); ++i)
+		{
+			if (PairKnots[i].Value == InKnot.Value)  // exact match
+			{
+				return;
+			}
+			else if (PairKnots[i].Value > InKnot.Value)
+			{
+				PairKnots.Insert(InKnot, i);
+				MarkFlatKnotsCacheDirty();
+				return;
+			}
+		}
+
+		PairKnots.Add(InKnot);
+		MarkFlatKnotsCacheDirty();
+	}
+
+	/**
+	 * Finds the nearest available knot value that does not conflict with existing knots.
+	 * @param DesiredParameter - The desired parameter value to find a valid knot for.
+	 * @return The nearest valid knot parameter that does not conflict with existing knots.
+	 */
+	float GetNearestAvailableKnotValue(float DesiredParameter) const
     {
-        TArray<float> DistinctKnots = GetDistinctKnotValues();
-        
-        if (DistinctKnots.Num() < 2)
-            return false;
-            
-        // Handle parameter outside range
-        if (Parameter <= DistinctKnots[0])
-        {
-            OutStartKnot = DistinctKnots[0];
-            OutEndKnot = DistinctKnots[1];
-            return true;
-        }
-        
-        if (Parameter >= DistinctKnots.Last())
+        if (PairKnots.Num() == 0)
         {
-            OutStartKnot = DistinctKnots[DistinctKnots.Num() - 2];
-            OutEndKnot = DistinctKnots.Last();
-            return true;
+            return DesiredParameter;
         }
+
+        // Get valid parameter space
+        FInterval1f ParamSpace = GetParameterSpace();
+
+		// Convert existing knots into a set for O(1) lookups
+		TSet<float> KnotValuesSet;
+		for (const FKnot& Knot : PairKnots)
+		{
+			KnotValuesSet.Add(Knot.Value);
+		}
+		
+		// If no conflict, return immediately
+		if (!KnotValuesSet.Contains(DesiredParameter))
+		{
+			return DesiredParameter;
+		}
+		
+        // Find the closest valid parameter by walking outward alternately
+        float Left = DesiredParameter;
+        float Right = DesiredParameter;
+        bool bSearchLeft = true;
+        bool bSearchRight = true;
         
-        // Find span containing parameter
-        for (int32 i = 0; i < DistinctKnots.Num() - 1; ++i)
+		// Search until valid parameter is found
+        while (bSearchLeft || bSearchRight)
         {
-            if (Parameter >= DistinctKnots[i] && Parameter < DistinctKnots[i + 1])
-            {
-                OutStartKnot = DistinctKnots[i];
-                OutEndKnot = DistinctKnots[i + 1];
-                return true;
-            }
+        	if (bSearchLeft)
+        	{
+        		Left = std::nextafter(Left, -std::numeric_limits<float>::infinity());
+        		if (!KnotValuesSet.Contains(Left))
+        		{
+        			return Left;
+        		}
+        	}
+
+        	if (bSearchRight)
+        	{
+        		Right = std::nextafter(Right, std::numeric_limits<float>::infinity());
+        		if (!KnotValuesSet.Contains(Right))
+        		{
+        			return Right;
+        		}
+        	}
         }
-        
-        // Fallback - should not reach here
-        OutStartKnot = DistinctKnots[0];
-        OutEndKnot = DistinctKnots.Last();
-        return false;
+		// Fallback (shouldn't be reached normally)
+		return DesiredParameter;
     }
+
+	virtual int32 GetExpectedNumKnots() const
+	{
+		const int32 NumValues = NumKeys();
+		// For closed loops, we repeat the first Degree points at the end
+		return this->IsClosedLoop()
+				   ? (NumValues + (2 * Degree) + 1) // closed loop: n + 2p + 1
+				   : (NumValues + Degree + 1);
+	}
 	
-	
-    
 	/** 
      * Generates a uniform knot vector with equal spacing
      * For clamped ends, multiplicity of degree+1 is used at endpoints
@@ -641,529 +907,62 @@ public:
 	}
 
 	/**
-	 * Inserts a knot at the specified parameter value
-	 * 
-	 * @param Value - Knot value where the knot should be inserted
-	 * @param Count - Number of times to insert the knot
+	 * Applies clamped multiplicity to the knot vector.
+	 * This ensures that the first and last knots have Degree + 1 multiplicity
 	 */
-	void InsertKnot(float Value, int32 Count = 1)
+	void ApplyClampedKnotsMultiplicity()
 	{
-		for (int32 i = 0; i < PairKnots.Num(); ++i)
+		if (IsClampedEnds() && PairKnots.Num() > 0)
 		{
-			if (PairKnots[i].Value == Value)  // exact match
-			{
-				PairKnots[i].Multiplicity += Count;
-				MarkFlatKnotsCacheDirty();
-				return;
-			}
-			else if (PairKnots[i].Value > Value)
-			{
-				PairKnots.Insert(FKnot(Value, Count), i);
-				MarkFlatKnotsCacheDirty();
-				return;
-			}
-		}
-
-		PairKnots.Add(FKnot(Value, Count));
-		MarkFlatKnotsCacheDirty();
-	}
-
-	/**
-	 * Replaces all instances of a knot.
-	 * 
-	 * @param KnotIdx - Index of the knot to replace
-	 * @param NewValue - Value to replace instances of OldValue
-	 */
-	void SetKnot(int32 KnotIdx, float NewValue)
-	{
-		if (KnotIdx < 0 || KnotIdx >= PairKnots.Num() || PairKnots[KnotIdx].Value == NewValue)
-		{
-			// Nothing to replace
-			return;
-		}
-
-		PairKnots[KnotIdx].Value = NewValue;
-	}
-
-	/**
-	 * Swaps all instances of two knots.
-	 * 
-	 * @param KnotIdxA - Index of first knot to swap.
-	 * @param KnotIdxB - Index of second knot to swap.
-	 */
-	void SwapKnots(int32 KnotIdxA, int32 KnotIdxB)
-	{
-		if (KnotIdxA < 0 || KnotIdxB < 0 || KnotIdxA >= PairKnots.Num() || KnotIdxB >= PairKnots.Num())
-		{
-			return;
+			// enforce degree + 1 multiplicity at the endpoints
+			PairKnots[0].Multiplicity = Degree + 1;
+			PairKnots.Last().Multiplicity = Degree + 1;
+			MarkFlatKnotsCacheDirty();
 		}
-		
-		float OldA = PairKnots[KnotIdxA].Value;
-		PairKnots[KnotIdxA].Value = PairKnots[KnotIdxB].Value;
-		PairKnots[KnotIdxB].Value = OldA;
 	}
 
 	/**
-	 * Returns how many times a given knot value appears.
-	 * @param Value - Knot value to check
-	 * @return Number of times the knot appears
+	 * Updates the flat knot cache from the pair representation if the cache is dirty.
 	 */
-	int32 GetKnotMultiplicity(float Value) const
-	{
-		for (const FKnot& Knot : PairKnots)
-		{
-			if (Knot.Value == Value)
-			{
-				return static_cast<int32>(Knot.Multiplicity);
-			}
-		}
-		return 0;
-	}
-
-	
-	
-	bool RemoveKnotsCompletely(float Value, int32 ExpectedKnotCount)
-	{
-		const int32 Removed = PairKnots.RemoveAll([&](const FKnot& K)
-		{
-			return K.Value == Value;
-		});
-
-		if (Removed > 0)
-		{
-			// fix end knots if an end knot is removed
-			FixEndKnots(ExpectedKnotCount);
-			MarkFlatKnotsCacheDirty();
-			return true;
-		}
-
-		return false;
-
+	void UpdateFlatKnotsCache() const
+    {
+    	if (bFlatKnotsCacheDirty)
+    	{
+    		FlatKnots = FKnot::ConvertPairToFlatKnots(PairKnots);
+    		bFlatKnotsCacheDirty = false;
+    	}
+    }
 
-	}
 	/**
-	 * Removes a single occurrence of a knot value.
-	 * For clamped endpoints, we avoid removing knots if doing so would drop the multiplicity
-	 * below the required (Degree+1) at the endpoints.
-	 * @param Value - Knot value to remove
-	 * @param Count - Number of times to remove the knot
-	 * @return true if the knot was removed
+	 * Marks the flat knots cache as dirty.
 	 */
-	bool RemoveKnot(float Value, int32 Count = 1)
-	{
-		const TArray<FKnot> OriginalKnots = PairKnots;
-
-		for (int32 i = 0; i < PairKnots.Num(); ++i)
-		{
-			if (PairKnots[i].Value == Value)
-			{
-				bool bIsEndpoint =
-					(Value == GetKnotRange().Min && bClampEnds) ||
-					(Value == GetKnotRange().Max && bClampEnds);
-
-				uint32 MinRequiredMultiplicity = bIsEndpoint ? Degree + 1 : 1;
-				uint32 Current = PairKnots[i].Multiplicity;
-
-				if (Current <= MinRequiredMultiplicity)
-				{
-					UE_LOG(LogSpline, Warning, TEXT("Cannot remove knot %f below required multiplicity %u"), Value, MinRequiredMultiplicity);
-					return false;
-				}
-
-				uint32 NewCount = Current - Count;
-				if (NewCount < MinRequiredMultiplicity)
-				{
-					PairKnots.RemoveAt(i);
-				}
-				else
-				{
-					PairKnots[i].Multiplicity = NewCount;
-				}
-
-				// Validate
-				TArray<float> Flat = ConvertPairToFlatKnots(PairKnots);
-				if (!Math::FSplineValidation::IsValidKnotVector(Flat, Degree))
-				{
-					UE_LOG(LogSpline, Warning, TEXT("Knot removal resulted in invalid knot vector. Reverting."));
-					PairKnots = OriginalKnots;
-					return false;
-				}
-
-				MarkFlatKnotsCacheDirty();
-				return true;
-			}
-		}
-
-		return false;
-	}
+	void MarkFlatKnotsCacheDirty() const
+    {
+    	bFlatKnotsCacheDirty = true;
+    }
 
 	/**
-	 * Converts a pair-based knot vector to a flat knot vector
-	 * @param PairKnots - The pair-based knot vector
-	 * @return Flat knot vector with multiplicities expanded
+	 * Prints the knot vector to the log for debugging.
 	 */
-	static TArray<float> ConvertPairToFlatKnots(const TArray<FKnot>& PairKnots)
+	void PrintKnotVector() const
 	{
-		TArray<float> FlatKnots;
-		for (const FKnot& Knot : PairKnots)
+		for (int32 Index = 0; Index < PairKnots.Num(); ++Index)
 		{
-			for (uint32 i = 0; i < Knot.Multiplicity; ++i)
-			{
-				FlatKnots.Add(Knot.Value);
-			}
+			const FKnot& Knot = PairKnots[Index];
+			UE_LOG(LogSpline, Verbose, TEXT("  [%d] Value: %f, Multiplicity: %u"), Index, Knot.Value, Knot.Multiplicity);
 		}
-		return FlatKnots;
-	}
-	/**
-	 * Sets a custom knot vector 
-	 * @param NewKnots - Array of knot values. Must be non-decreasing and have correct length (n + p + 1)
-	 * @param NumValues - Number of values (control points) in the spline
-	 * @return True if the knot vector was valid and set successfully
-	 */
-	bool SetCustomKnots(const TArray<float>& NewKnots)
-	{
-    	int32 NumValues = NumKeys();
-		// For closed loops, we repeat the first Degree points at the end
-		const int32 RequiredKnots = GetExpectedNumKnots();    
 		
-		if (NewKnots.Num() < RequiredKnots)
-		{
-			UE_LOG(LogSpline, Warning, TEXT("Invalid knot vector size. Need %d knots, got %d."), 
-				RequiredKnots, NewKnots.Num());
-			return false;
-		}
-
-		// Verify knot vector is non-decreasing
-		for (int32 i = 1; i < NewKnots.Num(); ++i)
-		{
-			if (NewKnots[i] < NewKnots[i - 1])
-			{
-				UE_LOG(LogSpline, Warning, TEXT("Invalid knot vector: knots must be non-decreasing."));
-				return false;
-			}
-		}
-
-		// Convert flat knots to pair representation
-		PairKnots.Reset();
-    
-		if (NewKnots.Num() > 0)
-		{
-			float CurrentValue = NewKnots[0];
-			uint32 Multiplicity = 1;
-        
-			for (int32 i = 1; i < NewKnots.Num(); ++i)
-			{
-				if (FMath::IsNearlyEqual(NewKnots[i], CurrentValue))
-				{
-					// Same value, increase multiplicity
-					Multiplicity++;
-				}
-				else
-				{
-					// Different value, add previous knot and start new one
-					PairKnots.Add(FKnot(CurrentValue, Multiplicity));
-					CurrentValue = NewKnots[i];
-					Multiplicity = 1;
-				}
-			}
-        
-			// Add the last knot
-			PairKnots.Add(FKnot(CurrentValue, Multiplicity));
-		}
-    
-		// Mark the flat knots cache as dirty
-		MarkFlatKnotsCacheDirty();
-    
-		UE_LOG(LogSpline, Verbose, TEXT("Set Custom knot vector - Points: %d, Degree: %d, Knots: %d"),
-			   NumValues, Degree, NewKnots.Num());
-		for (int32 i = 0; i < NewKnots.Num(); ++i)
+		for (int32 j = 0; j < this->FlatKnots.Num(); ++j)
 		{
-			UE_LOG(LogSpline, Verbose, TEXT("  Knot[%d] = %f"), i, NewKnots[i]);
+			UE_LOG(LogSpline, Verbose, TEXT("  FlatKnot[%d] = %f"), j, this->FlatKnots[j]);
 		}
-    
-		
-		return true;
-	}
-
-	/**
-	 * Updates the knot vector based on current control points and settings
-	 * 
-	 * @param Points - Array of control points and their Knots
-	 * @param ParameterizationPolicy - Policy for generating knot vector
-	 */
-
-    virtual void Reparameterize(EParameterizationPolicy ParameterizationPolicy)
-	{
-    	const TArray<ValueType>& Points = Values;
-		int32 NumValues = Points.Num();
-		if (NumValues < Degree + 1)
-		{
-			UE_LOG(LogSpline, Warning, TEXT("Not enough points for B-spline: %d (need >= %d)"),
-				   NumValues, Degree + 1);
-			return;
-		}
-
-		const int32 KnotCount = NumValues + Degree + 1;
-		PairKnots.Reset();
-
-		switch (ParameterizationPolicy)
-		{
-		case EParameterizationPolicy::Uniform:
-			GenerateUniformKnots(KnotCount);
-			break;
-		case EParameterizationPolicy::ChordLength:
-			GenerateChordLengthKnots(KnotCount);
-			break;
-		case EParameterizationPolicy::Centripetal:
-			GenerateCentripetalKnots(KnotCount);
-			break;
-		}
-
-		// Mark the flat knots cache as dirty
-		MarkFlatKnotsCacheDirty();
-
-		UE_LOG(LogSpline, Verbose, TEXT("Updated knot vector - Points: %d, Degree: %d, Knots: %d"),
-			   Points.Num(), Degree, PairKnots.Num());
-	}
-	
-	/** 
-	 * Sets whether the spline should be clamped at endpoints
-	 * When true, the curve will interpolate the first and last control points
-	 */
-	void SetClampedEnds(bool bInClampEnds)
-	{
-		bClampEnds = bInClampEnds;
-	}
-
-	bool IsClampedEnds() const
-	{
-		return bClampEnds;
-	}
-
-	// Get the valid parameter range in knot space
-	FInterval1f GetKnotRange() const
-	{
-		TArray<float> Knots = GetKnotVector();
-		if (Knots.Num() - Degree - 1 < 0)
-		{
-			return FInterval1f::Empty();
-		}
-		if (bClampEnds)
-		{
-			// For clamped B-splines, the valid range is [Degree, NumKnots - Degree - 1]
-			return FInterval1f(Knots[Degree],Knots[Knots.Num() - Degree - 1]);
-		}
-		else
-		{
-			// For unclamped B-splines, the valid range is [FirstKnot, LastKnot]
-			return FInterval1f(Knots[0], Knots.Last());
-		}
-	}
-
-	// Ensure we have the correct number of knots at the endpoints
-	void FixEndKnots(const int32 ExpectedKnots)
-	{
-		if (PairKnots.Num() == 0)
-			return;
-        
-		// Calculate current flat knot count
-		int32 CurrentFlatCount = 0;
-		for (const FKnot& Knot : PairKnots)
-		{
-			CurrentFlatCount += static_cast<int32>(Knot.Multiplicity);
-		}
-    
-		// Calculate how many knots to add or remove
-		int32 Diff = ExpectedKnots - CurrentFlatCount;
-    
-		if (Diff == 0)
-		{
-			// Nothing to do
-			return;
-		}
-		else if (Diff > 0)
-		{
-			// Need to add knots - add them to the end with the same value as the last knot
-			FKnot& LastKnot = PairKnots.Last();
-			LastKnot.Multiplicity += static_cast<uint32>(Diff);
-		}
-		else // Diff < 0
-		{
-			// Need to remove knots - remove them from the end
-			uint32 KnotsToRemove = static_cast<uint32>(-Diff);
-        
-			// Start removing from the last knot
-			while (KnotsToRemove > 0 && PairKnots.Num() > 0)
-			{
-				FKnot& LastKnot = PairKnots.Last();
-            
-				if (LastKnot.Multiplicity <= KnotsToRemove)
-				{
-					// Remove the entire knot
-					KnotsToRemove -= LastKnot.Multiplicity;
-					PairKnots.RemoveAt(PairKnots.Num() - 1);
-				}
-				else
-				{
-					// Reduce multiplicity
-					LastKnot.Multiplicity -= KnotsToRemove;
-					KnotsToRemove = 0;
-				}
-			}
-		}
-    
-		MarkFlatKnotsCacheDirty();
-	}
-	// Fill the end of the knot vector with the specified knot value
-	void FillEndKnots(float NewEndKnot, const int32 ExpectedKnots)
-	{
-		// Calculate current flat knot count
-		int32 CurrentFlatCount = 0;
-		for (const FKnot& Knot : PairKnots)
-		{
-			CurrentFlatCount += static_cast<int32>(Knot.Multiplicity);
-		}
-    
-		// Calculate how many knots to add
-		int32 KnotsToAdd = FMath::Max(0, ExpectedKnots - CurrentFlatCount);
-    
-		if (KnotsToAdd == 0)
-		{
-			return;
-		}
-    
-		UE_LOG(LogSpline, Log, TEXT("Adding %d knots at new end parameter %f"),
-			   KnotsToAdd, NewEndKnot);
-    
-		// Check if the last knot has the same value
-		if (PairKnots.Num() > 0 && FMath::IsNearlyEqual(PairKnots.Last().Value, NewEndKnot))
-		{
-			// If the last knot has the same value, just increase its multiplicity
-			PairKnots.Last().Multiplicity += static_cast<uint32>(KnotsToAdd);
-		}
-		else
-		{
-			// Otherwise, add a new knot with the desired multiplicity
-			PairKnots.Add(FKnot(NewEndKnot, static_cast<uint32>(KnotsToAdd)));
-		}
-    
-		MarkFlatKnotsCacheDirty();
-	}
-	
-protected:
-	// Find a valid parameter value that doesn't conflict with existing knots
-    float GetNearestValidKnot(float DesiredParameter, float ExcludeValue) const
-    {
-        if (PairKnots.Num() == 0)
-        {
-            return DesiredParameter;
-        }
-
-        // Get valid parameter space
-        FInterval1f ParamSpace = GetParameterSpace();
-
-        // Check if DesiredParameter conflicts with any existing knot (excluding ExcludeValue)
-        bool bConflict = false;
-        for (const FKnot& Knot : PairKnots)
-        {
-            if (Knot.Value != ExcludeValue && Knot.Value == DesiredParameter)
-            {
-                bConflict = true;
-                break;
-            }
-        }
-
-        // If no conflict found, parameter is already valid
-        if (!bConflict)
-        {
-            return DesiredParameter;
-        }
-
-        // Find the closest valid parameter by walking outward alternately
-        float Left = DesiredParameter;
-        float Right = DesiredParameter;
-        bool bSearchLeft = true;
-        bool bSearchRight = true;
-        
-        // Continue as long as we have valid search directions
-        while (bSearchLeft || bSearchRight)
-        {
-            // Try moving left if still in valid range
-            if (bSearchLeft)
-            {
-                Left = std::nextafter(Left, -std::numeric_limits<float>::infinity());
-                
-                // Stop if we've moved outside the parameter space
-                if (Left < ParamSpace.Min)
-                {
-                    bSearchLeft = false;
-                    continue; // Skip to right search
-                }
-                
-                // Check if this value conflicts with any knot (excluding ExcludeValue)
-                bool bLeftValid = true;
-                for (const FKnot& Knot : PairKnots)
-                {
-                    if (Knot.Value != ExcludeValue && Knot.Value == Left)
-                    {
-                        bLeftValid = false;
-                        break;
-                    }
-                }
-                
-                if (bLeftValid)
-                {
-                    return Left; // Found valid parameter to the left
-                }
-            }
-            
-            // Try moving right if still in valid range
-            if (bSearchRight)
-            {
-                Right = std::nextafter(Right, std::numeric_limits<float>::infinity());
-                
-                // Stop if we've moved outside the parameter space
-                if (Right > ParamSpace.Max)
-                {
-                    bSearchRight = false;
-                    continue; // Skip to next iteration
-                }
-                
-                // Check if this value conflicts with any knot (excluding ExcludeValue)
-                bool bRightValid = true;
-                for (const FKnot& Knot : PairKnots)
-                {
-                    if (Knot.Value != ExcludeValue && Knot.Value == Right)
-                    {
-                        bRightValid = false;
-                        break;
-                    }
-                }
-                
-                if (bRightValid)
-                {
-                    return Right; // Found valid parameter to the right
-                }
-            }
-        }
-        
-        // If we couldn't find a valid parameter within range, return the closest boundary
-        return (FMath::Abs(ParamSpace.Min - DesiredParameter) < FMath::Abs(ParamSpace.Max - DesiredParameter)) 
-            ? ParamSpace.Min : ParamSpace.Max;
-    }
-	
-	virtual int32 GetExpectedNumKnots() const
-	{
-		const int32 NumValues = NumKeys();
-		return this->IsClosedLoop()
-				   ? (NumValues + (2 * Degree) + 1) // closed loop: n + 2p + 1
-				   : (NumValues + Degree + 1);
 	}
 private:
+	
 	/** Find knot span using binary search */
 	int32 FindKnotSpan(float Parameter) const
 	{
-		EnsureFlatKnotsCache();
+		UpdateFlatKnotsCache();
 		const TArray<float>& Knots = FlatKnots;
 		// First, check that the knot vector is populated.
 		if (Knots.Num() == 0)
@@ -1261,45 +1060,19 @@ private:
     	TArray<float> Basis;
     	    
     	if (PairKnots.IsEmpty())
+    	{
     		return ValueType();
+    	}
 
-    	TArray<float> KnotVector = GetKnotVector();
+		UpdateFlatKnotsCache();
+		
     	int32 Span = FindKnotSpan(Parameter);
 				
-    	Math::ComputeBSplineBasisFunctions(KnotVector, Span, Parameter, Degree, Basis);
+    	Math::ComputeBSplineBasisFunctions(FlatKnots, Span, Parameter, Degree, Basis);
     	return TSplineInterpolationPolicy<ValueType>::InterpolateWithBasis(Window, Basis);
     }
 	
-protected:
-	// Ensure the flat knot cache is up to date
-	void EnsureFlatKnotsCache() const
-	{
-		if (bFlatKnotsCacheDirty)
-		{
-			RebuildFlatKnotsCache();
-			bFlatKnotsCacheDirty = false;
-		}
-	}
-    
-	// Rebuild the flat knot cache from the pair representation
-	void RebuildFlatKnotsCache() const
-	{
-		FlatKnots.Reset();
-		for (const FKnot& Knot : PairKnots)
-		{
-			for (uint32 i = 0; i < Knot.Multiplicity; ++i)
-			{
-				FlatKnots.Add(Knot.Value);
-			}
-		}
-	}
-    
-	// Mark the flat knot cache as dirty
-	void MarkFlatKnotsCacheDirty() const
-	{
-		bFlatKnotsCacheDirty = true;
-	}
-	
+protected:	
     TArray<ValueType> Values;
 
 	// Knot representation in pair form having multiplicity
diff --git a/Engine/Source/Runtime/GeometryCore/Public/Splines/PolyBezierSpline.h b/Engine/Source/Runtime/GeometryCore/Public/Splines/PolyBezierSpline.h
--- a/Engine/Source/Runtime/GeometryCore/Public/Splines/PolyBezierSpline.h
+++ b/Engine/Source/Runtime/GeometryCore/Public/Splines/PolyBezierSpline.h
@@ -449,8 +449,7 @@ public:
 	            {
 	            	
 	                // Points don't match, use AddBezierSegmentInternal with appropriately computed parameter
-	                float StartParam = GetParameter(NumPoints - 1);
-	            	this->AddBezierSegmentInternal(NumSegments, P0, P1, P2, P3, ParameterizationPolicy, StartParam);
+	            	this->AddBezierSegmentInternal(NumSegments, P0, P1, P2, P3, ParameterizationPolicy);
 	            }
 	        }
 	    }
@@ -644,10 +643,10 @@ public:
     {
 		// Get P0 from the last point of the previous segment
 		ValueType P0 = Base::GetValue(Base::NumKeys() - 1);
-		FInterval1f ParameterRange = GetParameterSpace();
 		int32 NumSegments = Base::NumKeys() / 4;
 		// Add segment with append behavior
-		return AddBezierSegmentInternal(NumSegments, P0, P1, P2, P3, ParameterizationPolicy, ParameterRange.Max);
+
+		return AddBezierSegmentInternal(NumSegments, P0, P1, P2, P3, ParameterizationPolicy);
     }
 
     /**
@@ -668,12 +667,11 @@ public:
 	{
 		// Get P3 from the first point of the existing spline
 		ValueType P3 = Base::GetValue(0);
-		// Get current parameter range start
-		
-		FInterval1f ParameterRange = GetParameterSpace();
-  //   
+   
 		// Add segment with prepend behavior
-		return AddBezierSegmentInternal(0, P0, P1, P2, P3, ParameterizationPolicy, TOptional<float>(),ParameterRange.Min);
+		return AddBezierSegmentInternal(0, P0, P1, P2, P3,
+										ParameterizationPolicy,
+										false); // Prepend to first segment
 	}
 
 	/**
@@ -757,6 +755,9 @@ public:
         // Find segment to split
         const int32 NumSegments = Base::NumKeys() / 4;
 		// Cast to proper parameterization type
+
+		// Make sure we have a valid parameter
+		Parameter =  Base::GetNearestAvailableKnotValue(Parameter);
 		
         FInterval1f ParameterRange = GetParameterSpace();
 		int32 SegmentIndex;
@@ -783,8 +784,6 @@ public:
         ValueType P3 = Base::GetValue(BaseIndex + 3);
  
         // Use De Casteljau algorithm to split the curve
-		float MinDelta = UE_KINDA_SMALL_NUMBER;
-		LocalT = FMathf::Clamp(LocalT, MinDelta, 1.f - MinDelta);
         const float t = LocalT;
         const float mt = 1.0f - t;
         
@@ -803,18 +802,18 @@ public:
 		float OriginalStartParam = SegmentRange.Min;
 		float OriginalEndParam = SegmentRange.Max;
 		float TargetSplitParam = SegmentRange.Interpolate(LocalT);
-		TargetSplitParam = FMath::Clamp(TargetSplitParam, OriginalStartParam + MinDelta, OriginalEndParam - MinDelta);
 		
 		for (int32 i = 0; i < 4; ++i)
 		{
 			Base::RemoveValue(BaseIndex);
 		}
-		Base::RemoveKnotsCompletely(OriginalEndParam, GetExpectedNumKnots());
 		
         // Add the two new segments using the policy
-        AddBezierSegmentInternal(SegmentIndex, Q0, Q1, Q2, Q3, ParameterizationPolicy, OriginalStartParam, TargetSplitParam);
-		int32 IndexToReturn = AddBezierSegmentInternal(SegmentIndex + 1, R0, R1, R2, R3, ParameterizationPolicy, TargetSplitParam, OriginalEndParam);
+        AddBezierSegmentInternal(SegmentIndex, Q0, Q1, Q2, Q3, ParameterizationPolicy);
+		Base::InsertKnot( FKnot(TargetSplitParam, this->bClampEnds ? this->GetDegree() : 1));
+		int32 IndexToReturn = AddBezierSegmentInternal(SegmentIndex + 1, R0, R1, R2, R3, ParameterizationPolicy);
 
+		
 		return IndexToReturn;
     }
 
@@ -863,13 +862,10 @@ public:
         ValueType PositionAtParam = Base::Evaluate(Parameter);
         
         // Use the InsertPoint method to split the curve at the parameter
-        int32 PointIndex = InsertPointAtSegmentParam(SegmentIndex, LocalParam, PositionAtParam, ParameterizationPolicy);
+        InsertPointAtSegmentParam(SegmentIndex, LocalParam, PositionAtParam, ParameterizationPolicy);
         
         // Add the new segment from the inserted point
-        int32 Result = AddBezierSegmentInternal(SegmentIndex+1, PositionAtParam, P1, P2, P3, ParameterizationPolicy, Parameter);
-
-    	// Update the knot vector by inserting the new boundary.
-    	//this->InsertLocalKnots(Parameter, Parameter, Base::NumKeys());
+        int32 Result = AddBezierSegmentInternal(SegmentIndex+1, PositionAtParam, P1, P2, P3, ParameterizationPolicy);
     	
         return Result;
     }
@@ -890,93 +886,88 @@ public:
 	    }
 
     	// Calculate storage indices
-    	const int32 BaseIndex = SegmentIndex * 4;
-    	const bool bIsFirstSegment = (SegmentIndex == 0);
     	const bool bIsLastSegment = (SegmentIndex == NumSegments - 1);
-
-    	// Get the Knots at segment boundaries
-    	FInterval1f SegmentRange = GetSegmentParameterRange(SegmentIndex);
-    	float StartParam = SegmentRange.Min;
-    	float EndParam = SegmentRange.Max;
-    
-    	UE_LOG(LogSpline, Log, TEXT("Removing segment %d: Start=%f, End=%f"), 
-			   SegmentIndex, StartParam, EndParam);
-    
-    	// Get references to the evaluator and knot vector
-    	const TArray<float>& KnotVector = Base::GetKnotVector();
-    	const int32 Degree = GetDegree(); // Typically 3 for cubic
     
-    	// Store new boundary Knots before removing anything
-    	float NewStartParam = 0.0f;
-    	float NewEndParam = 0.0f;
+    	UE_LOG(LogSpline, Verbose, TEXT("Removing segment at index %d"), SegmentIndex);
     
-    	// For first segment, the new start parameter will be the end of removed segment
-    	if (bIsFirstSegment && SegmentIndex + 1 < NumSegments)
-    	{
-    		NewStartParam = EndParam;
-    	}
-    	// For last segment, the new end parameter will be the start of removed segment
-    	else if (bIsLastSegment && SegmentIndex > 0)
-    	{
-    		NewEndParam = StartParam;
-    	}
-    
-    	// Now remove points from storage
-    	bool bSuccess = true;
-		const int32 ExpectedKnots = GetExpectedNumKnots();
-	    if (bIsFirstSegment)
+        if (bIsLastSegment)
 	    {
-	        for (int32 i = 3; i >= 0; --i)
-	        {
-	            bSuccess &= Base::RemoveValue(BaseIndex + i);
-	        }
-	    	// Remove all knots at original start parameter
-	    	Base::RemoveKnotsCompletely(StartParam, ExpectedKnots);
-	    	        
-	    	// The new start parameter becomes the end parameter of the removed segment
-        
-	    	// First count existing multiplicity
-	    	int32 NewStartMultiplicity = Base::GetKnotMultiplicity(NewStartParam);
-        
-	    	// Calculate how many knots to add to reach desired multiplicity
-	    	int32 DesiredMultiplicity = Degree + 1; // Typically 4 for cubic
-	    	int32 KnotsToAdd = FMath::Max(0, DesiredMultiplicity - NewStartMultiplicity);
-        
-	    	UE_LOG(LogSpline, Log, TEXT("Adding %d knots at new start parameter %f to reach multiplicity %d"),
-				   KnotsToAdd, NewStartParam, DesiredMultiplicity);
-        
-	    	// Insert knots to reach desired multiplicity
-	    	for (int32 i = 0; i < KnotsToAdd; ++i)
-	    	{
-	    		Base::InsertKnot(NewStartParam);
-	    	}
-	    }
-	    else if (bIsLastSegment)
-	    {
-	    	for (int32 i = 3; i >= 0; --i)
-	    	{
-	    		bSuccess &= Base::RemoveValue(BaseIndex + i);
-	    	}
-	        
-	    	// Remove all knots at original end parameter
-	    	Base::RemoveKnotsCompletely(EndParam, ExpectedKnots);
-	    	Base::FillEndKnots(NewEndParam, Base::NumKeys());
-	    	
+	    	return RemovePoint(GetNumSegments());
 	    }
 	    else
 	    {
-	        // Remove the internal segment's P0 (duplicate), P1, and P2 - keep P3 (which is P0 of next)
-	        for (int32 i = 3; i >= 0; --i)
-	        {
-	            bSuccess &= Base::RemoveValue(BaseIndex + i - 1);
-	        }
-	    	// For internal segments, just remove all knots at the start parameter
-	    	Base::RemoveKnotsCompletely(StartParam, ExpectedKnots);
-	    	
+	    	// in all other cases, removing the segment would be the same as removing the point at that index
+	        return RemovePoint(SegmentIndex);
 	    }
-	    return bSuccess;
     }
 
+	/**
+	 * Removes a point from the spline.
+	 * @param PointIndex - Index of the point to remove
+	 */
+	bool RemovePoint(int32 PointIndex)
+	{
+		const int32 NumSegments = GetNumSegments();
+        if (PointIndex < 0 || PointIndex > NumSegments)
+        {
+            return false;
+        }
+		bool bSuccess = true;
+		const int32 SegmentIndex = FMath::Clamp(PointIndex, 0, NumSegments - 1);
+		const int32 BaseIndex = SegmentIndex * 4;
+		const int32 Degree = GetDegree();
+		if (PointIndex == 0)
+		{
+			// Removing the first point - remove the first segment
+			for (int32 i = 3; i >= 0; --i)
+			{
+				bSuccess &= Base::RemoveValue(BaseIndex + i);
+			}
+			// Remove first knot
+			Base::RemoveKnot(0);
+
+			if (Base::IsClampedEnds())
+			{
+				Base::PairKnots[0].Multiplicity = Degree + 1;
+			}
+		}
+		// else if last point
+		else if (PointIndex == GetNumSegments())
+		{
+			// Removing the last point - remove the last segment
+			for (int32 i = 3; i >= 0; --i)
+			{
+				bSuccess &= Base::RemoveValue(BaseIndex + i);
+			}
+	        
+			// Remove end knot
+			Base::RemoveKnot(SegmentIndex + 1);
+			if (Base::IsClampedEnds())
+			{
+				Base::PairKnots.Last().Multiplicity = Degree + 1;
+			}
+		}
+		else
+		{
+			// For an internal point, we need to remove exactly 4 control points:
+			// - P2 and P3 from the previous segment (where P3 is the point)
+			// - P0 and P1 from the next segment (where P0 is the same point)
+        
+			// Start index for removal: P2 of previous segment
+			int32 StartRemoveIndex = (PointIndex - 1) * 4 + 2;
+        
+			// Remove 4 consecutive control points
+			for (int32 i = 0; i < 4; ++i)
+			{
+				Base::RemoveValue(StartRemoveIndex);
+			}
+        
+			Base::RemoveKnot(PointIndex);
+		}
+		Base::Dump();
+		return bSuccess;
+	}
+
     /**
 	 * Updates a single control point within a Bezier segment
 	 * @param Storage - Storage to update the point in
@@ -1039,245 +1030,10 @@ public:
 		const int32 NumPoints = this->Base::NumKeys();
 		return NumPoints/4;
 	}
-private:
-	/**
-	 * Helper method to add the first segment to the spline
-	 * @param Parameter - Parameter value where to insert the point (Spline Space)
-	 * @param ControlPoint - Bezier control point to insert
-	 * @param PointIndex - Index of the newly created bezier control point
-	 * @return true if the points were added, false otherwise
-	 */
-	bool InitializeFirstSegment(
-		float Parameter,
-		const ValueType& ControlPoint,
-		int32& PointIndex)
-    {
-    	if (Base::NumKeys() < 3)
-    	{
-    		// First point - just store it directly
-    		PointIndex = Base::AddValue(ControlPoint);
-    		return true;
-    	}
 
-    	// Adding the fourth control point would create a segment
-    	if (Base::NumKeys() == 3)
-    	{
-    		const ValueType& P0 = Base::GetValue(0);
-    		const ValueType& P1 = Base::GetValue(1);
-    		const ValueType& P2 = Base::GetValue(2);
-    		const ValueType& P3 = ControlPoint;
-	        
-    		// Add the first segment
-    		SetControlPoints({P0, P1, P2, P3}, EParameterizationPolicy::Centripetal);
-    		PointIndex = 3;
-    		return true;
-    	}
-    	return false;
-    }
-
-	/**
-     * Prepends a point to the spline at the beginning.
-     * The new point will be connected to the first segment.
-     * @param Position The position of the new point
-     * @param ParameterizationPolicy How to distribute Knots
-     * @return Index of the new segment
-     */
-    int32 PrependPoint(
-		const ValueType& Position,
-		EParameterizationPolicy ParameterizationPolicy = EParameterizationPolicy::Centripetal)
-	{
-		// Get current parameter range start
-		FInterval1f ParameterRange = GetParameterSpace();
-    
-		// Create a control point at the beginning
-		ValueType P3 = Base::GetValue(0); // First point of the spline
-    
-		// Calculate reasonable auto-tangents
-		ValueType Dir = P3 - Position;
-		ValueType P1 = Position + Dir * 0.33f;
-		ValueType P2 = Position + Dir * 0.66f;
-    
-		// Add the segment by reusing existing code
-		int32 NewIndex = AddBezierSegmentInternal(0, Position, P1, P2, P3, ParameterizationPolicy, TOptional<float>(), ParameterRange.Min);
-		check (NewIndex == 0);
-		return NewIndex;
-	}
-
-	/**
-	 * Appends a point to the end of the spline.
-	 * Creates a new cubic Bezier segment connecting the last point to the new point.
-	 * 
-	 * @param Position The position of the new point
-	 * @param ParameterizationPolicy How to distribute Knots
-	 * @return Index of the new segment
-	 */
-	int32 AppendPoint(
-		const ValueType& Position,
-		EParameterizationPolicy ParameterizationPolicy)
-	{
-		// Append - create a point at the end
-		int32 LastIndex = Base::NumKeys() - 1;
-		float StartParam = Base::GetParameter(LastIndex);
-		ValueType P0 = Base::GetValue(LastIndex);
-        
-		// Calculate reasonable auto-tangents
-		ValueType Dir = Position - P0;
-		ValueType P1 = P0 + Dir * 0.33f;
-		ValueType P2 = P0 + Dir * 0.66f;
-
-		int32 NumSegments = Base::NumKeys() / 4;
-		// Add the segment
-		int32 NewIndex = AddBezierSegmentInternal(NumSegments, P0, P1, P2, Position, ParameterizationPolicy, StartParam);
-		return NewIndex;
-	}
 	
-	/**
-	 * Helper to compute segment length based on parameterization policy
-	 */
-	float ComputeSegmentLength(
-		const ValueType& P0,
-		const ValueType& P3,
-		EParameterizationPolicy ParameterizationPolicy)
-    {
-		constexpr float MinSegmentLength = UE_KINDA_SMALL_NUMBER;
-		
-    	switch (ParameterizationPolicy)
-    	{
-    	case EParameterizationPolicy::ChordLength:
-    		return static_cast<float>(FMath::Max(Math::Distance(P0, P3), MinSegmentLength));
-    	case EParameterizationPolicy::Centripetal:
-    		return static_cast<float>(FMath::Max(Math::CentripetalDistance(P0, P3), MinSegmentLength));
-    	case EParameterizationPolicy::Uniform:
-    	default:
-			return 1.0f;
-    	}
-    }
-    
-	/**
-     * Internal implementation to add a Bezier segment at the specified index with optional fixed parameter
-	 * 
-	 * @param Storage Storage to add the segment to
-	 * @param SegmentIndex Index at which to add the segment
-	 * @param P0 Start point
-	 * @param P1 First control point
-	 * @param P2 Second control point
-	 * @param P3 End point
-	 * @param ParameterizationPolicy How to distribute Knots
-	 * @param FixedParameter Optional fixed parameter for a specific point
-	 * @param FixedPointIndex Which point the fixed parameter applies to (0-3)
-	 * @return Index of the new segment
-     */
-	int32 AddBezierSegmentInternal(
-		int32 SegmentIndex,
-		const ValueType& P0,
-		const ValueType& P1, 
-		const ValueType& P2, 
-		const ValueType& P3,
-		EParameterizationPolicy ParameterizationPolicy = EParameterizationPolicy::Centripetal,
-		TOptional<float> StartParam = TOptional<float>(),
-		TOptional<float> EndParam = TOptional<float>())
-    {
-		
-    	// Validate segment index
-    	const int32 NumPoints = Base::NumKeys();
-    	const int32 NumSegments = NumPoints/4;
-    	
-	    SegmentIndex = FMath::Clamp(SegmentIndex, 0, NumSegments);
-	    
-	    // Calculate the insertion index for the control points
-	    int32 ControlPointIndex = SegmentIndex * 4;
-	    
-	    // Calculate Knots based on endpoints and policy
-	    FInterval1f SegmentParams;
-	    float SegmentLength = ComputeSegmentLength(P0, P3, ParameterizationPolicy);
 
-    	// Handle all possible parameter combinations
-	    if (StartParam.IsSet() && EndParam.IsSet())
-	    {
-	        // Both Knots provided - use them directly
-	        SegmentParams.Min = StartParam.GetValue();
-	        SegmentParams.Max = EndParam.GetValue();
-	    }
-	    else if (StartParam.IsSet())
-	    {
-	        // Only start parameter provided - calculate end
-	        SegmentParams.Min = StartParam.GetValue();
-	        SegmentParams.Max = SegmentParams.Min + SegmentLength;
-	    }
-	    else if (EndParam.IsSet())
-	    {
-	        // Only end parameter provided - calculate start
-	        SegmentParams.Max = EndParam.GetValue();
-	        SegmentParams.Min = SegmentParams.Max - SegmentLength;
-	    }
-	    else
-	    {
-	        // Neither parameter provided - calculate based on context
-	        TOptional<float> ExistingParam;
-	        bool bConnectAtStart = true;
-
-	        if (NumSegments > 0)
-	        {
-	            if (SegmentIndex == 0)
-	            {
-	                // Prepending - connect to first point of existing spline
-	                ExistingParam = GetSegmentParameterRange(0).Min;
-	                bConnectAtStart = false;  // Connect at end of new segment
-	            }
-	            else if (SegmentIndex == NumSegments)
-	            {
-	                // Appending - connect to last point of existing spline
-	                ExistingParam = GetSegmentParameterRange(NumSegments - 1).Max;
-	                bConnectAtStart = true;  // Connect at start of new segment
-	            }
-	            else
-	            {
-	                // Inserting at segment boundary
-	                ExistingParam = GetSegmentParameterRange(SegmentIndex).Min;
-	                bConnectAtStart = true;  // Connect at start of new segment
-	            }
-	        }
-
-	        // Calculate Knots using existing context
-	        if (ExistingParam.IsSet())
-	        {
-	            float AnchorParam = ExistingParam.GetValue();
-	            if (bConnectAtStart)
-	            {
-	                SegmentParams.Min = AnchorParam;
-	                SegmentParams.Max = SegmentParams.Min + SegmentLength;
-	            }
-	            else
-	            {
-	                SegmentParams.Max = AnchorParam;
-	                SegmentParams.Min = SegmentParams.Max - SegmentLength;
-	            }
-	        }
-	        else
-	        {
-	            // No existing context - start at zero
-	            SegmentParams.Min = 0.0f;
-	            SegmentParams.Max = SegmentLength;
-	        }
-	    }
-	    
-	    // Add the control points
-	    int32 P0Index = Base::InsertValue(ControlPointIndex, P0);
-	    Base::InsertValue(ControlPointIndex + 1, P1);
-	    Base::InsertValue(ControlPointIndex + 2, P2);
-	    Base::InsertValue(ControlPointIndex + 3, P3);
-	    
-	    // Update the knot vector
-	    InsertLocalKnots(SegmentParams.Min, SegmentParams.Max);
-
-    	//Reparameterize(ParameterizationPolicy);
-	    
-	    return SegmentIndex;
-	}
-
-	
 	// Parameterization methods
-public:
 	/**
 	 * Reparameterizes the spline based on the provided points and mode
 	 * @param Mode - New parameterization mode
@@ -1306,19 +1062,14 @@ public:
 		}
     	
     	
-		UE_LOG(LogSpline, Verbose, TEXT("Set Custom knot vector - Points: %d, Knots: %d, Mode: %d"),
+		UE_LOG(LogSpline, Verbose, TEXT("Set knot vector - Points: %d, Knots: %d, Mode: %d"),
 					   NumValues,  this->PairKnots.Num(), static_cast<int32>(Mode));
-		for (int32 i = 0; i < this->PairKnots.Num(); ++i)
-		{
-			UE_LOG(LogSpline, Verbose, TEXT("  Knot[%d] = %f x %d"), i, this->PairKnots[i].Value, this->PairKnots[i].Multiplicity);
-		}
+		Base::PrintKnotVector();
 	}
 	virtual FInterval1f GetParameterSpace() const override
 	{
-		const TArray<float> Knots = Base::GetDistinctKnotValues();
-		
-		return Knots.Num() > 0
-			? FInterval1f(Knots[0], Knots.Last())
+		return Base::PairKnots.Num() > 0
+			? FInterval1f(Base::PairKnots[0].Value, Base::PairKnots.Last().Value)
 			: FInterval1f::Empty();
 	}
 
@@ -1358,7 +1109,7 @@ public:
 		// Early exit if parameter isn't changing
 		if (OldParameter == NewParameter)
 		{
-			return SegmentIndex * PointsPerSegment + LocalPointIndex;
+			return Index;
 		}
 		
 		UE_LOG(LogSpline, Verbose, TEXT("\t"));
@@ -1366,7 +1117,18 @@ public:
 		Base::Dump();
 		
 		// Prevent reuse of an existing knot.
-		NewParameter = this->GetNearestValidKnot(NewParameter, OldParameter);
+		const float DesiredNewParameter = NewParameter;
+		const float ActualNewParameter = this->GetNearestAvailableKnotValue(DesiredNewParameter);
+
+		const float DesiredToOldDistance = FMath::Abs(DesiredNewParameter - OldParameter);
+		const float DesiredToActualNewDistance = FMath::Abs(DesiredNewParameter - ActualNewParameter);
+
+		if (DesiredToActualNewDistance > DesiredToOldDistance)
+		{
+			// If our nearest valid knot is further away from the desired value than our current value is, we need to just keep our current value
+			return Index;
+		}
+		
 		int32 ReturnIndex = INDEX_NONE;
 		
 		if (LocalPointIndex == 0) // P0
@@ -1503,13 +1265,12 @@ public:
 	
 	virtual int32 FindIndexForParameter(float Parameter, float& OutLocalParam) const override
 	{
-		TArray<float> Knots = Base::GetDistinctKnotValues();
-		int32 NumSegments = Knots.Num() - 1;
+		int32 NumSegments = Base::PairKnots.Num() - 1;
 
 		for (int32 i = 0; i < NumSegments; ++i)
 		{
-			float Start = Knots[i];
-			float End = Knots[i + 1];
+			float Start = Base::PairKnots[i].Value;
+			float End = Base::PairKnots[i + 1].Value;
 			if ((Parameter >= Start && Parameter < End) || (i == NumSegments - 1 && FMath::IsNearlyEqual(Parameter, End)))
 			{
 				float SpanLength = End - Start;
@@ -1575,70 +1336,7 @@ public:
 	}
 
 	
-	/**
-	 * Helper to insert local knots for a segment boundary.
-	 * For clamped splines, internal boundaries typically should have a multiplicity equal
-	 * to the degree rather than the clamped endpoint multiplicity.
-	 * @param StartParam - Start parameter of the segment
-	 * @param EndParam - End parameter of the segment
-	 */
-	void InsertLocalKnots(float StartParam, float EndParam)
-	{
-		const int32 NumKeys = Base::NumKeys();
-		if (this->bClampEnds)
-		{
-			// Get parameter range to determine if these are endpoints
-			FInterval1f KnotRange = this->GetKnotRange();
-			float CachedEndKnot = KnotRange.Max;
-			// Check if StartParam is at the beginning of the knot range
-			bool bIsStartParamEndKnot = StartParam <= KnotRange.Min;
-        
-			// Check if EndParam is at the end of the knot range
-			bool bIsEndParamEndKnot = EndParam >= KnotRange.Max;
-        
-			// For endpoints, we want multiplicity of degree+1 for clamped splines
-			const int32 EndpointMultiplicity = this->GetDegree() + 1;
-        
-			// For internal knots, we typically want multiplicity of degree
-			const int32 InternalMultiplicity = this->GetDegree();
-
-			int32 StartMult = this->GetKnotMultiplicity(StartParam);
-			int32 DesiredStartMult = bIsStartParamEndKnot ? EndpointMultiplicity : InternalMultiplicity;
-
-			int32 EndMult = this->GetKnotMultiplicity(EndParam);
-			int32 DesiredEndMult = bIsEndParamEndKnot ? EndpointMultiplicity : InternalMultiplicity;
-        
-			// First insert required knots
-			if (StartMult < DesiredStartMult)
-			{
-				this->InsertKnot(StartParam, DesiredStartMult - StartMult);
-			}
-			if (EndMult < DesiredEndMult)
-			{
-				this->InsertKnot(EndParam, DesiredEndMult - EndMult);
-			}
-
-			// Now check if we need to remove excessive knots
-			// we do this later to avoid validation checks inside RemoveKnot()
-			if (StartMult > DesiredStartMult)
-			{
-				this->RemoveKnot(StartParam, StartMult - DesiredStartMult);
-			}
-			if (EndMult > DesiredEndMult)
-			{
-				this->RemoveKnot(EndParam, EndMult - DesiredEndMult);
-			}
-		}
-		else
-		{
-			// For unclamped splines, simply insert one knot at each boundary.
-			this->InsertKnot(StartParam, 1);
-			this->InsertKnot(EndParam, 1);
-		}
-		
-		const int32 ExpectedKnots = GetExpectedNumKnots();
-		this->FixEndKnots(ExpectedKnots);
-	}
+	
 
 	/**
 	 * Maps a segment index to its parameter range
@@ -1649,12 +1347,11 @@ public:
 	 */
 	FInterval1f GetSegmentParameterRange(int32 SegmentIndex) const
 	{
-		TArray<float> DistinctKnots = Base::GetDistinctKnotValues();
 		FInterval1f SegmentRange;
-		if (SegmentIndex >= 0 && SegmentIndex < DistinctKnots.Num() - 1)
+		if (SegmentIndex >= 0 && SegmentIndex < Base::PairKnots.Num() - 1)
 		{
-			SegmentRange.Min = DistinctKnots[SegmentIndex];
-			SegmentRange.Max = DistinctKnots[SegmentIndex + 1];
+			SegmentRange.Min = Base::PairKnots[SegmentIndex].Value;
+			SegmentRange.Max = Base::PairKnots[SegmentIndex + 1].Value;
 		}
         
 		return SegmentRange;
@@ -1681,42 +1378,40 @@ public:
 	 */
 	bool MapGlobalParameterToLocalSegment(float GlobalParam, int32& OutSegmentIndex, float& OutLocalParam) const
 	{
-		TArray<float> DistinctKnots = Base::GetDistinctKnotValues();
-        
 		// Handle parameter outside range
-		if (DistinctKnots.Num() < 2)
+		if (Base::PairKnots.Num() < 2)
 		{
 			OutSegmentIndex = 0;
 			OutLocalParam = 0.0f;
 			return false;
 		}
         
-		if (GlobalParam <= DistinctKnots[0])
+		if (GlobalParam <= Base::PairKnots[0].Value)
 		{
 			OutSegmentIndex = 0;
 			OutLocalParam = 0.0f;
 			return true;
 		}
         
-		if (GlobalParam >= DistinctKnots.Last())
+		if (GlobalParam >= Base::PairKnots.Last().Value)
 		{
-			OutSegmentIndex = DistinctKnots.Num() - 2;
+			OutSegmentIndex = Base::PairKnots.Num() - 2;
 			OutLocalParam = 1.0f;
 			return true;
 		}
         
 		// Find segment containing parameter
-		for (int32 i = 0; i < DistinctKnots.Num() - 1; ++i)
+		for (int32 i = 0; i < Base::PairKnots.Num() - 1; ++i)
 		{
-			if (GlobalParam >= DistinctKnots[i] && GlobalParam < DistinctKnots[i + 1])
+			if (GlobalParam >= Base::PairKnots[i].Value && GlobalParam < Base::PairKnots[i + 1].Value)
 			{
 				OutSegmentIndex = i;
                 
 				// Calculate local parameter
-				float SegmentLength = DistinctKnots[i + 1] - DistinctKnots[i];
+				float SegmentLength = Base::PairKnots[i + 1].Value - Base::PairKnots[i].Value;
 				if (SegmentLength > UE_SMALL_NUMBER)
 				{
-					OutLocalParam = (GlobalParam - DistinctKnots[i]) / SegmentLength;
+					OutLocalParam = (GlobalParam - Base::PairKnots[i].Value) / SegmentLength;
 				}
 				else
 				{
@@ -1738,7 +1433,7 @@ public:
 	 */
 	int32 GetNumDistinctSegments() const
 	{
-		return FMath::Max(0, Base::GetDistinctKnotValues().Num() - 1);
+		return FMath::Max(0, Base::PairKnots.Num() - 1);
 	}
 
 	int32 FindSegmentIndex(float Parameter, float &OutLocalParam) const
@@ -1765,10 +1460,11 @@ public:
 		}
 		return 4 + (NumSegments - 1) * 3 + 4;
 	}
-	void SetKnotVector(const TArray<float>& NewKnots)
+	void SetKnotVector(const TArray<FKnot>& NewKnots)
 	{
 		Base::SetCustomKnots(NewKnots);
 	}
+	
 private:
 	
 	/**
@@ -1776,7 +1472,6 @@ private:
 	 * Uses those distances to build a proportional knot vector with the proper multiplicities for each Bézier segment.
 	 * @param Mode - Knot generation mode to use
 	 */
-	
 	void GenerateDistanceKnotsForBezier(EParameterizationPolicy Mode)
 	{
 		TArray<ValueType> Points = this->Values;
@@ -2029,6 +1724,182 @@ private:
     	int32 NumSegments = Base::NumKeys() / 4;
     	return Index >= 0 && Index < NumSegments;
     }
+
+	/**
+	 * Helper method to add the first segment to the spline
+	 * @param Parameter - Parameter value where to insert the point (Spline Space)
+	 * @param ControlPoint - Bezier control point to insert
+	 * @param PointIndex - Index of the newly created bezier control point
+	 * @return true if the points were added, false otherwise
+	 */
+	bool InitializeFirstSegment(
+		float Parameter,
+		const ValueType& ControlPoint,
+		int32& PointIndex)
+    {
+    	if (Base::NumKeys() < GetDegree())
+    	{
+    		// First point - just store it directly
+    		PointIndex = Base::AddValue(ControlPoint);
+    		return true;
+    	}
+
+    	// Adding the fourth control point would create a segment
+    	if (Base::NumKeys() == GetDegree())
+    	{
+    		const ValueType& P0 = Base::GetValue(0);
+    		const ValueType& P1 = Base::GetValue(1);
+    		const ValueType& P2 = Base::GetValue(2);
+    		const ValueType& P3 = ControlPoint;
+	        
+    		// Add the first segment
+    		SetControlPoints({P0, P1, P2, P3}, EParameterizationPolicy::Centripetal);
+    		PointIndex = 3;
+    		return true;
+    	}
+    	return false;
+    }
+
+	/**
+     * Prepends a point to the spline at the beginning.
+     * The new point will be connected to the first segment.
+     * @param Position The position of the new point
+     * @param ParameterizationPolicy How to distribute Knots
+     * @return Index of the new segment
+     */
+    int32 PrependPoint(
+		const ValueType& Position,
+		EParameterizationPolicy ParameterizationPolicy = EParameterizationPolicy::Centripetal)
+	{
+		// Get current parameter range start
+		FInterval1f ParameterRange = GetParameterSpace();
+    
+		// Create a control point at the beginning
+		ValueType P3 = Base::GetValue(0); // First point of the spline
+    
+		// Calculate reasonable auto-tangents
+		ValueType Dir = P3 - Position;
+		ValueType P1 = Position + Dir * 0.33f;
+		ValueType P2 = Position + Dir * 0.66f;
+    
+		// Add the segment by reusing existing code
+		int32 NewIndex = AddBezierSegmentInternal(0, Position, P1, P2, P3, ParameterizationPolicy, false);
+		check (NewIndex == 0);
+		return NewIndex;
+	}
+
+	/**
+	 * Appends a point to the end of the spline.
+	 * Creates a new cubic Bezier segment connecting the last point to the new point.
+	 * 
+	 * @param Position The position of the new point
+	 * @param ParameterizationPolicy How to distribute Knots
+	 * @return Index of the new segment
+	 */
+	int32 AppendPoint(
+		const ValueType& Position,
+		EParameterizationPolicy ParameterizationPolicy)
+	{
+		// Append - create a point at the end
+		int32 LastIndex = Base::NumKeys() - 1;
+		float StartParam = Base::GetParameter(LastIndex);
+		ValueType P0 = Base::GetValue(LastIndex);
+        
+		// Calculate reasonable auto-tangents
+		ValueType Dir = Position - P0;
+		ValueType P1 = P0 + Dir * 0.33f;
+		ValueType P2 = P0 + Dir * 0.66f;
+
+		int32 NumSegments = Base::NumKeys() / 4;
+		// Add the segment
+		int32 NewIndex = AddBezierSegmentInternal(NumSegments, P0, P1, P2, Position, ParameterizationPolicy/*, StartParam*/);
+		return NewIndex;
+	}
+	
+	/**
+	 * Helper to compute segment length based on parameterization policy
+	 */
+	float ComputeSegmentLength(
+		const ValueType& P0,
+		const ValueType& P3,
+		EParameterizationPolicy ParameterizationPolicy)
+    {
+		constexpr float MinSegmentLength = UE_KINDA_SMALL_NUMBER;
+		
+    	switch (ParameterizationPolicy)
+    	{
+    	case EParameterizationPolicy::ChordLength:
+    		return static_cast<float>(FMath::Max(Math::Distance(P0, P3), MinSegmentLength));
+    	case EParameterizationPolicy::Centripetal:
+    		return static_cast<float>(FMath::Max(Math::CentripetalDistance(P0, P3), MinSegmentLength));
+    	case EParameterizationPolicy::Uniform:
+    	default:
+			return 1.0f;
+    	}
+    }
+
+	/**
+	 * Applies multiplicity to the internal knots based on the spline's degree
+	 */
+	void ApplyInternalKnotsMultiplicity()
+	{
+		// fix internal knots to be clamped to Degree
+		for (int32 i = 1; i < Base::PairKnots.Num() - 1; ++i)
+		{
+			Base::PairKnots[i].Multiplicity = GetDegree();
+		}
+	}
+
+	int32 AddBezierSegmentInternal(
+		int32 SegmentIndex,
+		const ValueType& P0,
+		const ValueType& P1, 
+		const ValueType& P2, 
+		const ValueType& P3,
+		EParameterizationPolicy ParameterizationPolicy = EParameterizationPolicy::Centripetal,
+		bool bAppendToSegment = true) // true = connect to end, false = connect to start
+	{
+		const int32 NumSegments = Base::NumKeys() / 4;
+		SegmentIndex = FMath::Clamp(SegmentIndex, 0, NumSegments);
+    
+		// Calculate segment length
+		float SegmentLength = ComputeSegmentLength(P0, P3, ParameterizationPolicy);
+    
+		// Add the control points
+		int32 ControlPointIndex = SegmentIndex * 4;
+		Base::InsertValue(ControlPointIndex, P0);
+		Base::InsertValue(ControlPointIndex + 1, P1);
+		Base::InsertValue(ControlPointIndex + 2, P2);
+		Base::InsertValue(ControlPointIndex + 3, P3);
+
+		// FInterval1f RefRange = GetSegmentParameterRange(SegmentIndex);
+		bool bAppendToEnd = bAppendToSegment && SegmentIndex > (this->GetNumDistinctSegments() - 1);
+		bool bPrependToStart = !bAppendToSegment && SegmentIndex == 0;
+	
+		// Insert knots based on append/prepend to the segment
+		if (NumSegments == 0) // First segment
+		{
+			this->ResetKnotVector();
+			this->InsertKnot( FKnot(0.0f, this->bClampEnds ? this->GetDegree() + 1 : 1));
+			this->InsertKnot( FKnot(SegmentLength, this->bClampEnds ? this->GetDegree() + 1 : 1));
+		}
+		else if (bAppendToEnd)
+		{
+			this->InsertKnot( FKnot(Base::PairKnots.Last().Value + SegmentLength, this->bClampEnds ? this->GetDegree() : 1));
+		}
+		else if (bPrependToStart)
+		{
+			this->InsertKnot( FKnot(Base::PairKnots[SegmentIndex].Value - SegmentLength, this->bClampEnds ? this->GetDegree() : 1));
+		}
+	
+		this->ApplyClampedKnotsMultiplicity();
+
+		// make sure the internal knots have the correct multiplicity as one of the end knots
+		// might have become an internal knot
+		this->ApplyInternalKnotsMultiplicity();
+		Base::Dump();
+		return SegmentIndex;
+	}
 };
 
 
diff --git a/Engine/Source/Runtime/GeometryCore/Public/Splines/SplineMath.h b/Engine/Source/Runtime/GeometryCore/Public/Splines/SplineMath.h
--- a/Engine/Source/Runtime/GeometryCore/Public/Splines/SplineMath.h
+++ b/Engine/Source/Runtime/GeometryCore/Public/Splines/SplineMath.h
@@ -3,6 +3,12 @@
 #pragma once
 #include "CoreMinimal.h"
 #include "ParameterizedTypes.h"
+
+namespace UE::Geometry::Spline
+{
+	struct FKnot;
+}
+
 namespace UE
 {
 namespace Geometry
@@ -147,26 +153,6 @@ struct FSplineValidation
 		return true;
 	}
 
-	static bool IsValidKnotVector(const TArray<float>& Knots, int32 Degree)
-	{
-		// Need at least 2*(degree+1) knots
-		if (Knots.Num() < 2*(Degree + 1))
-		{
-			return false;
-		}
-
-		// Check monotonicity
-		for (int32 i = 1; i < Knots.Num(); ++i)
-		{
-			if (Knots[i] < Knots[i-1])
-			{
-				return false;
-			}
-		}
-
-		return true;
-	}
-
 	static bool IsValidParameter(float Parameter)
 	{
 		return !FPlatformMath::IsNaN(Parameter) && 
diff --git a/Engine/Source/Runtime/GeometryCore/Public/Splines/TangentBezierSpline.h b/Engine/Source/Runtime/GeometryCore/Public/Splines/TangentBezierSpline.h
--- a/Engine/Source/Runtime/GeometryCore/Public/Splines/TangentBezierSpline.h
+++ b/Engine/Source/Runtime/GeometryCore/Public/Splines/TangentBezierSpline.h
@@ -1020,42 +1020,8 @@ public:
 		{
 			TangentModes.RemoveAt(Index);
 		}
-    
-		if (Index == 0)
-		{
-			// Removing the first point - remove the first segment
-			InternalSpline.RemoveSegment(0);
-		}
-		else if (Index == NumPoints - 1 && !this->IsClosedLoop())
-		{
-			// Removing the last point - remove the last segment
-			InternalSpline.RemoveSegment(InternalSpline.GetNumSegments() - 1);
-		}
-		else
-		{
-			// Removing an internal point
-			// Get the parameter of the point being removed
-			float PointParam = InternalSpline.GetParameter((Index - 1) * 4 + 3);
-        
-			// For an internal point, we need to remove exactly 4 control points:
-			// - P2 and P3 from the previous segment (where P3 is the point)
-			// - P0 and P1 from the next segment (where P0 is the same point)
-        
-			// Start index for removal: P2 of previous segment
-			int32 StartRemoveIndex = (Index - 1) * 4 + 2;
-        
-			// Remove 4 consecutive control points
-			for (int32 i = 0; i < 4; ++i)
-			{
-				InternalSpline.RemoveValue(StartRemoveIndex);
-			}
-        
-			// Remove the knot associated with this point
-			
-			
-			const int32 ExpectedKnots = 3 * GetNumSegments() + 4;
-			InternalSpline.RemoveKnotsCompletely(PointParam, ExpectedKnots);
-		}
+
+		InternalSpline.RemovePoint(Index);
     
 		// Ensure tangent modes array size matches number of points after removal
 		if (TangentModes.Num() > GetNumPoints())
@@ -1582,14 +1548,14 @@ public:
 	    UpdateTangents();
     }
 
-	void SetKnotVector(const TArray<float>& InKnots)
+	void SetKnotVector(const TArray<FKnot>& InKnots)
 	{
 		InternalSpline.SetKnotVector(InKnots);
 
 		UpdateTangents();
 	}
 	
-	const TArray<float>& GetKnotVector() const
+	const TArray<FKnot>& GetKnotVector() const
 	{
 		return InternalSpline.GetKnotVector();
 	}
