Fix GetNearestAvailableKnotValue static analysis warning of loop always being true.

diff --git a/Engine/Source/Runtime/GeometryCore/Public/Splines/BSpline.h b/Engine/Source/Runtime/GeometryCore/Public/Splines/BSpline.h
--- a/Engine/Source/Runtime/GeometryCore/Public/Splines/BSpline.h
+++ b/Engine/Source/Runtime/GeometryCore/Public/Splines/BSpline.h
@@ -599,32 +599,30 @@ protected:
         // Find the closest valid parameter by walking outward alternately
         float Left = DesiredParameter;
         float Right = DesiredParameter;
-        bool bSearchLeft = true;
-        bool bSearchRight = true;
-        
-		// Search until valid parameter is found
-        while (bSearchLeft || bSearchRight)
-        {
-        	if (bSearchLeft)
+        constexpr int32 MaxIterations = 100000; // Set to a high but safe value
+		int32 IterationCount = 0;
+
+		// Search until a valid parameter is found or we hit the limit
+		while (IterationCount < MaxIterations)
+		{
+			IterationCount++;
+			
+        	Left = std::nextafter(Left, -std::numeric_limits<float>::infinity());
+        	if (!KnotValuesSet.Contains(Left))
         	{
-        		Left = std::nextafter(Left, -std::numeric_limits<float>::infinity());
-        		if (!KnotValuesSet.Contains(Left))
-        		{
-        			return Left;
-        		}
+        		return Left;
         	}
-
-        	if (bSearchRight)
+        	Right = std::nextafter(Right, std::numeric_limits<float>::infinity());
+        	if (!KnotValuesSet.Contains(Right))
         	{
-        		Right = std::nextafter(Right, std::numeric_limits<float>::infinity());
-        		if (!KnotValuesSet.Contains(Right))
-        		{
-        			return Right;
-        		}
+        		return Right;
         	}
         }
-		// Fallback (shouldn't be reached normally)
-		return DesiredParameter;
+		
+		// If no valid parameter was found within limits, return the original (or handle error)
+        // This should rarely be hit unless space is densely packed
+        ensureMsgf(false, TEXT("Max iterations reached in GetNearestAvailableKnotValue"));
+        return DesiredParameter;
     }
 
 	virtual int32 GetExpectedNumKnots() const
