GeometryCore: Fix TPointHashGrid3 breaking for values that are not close to the origin due to overflow in its indexer.

diff --git a/Engine/Source/Runtime/GeometryCore/Public/Spatial/PointHashGrid3.h b/Engine/Source/Runtime/GeometryCore/Public/Spatial/PointHashGrid3.h
--- a/Engine/Source/Runtime/GeometryCore/Public/Spatial/PointHashGrid3.h
+++ b/Engine/Source/Runtime/GeometryCore/Public/Spatial/PointHashGrid3.h
@@ -33,7 +33,7 @@ class TPointHashGrid3
 private:
 	TMultiMap<FVector3i, PointDataType> Hash;
 	FCriticalSection CriticalSection;
-	TScaleGridIndexer3<RealType> Indexer;
+	TWrapAroundGridIndexer3<RealType> Indexer;
 	PointDataType InvalidValue;
 
 public:
@@ -410,27 +410,17 @@ private:
 			SearchRadius = FMath::Sqrt(MinDistSq);
 		}
 
-		FVector3i min_idx = Indexer.ToGrid(QueryPoint - SearchRadius * TVector<RealType>::One());
-		FVector3i max_idx = Indexer.ToGrid(QueryPoint + SearchRadius * TVector<RealType>::One());
-
-		for (int zi = min_idx.Z; zi <= max_idx.Z; zi++)
+		Indexer.IterateAcrossBounds(
+			QueryPoint - SearchRadius * TVector<RealType>::One(),
+			QueryPoint + SearchRadius * TVector<RealType>::One(), [&CenterIdx, &SearchCell](const FVector3i& Idx)
 		{
-			for (int yi = min_idx.Y; yi <= max_idx.Y; yi++)
+			if (Idx == CenterIdx)
 			{
-				for (int xi = min_idx.X; xi <= max_idx.X; xi++)
-				{
-					FVector3i Idx(xi, yi, zi);
-					if (Idx != CenterIdx)
-					{
-						bool bFound = SearchCell(Idx);
-						if (bEarlyOut && bFound)
-						{
-							return TPair<PointDataType, RealType>(Nearest, MinDistSq);
-						}
-					}
-				}
+				return true; // continue
 			}
-		}
+			bool bFound = SearchCell(Idx);
+			return !(bEarlyOut && bFound);
+		});
 
 		if (Nearest == GetInvalidValue())
 		{
diff --git a/Engine/Source/Runtime/GeometryCore/Public/Util/GridIndexing3.h b/Engine/Source/Runtime/GeometryCore/Public/Util/GridIndexing3.h
--- a/Engine/Source/Runtime/GeometryCore/Public/Util/GridIndexing3.h
+++ b/Engine/Source/Runtime/GeometryCore/Public/Util/GridIndexing3.h
@@ -5,6 +5,7 @@
 #pragma once
 
 #include "VectorTypes.h"
+#include "HAL/Platform.h" // int32
 #include "IntVectorTypes.h"
 #include "BoxTypes.h"
 
@@ -52,7 +53,93 @@ struct TScaleGridIndexer3
 typedef TScaleGridIndexer3<float> FScaleGridIndexer3f;
 typedef TScaleGridIndexer3<double> FScaleGridIndexer3d;
 
+/**
+ * Converts real-valued coordinates to integer grid coordinates, wrapping around outside 
+ *  the representable integer range.
+ * 
+ * Lacks a conversion back from the integer vector since it maps to multiple ranges, but
+ *  this is frequently not needed. Also requires use of a function to iterate over a
+ *  range, to ensure proper overflow.
+ */
+template<typename RealType>
+struct TWrapAroundGridIndexer3
+{
+	/** Real-valued size of an integer grid cell */
+	RealType CellSize;
+
+	TWrapAroundGridIndexer3() : CellSize((RealType)1)
+	{
+	}
+
+	TWrapAroundGridIndexer3(RealType CellSize) : CellSize(CellSize)
+	{
+		// Note a very small cell size could overflow division by cell size
+		ensure(CellSize >= TMathUtil<RealType>::ZeroTolerance);
+	}
+
+	/** 
+	 * Convert real-valued point to integer grid coordinates, wrapping around when
+	 *  P/CellSize falls outside the representable integer range.
+	 */
+	inline FVector3i ToGrid(const TVector<RealType>& P) const
+	{
+		// The conversion below always goes through double. It could be done in RealType, but that is actually messy
+		//  and error-prone because for floats, we would not want to do the WraparoundDistance addition below, as it
+		//  loses precision unnecessarily for small negative numbers. So we would have to do more work to get the 
+		//  same kind of single-direction rounding and wraparound behavior, and it is not worth the bother (note if
+		//  you do choose to switch: WraparoundDistance initialization would need to change too, using std::nextafter
+		//  if we happen to round up)
+
+		static const double WraparoundDistance = static_cast<double>(TNumericLimits<uint32>::Max()) + 1;
+		auto GetInt = [](RealType Real) -> int32
+		{
+			double ModResult = FMath::Fmod(static_cast<double>(Real), WraparoundDistance);
+			ModResult = ModResult < 0 ? ModResult + WraparoundDistance : ModResult;
+			uint32 Truncated = static_cast<uint32>(ModResult);
+
+			// Seems that even with 2's complement being required in C++20, static_cast might theoretically still be
+			//  undefined outside representable range (?), hence the reinterpret here
+			return reinterpret_cast<int32&>(Truncated);
+		};
+
+		return FVector3i(
+			GetInt(P.X / CellSize),
+			GetInt(P.Y / CellSize),
+			GetInt(P.Z / CellSize));
+	}
+
+	/**
+	 * Iterate across a range of indices going from the one that corresponds to RealMean to one that corresponds
+	 *  to RealMax. Iteration stops early if the passed function returns false.
+	 */
+	void IterateAcrossBounds(const TVector<RealType>& RealMin, const TVector<RealType>& RealMax, 
+		TFunctionRef<bool(const FVector3i& Index)> ShouldContinue) const
+	{
+		FVector3i IntMin = ToGrid(RealMin);
+		FVector3i ExclusiveIntMax = ToGrid(RealMax) + FVector3i::One();
+
+		// Technically signed int overflow is undefined, so we use unsigned int to overflow safely.
+		//  Overflow is also the reason we have to use != for the comparison. 
+		for (uint32 zi = static_cast<uint32>(IntMin.Z); zi != static_cast<uint32>(ExclusiveIntMax.Z); ++zi)
+		{
+			for (uint32 yi = static_cast<uint32>(IntMin.Y); yi != static_cast<uint32>(ExclusiveIntMax.Y); ++yi)
+			{
+				for (uint32 xi = static_cast<uint32>(IntMin.X); xi != static_cast<uint32>(ExclusiveIntMax.X); ++xi)
+				{
+					// Seems that even with 2's complement being required in C++20, static_cast might theoretically still be
+					//  undefined outside representable range (?), hence the reinterpret here
+					if (!ShouldContinue(FVector3i(reinterpret_cast<int32&>(xi), reinterpret_cast<int32&>(yi), reinterpret_cast<int32&>(zi))))
+					{
+						return;
+					}
+				}
+			}
+		}
+	}
+};
 
+typedef TWrapAroundGridIndexer3<float> FWrapAroundGridIndexer3f;
+typedef TWrapAroundGridIndexer3<double> FWrapAroundGridIndexer3d;
 
 /**
  * Convert between integer grid coordinates and scaled+translated real-valued coordinates
