Adaptive tessellator: handle failed Split/Poke/Flip cases

diff --git a/Engine/Plugins/Runtime/GeometryProcessing/Source/DynamicMesh/Public/Tessellation/DynamicMeshPolicy.h b/Engine/Plugins/Runtime/GeometryProcessing/Source/DynamicMesh/Public/Tessellation/DynamicMeshPolicy.h
--- a/Engine/Plugins/Runtime/GeometryProcessing/Source/DynamicMesh/Public/Tessellation/DynamicMeshPolicy.h
+++ b/Engine/Plugins/Runtime/GeometryProcessing/Source/DynamicMesh/Public/Tessellation/DynamicMeshPolicy.h
@@ -203,6 +203,13 @@ class FDynamicMesh3Policy
 		FDynamicMesh3::FEdgeSplitInfo EdgeSplitInfo;
 		const UE::Geometry::EMeshResult Result = Mesh.SplitEdge(EdgeID, EdgeSplitInfo, Weight);
 
+		UE::Geometry::FEdgeSplitInfo SplitInfoResult;
+		if (Result != UE::Geometry::EMeshResult::Ok)
+		{
+			SplitInfoResult.NewTriIndex[0] = IndexConstants::InvalidID;
+			return SplitInfoResult;
+		}
+
 		if (bHasEdgeConstraints)
 		{
 			// update the edge constraints, the constraints on the edge remain
@@ -210,10 +217,7 @@ class FDynamicMesh3Policy
 			// the other 2 edges are new and don't inherit any constraints.
 			Constraints.SetOrUpdateEdgeConstraint(EdgeSplitInfo.NewEdges[0], Constraints.GetEdgeConstraint(EdgeID));
 		}
-
-		check(Result == UE::Geometry::EMeshResult::Ok);
-
-		UE::Geometry::FEdgeSplitInfo SplitInfoResult;
+				
 		SplitInfoResult.OldTriIndex = EdgeSplitInfo.OriginalTriangles;
 		SplitInfoResult.NewTriIndex = EdgeSplitInfo.NewTriangles;
 		SplitInfoResult.NewVertIndex[0] = EdgeSplitInfo.NewVertex;
@@ -229,10 +233,16 @@ class FDynamicMesh3Policy
 		const FIndex3i Triangle = Mesh.GetTriangle(TriIndex);
 
 		FDynamicMesh3::FPokeTriangleInfo PokeTriangleInfo;
-		const UE::Geometry::EMeshResult Result = Mesh.PokeTriangle(TriIndex, FVector3d(Barycentrics), PokeTriangleInfo);
-		check(Result == UE::Geometry::EMeshResult::Ok);
 
 		UE::Geometry::FPokeInfo PokeInfoResult;
+
+		const UE::Geometry::EMeshResult Result = Mesh.PokeTriangle(TriIndex, FVector3d(Barycentrics), PokeTriangleInfo);
+		if (Result != UE::Geometry::EMeshResult::Ok)
+		{
+			PokeInfoResult.NewTriIndex[0] = IndexConstants::InvalidID;
+			return PokeInfoResult;
+		}
+				
 		PokeInfoResult.NewTriIndex = { TriIndex, PokeTriangleInfo.NewTriangles[0], PokeTriangleInfo.NewTriangles[1] };
 		PokeInfoResult.EdgeIndex = { 0, 0, 0 };
 
@@ -268,8 +278,15 @@ class FDynamicMesh3Policy
 
 		check(!bHasEdgeConstraints || Constraints.GetEdgeConstraint(EdgeID).CanFlip());
 
+		UE::Geometry::FFlipEdgeInfo FlipEdgeInfo;
+
 		const UE::Geometry::EMeshResult Result = Mesh.FlipEdge(EdgeID, DMEdgeFlipInfo);
-		check(Result == UE::Geometry::EMeshResult::Ok);
+
+		if (Result != UE::Geometry::EMeshResult::Ok)
+		{
+			FlipEdgeInfo.Triangles = FIndex2i(IndexConstants::InvalidID, IndexConstants::InvalidID);
+			return FlipEdgeInfo;
+		}
 
 		// remove constraints from flipped edge, it's a new edge that can't inherit in a meaningful way
 		// from the original edge. if there were split constraints, it would have implied also setting a flip constraint
@@ -283,7 +300,7 @@ class FDynamicMesh3Policy
 		check(FlippedEdge.Vert[0] == FMath::Min(DMEdgeFlipInfo.OpposingVerts[0], DMEdgeFlipInfo.OpposingVerts[1]));
 		check(FlippedEdge.Vert[1] == FMath::Max(DMEdgeFlipInfo.OpposingVerts[0], DMEdgeFlipInfo.OpposingVerts[1]));
 
-		UE::Geometry::FFlipEdgeInfo FlipEdgeInfo;
+		
 		FlipEdgeInfo.Triangles = DMEdgeFlipInfo.Triangles;
 
 		// FDynamicMesh3::FlipEdge always constructs the triangles with shared edges first.
diff --git a/Engine/Source/Runtime/GeometryCore/Public/Tessellation/AdaptiveTessellator.h b/Engine/Source/Runtime/GeometryCore/Public/Tessellation/AdaptiveTessellator.h
--- a/Engine/Source/Runtime/GeometryCore/Public/Tessellation/AdaptiveTessellator.h
+++ b/Engine/Source/Runtime/GeometryCore/Public/Tessellation/AdaptiveTessellator.h
@@ -29,22 +29,22 @@ struct FEdgeSplitInfo
 {
 	using FIndex2i = UE::Geometry::FIndex2i;
 	FIndex2i OldTriIndex;  //< triangle indices before split
-	FIndex2i NewTriIndex;  //< newly added triangles indices NewTriIndex[i] will be on the same side as OldTriIndex[i] for i=0,1
+	FIndex2i NewTriIndex;  //< newly added triangles indices NewTriIndex[i] will be on the same side as OldTriIndex[i] for i=0,1, or IndexConstants::InvalidID on failure
 	FIndex2i NewVertIndex; //< indices of newly introduced vertices (NewVertIndex[i] is the shared vertex of OldTriIndex[i] and NewTriIndex[i])
 };
 
 struct FPokeInfo
 {
 	using FIndex3i = UE::Geometry::FIndex3i;
-	FIndex3i NewTriIndex  { -1, -1, -1 };  //< indices of the newly created triangles
-	FIndex3i EdgeIndex    { -1, -1, -1 };  //< for each of the new triangles the edge index corresponding to the original triangle edge
-	int32    NewVertIndex { -1 };          //< index of newly created vertex
+	FIndex3i NewTriIndex;  //< indices of the newly created triangles, or IndexConstants::InvalidID on failure
+	FIndex3i EdgeIndex;    //< for each of the new triangles the edge index corresponding to the original triangle edge
+	int32    NewVertIndex; //< index of newly created vertex
 };
 
 struct FFlipEdgeInfo
 {
 	using FIndex2i = UE::Geometry::FIndex2i;
-	FIndex2i Triangles;    //< triangle indices corresponding to the newly flipped triangles
+	FIndex2i Triangles;    //< triangle indices corresponding to the newly flipped triangles, or IndexConstants::InvalidID on failure
 	FIndex2i SharedEdge;   //< local edge index [0..2] of the new shared edge, with respect to triangle indices above
 };
 
@@ -698,6 +698,11 @@ inline void FAdaptiveTessellator<MeshT, DisplacementPolicyT>::SplitTriangle( uin
 
 			const FEdgeSplitInfo SplitInfo = Mesh.SplitEdge( TriIndex, EdgeIndex, Barycentrics[EdgeIndex] );
 
+			if (SplitInfo.NewTriIndex[0] == IndexConstants::InvalidID)
+			{
+				return;
+			}
+
 			int NumNewTris = 0;
 			// evaluate the displacement at the newly created vertices, possibly material seam
 			// requires to evaluate displacement twice
@@ -745,22 +750,26 @@ inline void FAdaptiveTessellator<MeshT, DisplacementPolicyT>::SplitTriangle( uin
 	}
 
 	// regular interior point split (poke)
-	const VecType NewDisplacement = DisplacementPolicy.GetDisplacement( Barycentrics, TriIndex );
 	const FPokeInfo pokeInfo = Mesh.PokeTriangle(TriIndex, Barycentrics);
 
-	check(pokeInfo.NewTriIndex[0] == TriIndex);
-	
-	GrowTriRecords(FMath::Max(pokeInfo.NewTriIndex[1], pokeInfo.NewTriIndex[2])+1);
-	check(pokeInfo.NewTriIndex[1] < Mesh.MaxTriID());
-	check(pokeInfo.NewTriIndex[2] < Mesh.MaxTriID());
-
-	for( uint32 LocalTriIndex = 0; LocalTriIndex < 3; ++LocalTriIndex )
+	if (pokeInfo.NewTriIndex[0] != IndexConstants::InvalidID)
 	{
-		AddFindRequest( pokeInfo.NewTriIndex[LocalTriIndex] );
-		TryDelaunayFlip( pokeInfo.NewTriIndex[LocalTriIndex], pokeInfo.EdgeIndex[LocalTriIndex] );
+		const VecType NewDisplacement = DisplacementPolicy.GetDisplacement(Barycentrics, TriIndex);
+
+		check(pokeInfo.NewTriIndex[0] == TriIndex);
+
+		GrowTriRecords(FMath::Max(pokeInfo.NewTriIndex[1], pokeInfo.NewTriIndex[2]) + 1);
+		check(pokeInfo.NewTriIndex[1] < Mesh.MaxTriID());
+		check(pokeInfo.NewTriIndex[2] < Mesh.MaxTriID());
+
+		for (uint32 LocalTriIndex = 0; LocalTriIndex < 3; ++LocalTriIndex)
+		{
+			AddFindRequest(pokeInfo.NewTriIndex[LocalTriIndex]);
+			TryDelaunayFlip(pokeInfo.NewTriIndex[LocalTriIndex], pokeInfo.EdgeIndex[LocalTriIndex]);
+		}
+
+		SetDisplacementAt(pokeInfo.NewVertIndex, NewDisplacement);
 	}
-	
-	SetDisplacementAt(pokeInfo.NewVertIndex, NewDisplacement);
 }
 
 template <typename MeshT, typename DisplacementPolicyT>
@@ -867,22 +876,27 @@ inline void FAdaptiveTessellator<MeshT, DisplacementPolicyT>::TryDelaunayFlip( i
 	const bool bFlipEdge = LaplacianWeight < RealType(-1e-6);
 	if( bFlipEdge )
 	{
-		const int32 AdjTriIndex     = OppositeEdge.Get<0>();
-		const int32 AdjTriEdgeIndex = OppositeEdge.Get<1>();
+		const FFlipEdgeInfo FlipEdgeInfo = Mesh.FlipEdge(TriIndex, TriEdgeIndex);
+
+		if (FlipEdgeInfo.Triangles[0] != IndexConstants::InvalidID)
+		{
+			check(FlipEdgeInfo.Triangles[1] != IndexConstants::InvalidID);
 
-		RemoveSplitRequest( TriIndex );
-		RemoveSplitRequest( AdjTriIndex );
+			const int32 AdjTriIndex = OppositeEdge.Get<0>();
+			const int32 AdjTriEdgeIndex = OppositeEdge.Get<1>();
 
-		AddFindRequest( TriIndex );
-		AddFindRequest( AdjTriIndex );
+			RemoveSplitRequest(TriIndex);
+			RemoveSplitRequest(AdjTriIndex);
 
-		const auto FlipEdgeInfo = Mesh.FlipEdge(TriIndex, TriEdgeIndex);
+			AddFindRequest(TriIndex);
+			AddFindRequest(AdjTriIndex);
 
-		TryDelaunayFlip(FlipEdgeInfo.Triangles[0], NextEdge(FlipEdgeInfo.SharedEdge[0]));
-		TryDelaunayFlip(FlipEdgeInfo.Triangles[0], PrevEdge(FlipEdgeInfo.SharedEdge[0]));
+			TryDelaunayFlip(FlipEdgeInfo.Triangles[0], NextEdge(FlipEdgeInfo.SharedEdge[0]));
+			TryDelaunayFlip(FlipEdgeInfo.Triangles[0], PrevEdge(FlipEdgeInfo.SharedEdge[0]));
 
-		TryDelaunayFlip(FlipEdgeInfo.Triangles[1], PrevEdge(FlipEdgeInfo.SharedEdge[1]));
-		TryDelaunayFlip(FlipEdgeInfo.Triangles[1], NextEdge(FlipEdgeInfo.SharedEdge[1]));
+			TryDelaunayFlip(FlipEdgeInfo.Triangles[1], PrevEdge(FlipEdgeInfo.SharedEdge[1]));
+			TryDelaunayFlip(FlipEdgeInfo.Triangles[1], NextEdge(FlipEdgeInfo.SharedEdge[1]));
+		}
 	}
 }
 
