DynamicMesh3 bulk initialize

#rnx

diff --git a/Engine/Source/Runtime/GeometryCore/Private/Util/SmallListSet.cpp b/Engine/Source/Runtime/GeometryCore/Private/Util/SmallListSet.cpp
--- a/Engine/Source/Runtime/GeometryCore/Private/Util/SmallListSet.cpp
+++ b/Engine/Source/Runtime/GeometryCore/Private/Util/SmallListSet.cpp
@@ -19,6 +19,27 @@ void FSmallListSet::Resize(int32 NewSize)
 	}
 }
 
+void FSmallListSet::ResizeAndAllocateBlocks(int32 NewSize)
+{
+	Reset();
+	Resize(NewSize);
+	ListBlocks.Resize(NewSize * (BLOCK_LIST_OFFSET+1));
+	for (int32 i(0); i<NewSize; ++i)
+	{
+		int32 ListHead = i * (BLOCK_LIST_OFFSET+1);
+		ListHeads[i] = ListHead;
+		ListBlocks[ListHead] = 0;
+		ListBlocks[ListHead + BLOCK_LIST_OFFSET] = NullValue;
+	}
+
+	for (int32 i(0); i < NewSize; ++i)
+	{
+		checkSlow(ListBlocks[ListHeads[i]] == 0);
+		checkSlow(ListBlocks[ListHeads[i] + BLOCK_LIST_OFFSET] == NullValue);
+	}
+
+	AllocatedCount = NewSize;
+}
 
 void FSmallListSet::AllocateAt(int32 ListIndex)
 {
diff --git a/Engine/Source/Runtime/GeometryCore/Public/DynamicMesh/DynamicMesh3.h b/Engine/Source/Runtime/GeometryCore/Public/DynamicMesh/DynamicMesh3.h
--- a/Engine/Source/Runtime/GeometryCore/Public/DynamicMesh/DynamicMesh3.h
+++ b/Engine/Source/Runtime/GeometryCore/Public/DynamicMesh/DynamicMesh3.h
@@ -1502,7 +1502,7 @@ public:
 	 * Returns a debug string that contains mesh statistics and other information
 	 */
 	GEOMETRYCORE_API virtual FString MeshInfoString() const;
-
+	
 	GEOMETRYCORE_API SIZE_T GetByteCount() const;
 
 	/**
@@ -1525,7 +1525,7 @@ public:
 		float EpsilonFloat { TMathUtil<float>::Epsilon };
 
 		// comparisons check for smaller or equal than this value. used for 64-bit comparisons, such as on points.
-		double EpsilonDouble{ static_cast<double>(TMathUtil<float>::Epsilon) };
+		double EpsilonDouble { static_cast<double>(TMathUtil<float>::Epsilon) };
 	};
 
 	struct FMeshDifferenceInfo
@@ -1808,6 +1808,7 @@ private:
 		const FCollapseEdgeOptions& Options, FEdgeCollapseInfo* OutCollapseInfo) const;
 
 	friend class FDynamicMeshAttributeSet;
+	friend class FDynamicMeshBulkEdit;
 };
 
 
diff --git a/Engine/Source/Runtime/GeometryCore/Public/Util/RefCountVector.h b/Engine/Source/Runtime/GeometryCore/Public/Util/RefCountVector.h
--- a/Engine/Source/Runtime/GeometryCore/Public/Util/RefCountVector.h
+++ b/Engine/Source/Runtime/GeometryCore/Public/Util/RefCountVector.h
@@ -325,6 +325,14 @@ public:
 		UsedCount = 0;
 	}
 
+	// initialize and set all refcounts to given value
+	void InitDense(int Size, uint16 RefCountValue = 1)
+	{
+		FreeIndices.Clear();
+		RefCounts.Clear();
+		RefCounts.Resize(Size, RefCountValue);
+		UsedCount = Size;
+	}
 	//
 	// Iterators
 	//
diff --git a/Engine/Source/Runtime/GeometryCore/Public/Util/SmallListSet.h b/Engine/Source/Runtime/GeometryCore/Public/Util/SmallListSet.h
--- a/Engine/Source/Runtime/GeometryCore/Public/Util/SmallListSet.h
+++ b/Engine/Source/Runtime/GeometryCore/Public/Util/SmallListSet.h
@@ -82,6 +82,12 @@ public:
 	 */
 	GEOMETRYCORE_API void Resize(int32 NewSize);
 
+
+	/**
+	 * resize to a new number of lists and pre-allocate an initial empty block for each list.
+	 */
+	GEOMETRYCORE_API void ResizeAndAllocateBlocks(int32 NewSize);
+
 	/**
 	 * Reset to initial state
 	 */
