[AutoRTFM] Guard changestamps with a mutex instead of using atomics.

This code is reachable from Verse because a cinematic sequence can be triggered, creating a body of water, which necessitates the creation of a collision mesh. Using a transactionally-safe mutex avoids a server crash.

In local testing, the GeometryProcessing unit test runtime seemed entirely unaffected by this change, despite using change stamps in many places.

The old members have been marked with UE_DEPRECATED even though they were not public, since it would have been possible for a licensee to subclass the dynamic mesh and then access them directly. The new members are in a private section to avoid this possibility.

diff --git a/Engine/Source/Runtime/GeometryCore/Private/DynamicMesh/DynamicMesh3.cpp b/Engine/Source/Runtime/GeometryCore/Private/DynamicMesh/DynamicMesh3.cpp
--- a/Engine/Source/Runtime/GeometryCore/Private/DynamicMesh/DynamicMesh3.cpp
+++ b/Engine/Source/Runtime/GeometryCore/Private/DynamicMesh/DynamicMesh3.cpp
@@ -53,8 +53,8 @@ FDynamicMesh3::FDynamicMesh3(const FDynamicMesh3& Other)
 	Edges{ Other.Edges },
 	EdgeRefCounts{ Other.EdgeRefCounts }
 {
-	ShapeChangeStamp.store(Other.ShapeChangeStamp);
-	TopologyChangeStamp.store(Other.TopologyChangeStamp);
+	ChangeStampShape.Set(Other.ChangeStampShape.GetValue());
+	ChangeStampTopology.Set(Other.ChangeStampTopology.GetValue());
 
 	if (Other.HasAttributes())
 	{
@@ -83,8 +83,8 @@ FDynamicMesh3::FDynamicMesh3(FDynamicMesh3&& Other)
 
 	AttributeSet{ MoveTemp( Other.AttributeSet ) }
 {
-	ShapeChangeStamp.store(Other.ShapeChangeStamp);
-	TopologyChangeStamp.store(Other.TopologyChangeStamp);
+	ChangeStampShape.Set(Other.ChangeStampShape.GetValue());
+	ChangeStampTopology.Set(Other.ChangeStampTopology.GetValue());
 
 	if (AttributeSet)
 	{
@@ -124,8 +124,8 @@ const FDynamicMesh3 & FDynamicMesh3::operator=(FDynamicMesh3 && Other)
 		AttributeSet->Reparent(this);
 	}
 	
-	ShapeChangeStamp.store(Other.ShapeChangeStamp);
-	TopologyChangeStamp.store(Other.TopologyChangeStamp);
+	ChangeStampShape.Set(Other.ChangeStampShape.GetValue());
+	ChangeStampTopology.Set(Other.ChangeStampTopology.GetValue());
 
 	return *this;
 }
@@ -235,8 +235,8 @@ void FDynamicMesh3::Copy(const FDynamicMesh3& copy, bool bNormals, bool bColors,
 		DiscardAttributes();
 	}
 
-	ShapeChangeStamp.store(copy.ShapeChangeStamp);
-	TopologyChangeStamp.store(copy.TopologyChangeStamp);
+	ChangeStampShape.Set(copy.ChangeStampShape.GetValue());
+	ChangeStampTopology.Set(copy.ChangeStampTopology.GetValue());
 }
 
 void FDynamicMesh3::AppendWithOffsets(const FDynamicMesh3& ToAppend, FAppendInfo* OutAppendInfo)
@@ -425,8 +425,8 @@ void FDynamicMesh3::CompactCopy(const FDynamicMesh3& copy, bool bNormals, bool b
 		AttributeSet->CompactCopy(*CompactInfo, *copy.Attributes());
 	}
 
-	ShapeChangeStamp.store(copy.ShapeChangeStamp);
-	TopologyChangeStamp.store(copy.TopologyChangeStamp);
+	ChangeStampShape.Set(copy.ChangeStampShape.GetValue());
+	ChangeStampTopology.Set(copy.ChangeStampTopology.GetValue());
 }
 
 void FDynamicMesh3::Clear()
@@ -449,8 +449,8 @@ void FDynamicMesh3::Clear()
 
 	AttributeSet.Reset();
 
-	ShapeChangeStamp.store(1);
-	TopologyChangeStamp.store(1);
+	ChangeStampShape.Set(1);
+	ChangeStampTopology.Set(1);
 }
 
 
diff --git a/Engine/Source/Runtime/GeometryCore/Private/Tests/AutoRTFMDynamicMesh3Tests.cpp b/Engine/Source/Runtime/GeometryCore/Private/Tests/AutoRTFMDynamicMesh3Tests.cpp
new file mode 100644
--- /dev/null
+++ b/Engine/Source/Runtime/GeometryCore/Private/Tests/AutoRTFMDynamicMesh3Tests.cpp
@@ -0,0 +1,116 @@
+// Copyright Epic Games, Inc. All Rights Reserved.
+
+#include "Misc/AutomationTest.h"
+#include "AutoRTFM.h"
+#include "Math/UnrealMathUtility.h"
+#include "Math/Vector.h"
+#include "OrientedBoxTypes.h"
+#include "DynamicMesh/DynamicMesh3.h"
+#include "Generators/GridBoxMeshGenerator.h"
+
+#if WITH_DEV_AUTOMATION_TESTS
+
+IMPLEMENT_SIMPLE_AUTOMATION_TEST(FAutoRTFMDynamicMesh3Test, "AutoRTFM + FDynamicMesh3", EAutomationTestFlags::EngineFilter | EAutomationTestFlags::ClientContext | EAutomationTestFlags::ServerContext | EAutomationTestFlags::CommandletContext)
+
+bool FAutoRTFMDynamicMesh3Test::RunTest(const FString& Parameters)
+{
+	if (!AutoRTFM::ForTheRuntime::IsAutoRTFMRuntimeEnabled())
+	{
+		ExecutionInfo.AddEvent(FAutomationEvent(EAutomationEventType::Info, TEXT("SKIPPED 'FAutoRTFMDynamicMesh3Test' test. AutoRTFM disabled.")));
+		return true;
+	}
+
+	// Setup code is adapted from TEST_CASE("Remove Vertices") in DynamicMeshTests.cpp.
+	UE::Geometry::FGridBoxMeshGenerator Gen;
+	Gen.Box = UE::Geometry::FOrientedBox3d{ FVector3d{ 0,0,0 }, FVector3d{ 1.0,1.0,1.0 } };
+	Gen.EdgeVertices = UE::Geometry::FIndex3i{ 3, 3, 3 };
+	UE::Geometry::FDynamicMesh3 Mesh(&Gen.Generate());
+
+	Mesh.SetShapeChangeStampEnabled(true);
+	Mesh.SetTopologyChangeStampEnabled(true);
+
+	bool bSuccess = true;
+	uint32 StartingShapeChangeStamp = Mesh.GetShapeChangeStamp();
+	uint32 StartingTopologyChangeStamp = Mesh.GetTopologyChangeStamp();
+	constexpr int VertexID = 10;  // any vertex ID will do
+
+	// Transaction abort after `UpdateChangeStamps(true, false)` restores the starting ShapeChangeStamp.
+	AutoRTFM::Transact([&]
+	{
+		Mesh.UpdateChangeStamps(true, false);
+		check(Mesh.GetShapeChangeStamp()    >  StartingShapeChangeStamp);
+		check(Mesh.GetTopologyChangeStamp() == StartingTopologyChangeStamp);
+		AutoRTFM::AbortTransaction();
+	});
+
+	bSuccess &= TestTrueExpr(Mesh.GetShapeChangeStamp()    == StartingShapeChangeStamp);
+	bSuccess &= TestTrueExpr(Mesh.GetTopologyChangeStamp() == StartingTopologyChangeStamp);
+
+	// Transaction abort after `UpdateChangeStamps(true, true)` restores the starting ShapeChangeStamp and TopologyChangeStamp.
+	AutoRTFM::Transact([&]
+	{
+		Mesh.UpdateChangeStamps(true, true);
+		check(Mesh.GetShapeChangeStamp()    > StartingShapeChangeStamp);
+		check(Mesh.GetTopologyChangeStamp() > StartingTopologyChangeStamp);
+		AutoRTFM::AbortTransaction();
+	});
+
+	bSuccess &= TestTrueExpr(Mesh.GetShapeChangeStamp()    == StartingShapeChangeStamp);
+	bSuccess &= TestTrueExpr(Mesh.GetTopologyChangeStamp() == StartingTopologyChangeStamp);
+
+	// Transaction abort after `RemoveVertex` restores the starting ShapeChangeStamp and TopologyChangeStamp.
+	AutoRTFM::Transact([&]
+	{
+		check(Mesh.RemoveVertex(VertexID, /*bPreserveManifold=*/ false) == UE::Geometry::EMeshResult::Ok);
+		check(Mesh.GetShapeChangeStamp()    > StartingShapeChangeStamp);
+		check(Mesh.GetTopologyChangeStamp() > StartingTopologyChangeStamp);
+		AutoRTFM::AbortTransaction();
+	});
+
+	bSuccess &= TestTrueExpr(Mesh.GetShapeChangeStamp()    == StartingShapeChangeStamp);
+	bSuccess &= TestTrueExpr(Mesh.GetTopologyChangeStamp() == StartingTopologyChangeStamp);
+
+	// Transaction commit after `UpdateChangeStamps(true, false)` keeps the new ShapeChangeStamp.
+	AutoRTFM::Transact([&]
+	{
+		Mesh.UpdateChangeStamps(true, false);
+		check(Mesh.GetShapeChangeStamp()    >  StartingShapeChangeStamp);
+		check(Mesh.GetTopologyChangeStamp() == StartingTopologyChangeStamp);
+	});
+
+	bSuccess &= TestTrueExpr(Mesh.GetShapeChangeStamp()    >  StartingShapeChangeStamp);
+	bSuccess &= TestTrueExpr(Mesh.GetTopologyChangeStamp() == StartingTopologyChangeStamp);
+
+	// Transaction commit after `UpdateChangeStamps(true, true)` keeps the new ShapeChangeStamp and TopologyChangeStamp.
+	StartingShapeChangeStamp = Mesh.GetShapeChangeStamp();
+	StartingTopologyChangeStamp = Mesh.GetTopologyChangeStamp();
+
+	AutoRTFM::Transact([&]
+	{
+		Mesh.UpdateChangeStamps(true, true);
+		check(Mesh.GetShapeChangeStamp()    > StartingShapeChangeStamp);
+		check(Mesh.GetTopologyChangeStamp() > StartingTopologyChangeStamp);
+	});
+
+	bSuccess &= TestTrueExpr(Mesh.GetShapeChangeStamp()    > StartingShapeChangeStamp);
+	bSuccess &= TestTrueExpr(Mesh.GetTopologyChangeStamp() > StartingTopologyChangeStamp);
+
+	// Transaction commit after `RemoveVertex` keeps the new ShapeChangeStamp and TopologyChangeStamp.
+	StartingShapeChangeStamp = Mesh.GetShapeChangeStamp();
+	StartingTopologyChangeStamp = Mesh.GetTopologyChangeStamp();
+
+	AutoRTFM::Transact([&]
+	{
+		check(Mesh.RemoveVertex(VertexID, /*bPreserveManifold=*/ false) == UE::Geometry::EMeshResult::Ok);
+		check(Mesh.GetShapeChangeStamp()    > StartingShapeChangeStamp);
+		check(Mesh.GetTopologyChangeStamp() > StartingTopologyChangeStamp);
+	});
+
+	bSuccess &= TestTrueExpr(Mesh.GetShapeChangeStamp()    > StartingShapeChangeStamp);
+	bSuccess &= TestTrueExpr(Mesh.GetTopologyChangeStamp() > StartingTopologyChangeStamp);
+
+	return bSuccess;
+}
+
+#endif //WITH_DEV_AUTOMATION_TESTS
+
diff --git a/Engine/Source/Runtime/GeometryCore/Public/DynamicMesh/DynamicMesh3.h b/Engine/Source/Runtime/GeometryCore/Public/DynamicMesh/DynamicMesh3.h
--- a/Engine/Source/Runtime/GeometryCore/Public/DynamicMesh/DynamicMesh3.h
+++ b/Engine/Source/Runtime/GeometryCore/Public/DynamicMesh/DynamicMesh3.h
@@ -4,6 +4,7 @@
 
 #pragma once
 
+#include "Async/TransactionallySafeMutex.h"
 #include "BoxTypes.h"
 #include "Containers/Array.h"
 #include "Containers/ArrayView.h"
@@ -193,15 +194,74 @@ protected:
 	/** Extended Attributes for the Mesh (UV layers, Hard Normals, additional Polygroup Layers, etc) */
 	TUniquePtr<FDynamicMeshAttributeSet> AttributeSet{};
 
-	/** 
-	 * Enable/Disable updating of ShapeChangeStamp. This can be problematic in multi-threaded contexts so it is disabled by default. 
-	 * In fact, it is not suggested that these be used at all (see commment for SetShapeChangeStampEnabled).
-	 */
+protected:
+	UE_DEPRECATED(5.7, "bEnableShapeChangeStamp is no longer used. Use SetShapeChangeStampEnabled and HasShapeChangeStampEnabled instead.") 
 	bool bEnableShapeChangeStamp = false;
-	/** If bEnableShapeChangeStamp=true, The shape change stamp is incremented any time a function that modifies the mesh shape or topology is called */
-	std::atomic<uint32> ShapeChangeStamp = 1;
-	/** The topology change stamp is incremented any time a function that modifies the mesh topology is called */
-	std::atomic<uint32> TopologyChangeStamp = 1;
+	
+	UE_DEPRECATED(5.7, "ShapeChangeStamp is no longer used. Use GetShapeChangeStamp or UpdateChangeStamps instead.") 
+	std::atomic<uint32> ShapeChangeStamp = 0;
+	UE_DEPRECATED(5.7, "TopologyChangeStamp is no longer used. Use GetTopologyChangeStamp or UpdateChangeStamps instead.") 
+	std::atomic<uint32> TopologyChangeStamp = 0;
+
+private:
+	struct ChangeStamp
+	{
+		/** Enable/disable this change stamp. */
+		bool bIsEnabled = false;
+
+		/** The change stamp is incremented when modifications occur, if `bIsEnabled` is set. It's guarded by Mutex. */
+		uint32 Value = 1;
+
+		/** Guards Value. */
+		FTransactionallySafeMutex& Mutex;
+
+		/** Updates the change stamp in an thread-safe, transactionally-safe way. Does nothing if `bIsEnabled` is not set. */
+		void IncrementIfEnabled()
+		{
+			if (bIsEnabled)
+			{
+				Mutex.Lock();
+				++Value;
+				Mutex.Unlock();
+			}
+		}
+
+		/** Overwrites the change stamp in an thread-safe, transactionally-safe way. Does nothing if `bIsEnabled` is not set. */
+		void Set(uint32 NewValue)
+		{
+			if (bIsEnabled)
+			{
+				Mutex.Lock();
+				Value = NewValue;
+				Mutex.Unlock();
+			}
+		}
+
+		/** Returns the current change value in a thread-safe, transactionally-safe way. Returns 1 if `bIsEnabled` is not set. */
+		uint32 GetValue() const
+		{
+			uint32 Result = 1; 
+
+			if (bIsEnabled)
+			{
+				Mutex.Lock();
+				Result = Value;
+				Mutex.Unlock();
+			}
+
+			return Result;
+		}
+	};
+
+	// We share this mutex across both change stamps since it's only held for a few cycles.
+	FTransactionallySafeMutex ChangeStampMutex;
+
+	// Shape change tracking can be problematic in multi-threaded contexts so they're disabled by default. 
+	// In fact, it is not suggested that these be used at all (see comment for SetShapeChangeStampEnabled.
+	ChangeStamp ChangeStampShape{.bIsEnabled = false, .Mutex = ChangeStampMutex};
+
+	// Topological change tracking is enabled by default.
+	ChangeStamp ChangeStampTopology{.bIsEnabled = true, .Mutex = ChangeStampMutex};
 
 public:
 	/** Default constructor */
@@ -416,49 +476,66 @@ public:
 	 */
 	void SetShapeChangeStampEnabled(bool bEnabled)
 	{
-		bEnableShapeChangeStamp = bEnabled;
+		ChangeStampShape.bIsEnabled = bEnabled;
 	}
 
-	/** @return true if Shape ChangeStamp is enabled (disabled by default) */
+	/** Enable/disable incrementing of the TopologyChangeStamp. */
+	void SetTopologyChangeStampEnabled(bool bEnabled)
+	{
+		ChangeStampTopology.bIsEnabled = bEnabled;
+	}
+
+	/** @return true if shape ChangeStamp is enabled (disabled by default) */
 	bool HasShapeChangeStampEnabled() const
 	{
-		return bEnableShapeChangeStamp;
+		return ChangeStampShape.bIsEnabled;
+	}
+
+	/** @return true if topology ChangeStamp is enabled (disabled by default) */
+	bool HasTopologyChangeStampEnabled() const
+	{
+		return ChangeStampTopology.bIsEnabled;
 	}
 
 	/** Increment the specified ChangeStamps, if they are enabled. Thread-safe. */
 	void UpdateChangeStamps(bool bShapeChange, bool bTopologyChange)
 	{
-		if (bEnableShapeChangeStamp && (bShapeChange || bTopologyChange))
+		if (bShapeChange || bTopologyChange)
 		{
-			ShapeChangeStamp++;		// atomic increment
+			ChangeStampShape.IncrementIfEnabled();
 		}
 		if (bTopologyChange)
 		{
-			TopologyChangeStamp++;	// atomic increment
+			ChangeStampTopology.IncrementIfEnabled();
 		}
 	}
 
-	/** ShapeChangeStamp is incremented any time a mesh vertex position is changed or the mesh topology is modified, if bEnableShapeChangeStamp=true */
+	/** 
+	 * Returns the current ShapeChangeStamp. This is incremented any time a mesh vertex position is changed _or_ the mesh topology is modified. 
+	 * Change stamps are disabled by default.
+	 */
 	inline uint32 GetShapeChangeStamp() const
 	{
-		ensureMsgf(bEnableShapeChangeStamp, TEXT("ShapeChange Tracking is not enabled on this mesh. Use SetShapeChangeStampEnabled() to enable, or use GetChangeStamp() to avoid this ensure."));
-		return ShapeChangeStamp;
+		ensureMsgf(ChangeStampShape.bIsEnabled, TEXT("Shape change tracking is not enabled on this mesh. Use SetShapeChangeStampEnabled() to enable."));
+		return ChangeStampShape.GetValue();
 	}
 
-	/** TopologyChangeStamp is incremented when the mesh topology is modified */
+	/** 
+	 * Returns the current TopologyChangeStamp. This is incremented when the mesh topology is modified. 
+	 * Change stamps are disabled by default. 
+	 */
 	inline uint32 GetTopologyChangeStamp() const
 	{
-		return TopologyChangeStamp;
+		ensureMsgf(ChangeStampTopology.bIsEnabled, TEXT("Topology change tracking is not enabled on this mesh. Use SetTopologyChangeStampEnabled() to enable."));
+		return ChangeStampTopology.GetValue();
 	}
 
-	/** ChangeStamp is a combination of the Shape and Topology ChangeStamps */
+	/** ChangeStamp is a combination of the Shape and Topology ChangeStamps. If neither flag is enabled, this value will never change. */
 	inline uint64 GetChangeStamp() const
 	{
-		return ShapeChangeStamp + TopologyChangeStamp;
+		return ChangeStampShape.GetValue() + ChangeStampTopology.GetValue();
 	}
 
-
-
 	//
 	// Mesh Element Iterators
 	//   The functions VertexIndicesItr() / TriangleIndicesItr() / EdgeIndicesItr() allow you to do:
