Rename struct `ChangeStamp` to `FChangeStamp` to follow Epic style rules.

diff --git a/Engine/Source/Runtime/GeometryCore/Public/DynamicMesh/DynamicMesh3.h b/Engine/Source/Runtime/GeometryCore/Public/DynamicMesh/DynamicMesh3.h
--- a/Engine/Source/Runtime/GeometryCore/Public/DynamicMesh/DynamicMesh3.h
+++ b/Engine/Source/Runtime/GeometryCore/Public/DynamicMesh/DynamicMesh3.h
@@ -204,7 +204,7 @@ protected:
 	std::atomic<uint32> TopologyChangeStamp = 0;
 
 private:
-	struct ChangeStamp
+	struct FChangeStamp
 	{
 		/** Enable/disable this change stamp. */
 		bool bIsEnabled = false;
@@ -257,11 +257,11 @@ private:
 	FTransactionallySafeMutex ChangeStampMutex;
 
 	// Shape change tracking can be problematic in multi-threaded contexts so they're disabled by default. 
-	// In fact, it is not suggested that these be used at all (see comment for SetShapeChangeStampEnabled.
-	ChangeStamp ChangeStampShape{.bIsEnabled = false, .Mutex = ChangeStampMutex};
+	// In fact, it is not suggested that these be used at all. (See comment for SetShapeChangeStampEnabled.)
+	FChangeStamp ChangeStampShape{.bIsEnabled = false, .Mutex = ChangeStampMutex};
 
 	// Topological change tracking is enabled by default.
-	ChangeStamp ChangeStampTopology{.bIsEnabled = true, .Mutex = ChangeStampMutex};
+	FChangeStamp ChangeStampTopology{.bIsEnabled = true, .Mutex = ChangeStampMutex};
 
 public:
 	/** Default constructor */
