Splines - Add segment semantics to spline interface

Splines don't necessarily have points or mutable segments, but they do have segments which have parameter intervals because a spline is by definition a continuous piecewise function. Each piece is a segment.

This also allows fetching a TSplineInterface from FSpline for code which intends to operate using our spline core interfaces directly.

diff --git a/Engine/Source/Runtime/Engine/Classes/Curves/Spline.h b/Engine/Source/Runtime/Engine/Classes/Curves/Spline.h
--- a/Engine/Source/Runtime/Engine/Classes/Curves/Spline.h
+++ b/Engine/Source/Runtime/Engine/Classes/Curves/Spline.h
@@ -4,6 +4,7 @@
 
 #include "CoreMinimal.h"
 #include "Math/InterpCurve.h"
+#include "Splines/SplineInterfaces.h"
 
 #include "Spline.generated.h"
 
@@ -143,6 +144,8 @@ struct FSpline
 	/** Update the spline's internal data according to the passed-in params. */
 	ENGINE_API void UpdateSpline(const FUpdateSplineParams& InParams);
 
+	ENGINE_API TSharedPtr<UE::Geometry::Spline::TSplineInterface<FVector>> GetSplineInterface() const;
+
 private:
 	
 	static inline const FInterpCurveVector PositionCurve;
diff --git a/Engine/Source/Runtime/Engine/Private/Curves/Spline.cpp b/Engine/Source/Runtime/Engine/Private/Curves/Spline.cpp
--- a/Engine/Source/Runtime/Engine/Private/Curves/Spline.cpp
+++ b/Engine/Source/Runtime/Engine/Private/Curves/Spline.cpp
@@ -1558,6 +1558,16 @@ void FSpline::UpdateSpline(const FUpdateSplineParams& InParams)
     }
 }
 
+TSharedPtr<TSplineInterface<FVector>> FSpline::GetSplineInterface() const
+{
+	if (IsNew())
+	{
+		return StaticCastSharedPtr<TSplineInterface<FVector>>(NewData);
+	}
+
+	return nullptr;
+}
+
 #if WITH_EDITOR
 void FSpline::OnSplineImplementationChanged()
 {
@@ -1804,7 +1814,7 @@ float FNewSpline::GetSegmentLength(const int32 Index, const float Param, const F
 float FNewSpline::GetSplineLength()
 {	
 	// Evaluate the ParamToLength channel at the very end of the spline
-	const float MaxParameter = static_cast<float>(GetSpline().GetNumSegments());
+	const float MaxParameter = static_cast<float>(GetSpline().GetNumberOfSegments());
     return GetDistanceAtParameter(MaxParameter);
 }
 
@@ -2279,12 +2289,12 @@ void FNewSpline::UpdateSpline(const FSpline::FUpdateSplineParams& Params)
 
 float FNewSpline::FromNormalizedSpace(float Parameter) const
 {
-	return GetSpline().GetNumSegments() * Parameter;
+	return GetSpline().GetNumberOfSegments() * Parameter;
 }
 
 float FNewSpline::ToNormalizedSpace(float Parameter) const
 {
-	return Parameter / GetSpline().GetNumSegments();
+	return Parameter / GetSpline().GetNumberOfSegments();
 }
 
 float FNewSpline::FromInternalSplineSpace(float Parameter) const
@@ -2295,7 +2305,7 @@ float FNewSpline::FromInternalSplineSpace(float Parameter) const
 }
 float FNewSpline::ToInternalSplineSpace(float Parameter) const
 {
-	const int32 NumSegments = GetSpline().GetNumSegments();
+	const int32 NumSegments = GetSpline().GetNumberOfSegments();
 	float ClampedParam = FMath::Clamp(Parameter, 0.0f, static_cast<float>(NumSegments));
 	const int32 SegmentIndex = FMath::Min(FMath::FloorToInt(ClampedParam), NumSegments - 1);
 	const float LocalT = ClampedParam - static_cast<float>(SegmentIndex); // Parameter [0,1] within segment
@@ -2309,7 +2319,7 @@ void FNewSpline::UpdateReparamTable()
 		return;
 	}
 
-	const int32 NumSegments = GetSpline().GetNumSegments();
+	const int32 NumSegments = GetSpline().GetNumberOfSegments();
 	
 	// Start by clearing it
 	ClearAttributeChannel(ParamToLengthAttrName);
diff --git a/Engine/Source/Runtime/GeometryCore/Public/Splines/BSpline.h b/Engine/Source/Runtime/GeometryCore/Public/Splines/BSpline.h
--- a/Engine/Source/Runtime/GeometryCore/Public/Splines/BSpline.h
+++ b/Engine/Source/Runtime/GeometryCore/Public/Splines/BSpline.h
@@ -220,6 +220,26 @@ public:
 		return Clone;
 	}
 
+	/** Returns the number of Bezier segments in the spline */
+	virtual int32 GetNumberOfSegments() const override
+	{
+		// todo: implement
+		ensureAlwaysMsgf(false, TEXT("TBSpline::GetNumberOfSegments is unimplemented!"));
+		return INDEX_NONE;
+	}
+
+	/**
+	 * Maps a segment index to its parameter range
+	 * @param SegmentIndex - Index of the segment (0-based)
+	 * @return True if the segment index is valid and mapping succeeded
+	 */
+	virtual FInterval1f GetSegmentParameterRange(int32 SegmentIndex) const override
+	{
+		// todo: implement
+		ensureAlwaysMsgf(false, TEXT("TBSpline::GetSegmentParameterRange is unimplemented!"));
+		return FInterval1f();
+	}
+	
 	// TSplineInterface Implementation
 	virtual ValueType EvaluateImpl(float Parameter) const override
 	{
diff --git a/Engine/Source/Runtime/GeometryCore/Public/Splines/PolyBezierSpline.h b/Engine/Source/Runtime/GeometryCore/Public/Splines/PolyBezierSpline.h
--- a/Engine/Source/Runtime/GeometryCore/Public/Splines/PolyBezierSpline.h
+++ b/Engine/Source/Runtime/GeometryCore/Public/Splines/PolyBezierSpline.h
@@ -358,7 +358,7 @@ public:
 	        return ValueType();
 	    }
 
-    	int32 NumSegments = GetNumSegments();
+    	int32 NumSegments = GetNumberOfSegments();
     	// Map parameter to parent space
     	
 	    // For PolyBezier, each segment is normalized to unit parameter space
@@ -892,7 +892,7 @@ public:
     
         if (bIsLastSegment)
 	    {
-	    	return RemovePoint(GetNumSegments());
+	    	return RemovePoint(GetNumberOfSegments());
 	    }
 	    else
 	    {
@@ -907,7 +907,7 @@ public:
 	 */
 	bool RemovePoint(int32 PointIndex)
 	{
-		const int32 NumSegments = GetNumSegments();
+		const int32 NumSegments = GetNumberOfSegments();
         if (PointIndex < 0 || PointIndex > NumSegments)
         {
             return false;
@@ -932,7 +932,7 @@ public:
 			}
 		}
 		// else if last point
-		else if (PointIndex == GetNumSegments())
+		else if (PointIndex == GetNumberOfSegments())
 		{
 			// Removing the last point - remove the last segment
 			for (int32 i = 3; i >= 0; --i)
@@ -1024,14 +1024,30 @@ public:
     
         return bSuccess;
     }
+	
 	/** Returns the number of Bezier segments in the spline */
-	int32 GetNumSegments() const
+	virtual int32 GetNumberOfSegments() const override
 	{
 		const int32 NumPoints = this->Base::NumKeys();
 		return NumPoints/4;
 	}
 
-	
+	/**
+	 * Maps a segment index to its parameter range
+	 * @param SegmentIndex - Index of the segment (0-based)
+	 * @return True if the segment index is valid and mapping succeeded
+	 */
+	virtual FInterval1f GetSegmentParameterRange(int32 SegmentIndex) const override
+	{
+		FInterval1f SegmentRange;
+		if (SegmentIndex >= 0 && SegmentIndex < Base::PairKnots.Num() - 1)
+		{
+			SegmentRange.Min = Base::PairKnots[SegmentIndex].Value;
+			SegmentRange.Max = Base::PairKnots[SegmentIndex + 1].Value;
+		}
+        
+		return SegmentRange;
+	}
 
 	// Parameterization methods
 	/**
@@ -1100,7 +1116,7 @@ public:
 		constexpr int32 PointsPerSegment = 4;
 		const int32 SegmentIndex = Index / PointsPerSegment;
 		const int32 LocalPointIndex = Index % PointsPerSegment;
-		int32 NumSegments = GetNumSegments();
+		int32 NumSegments = GetNumberOfSegments();
 
 		// Get the original parameter value
 		int32 OldParameterIndex = LocalPointIndex == 0 ? SegmentIndex : SegmentIndex + 1;
@@ -1206,7 +1222,7 @@ public:
 		static constexpr int32 Degree = 3; // Cubic
 		static constexpr int32 ControlPointsPerSegment = 4;
 		
-		const int32 NumSegments = GetNumSegments();
+		const int32 NumSegments = GetNumberOfSegments();
 		if (Segment < 0 || Segment >= NumSegments)
 		{
 			return;
@@ -1299,7 +1315,7 @@ public:
 		if (bClosed)
 		{
 			// When closing a loop, add a segment to connect the last point back to the first
-			const int32 NumSegments = GetNumSegments();
+			const int32 NumSegments = GetNumberOfSegments();
         
 			if (NumSegments >= 1) // Need at least one segment to close
 			{
@@ -1323,7 +1339,7 @@ public:
 		else
 		{
 			// For an open spline, remove the closing segment
-			const int32 NumSegments = GetNumSegments();
+			const int32 NumSegments = GetNumberOfSegments();
 			if (NumSegments > 0)
 			{
 				// Remove the last segment
@@ -1335,28 +1351,6 @@ public:
 		SetClosedLoopFlag(bClosed);
 	}
 
-	
-	
-
-	/**
-	 * Maps a segment index to its parameter range
-	 * @param SegmentIndex - Index of the segment (0-based)
-	 * @param OutStartParam - Output parameter for segment start value
-	 * @param OutEndParam - Output parameter for segment end value
-	 * @return True if the segment index is valid and mapping succeeded
-	 */
-	FInterval1f GetSegmentParameterRange(int32 SegmentIndex) const
-	{
-		FInterval1f SegmentRange;
-		if (SegmentIndex >= 0 && SegmentIndex < Base::PairKnots.Num() - 1)
-		{
-			SegmentRange.Min = Base::PairKnots[SegmentIndex].Value;
-			SegmentRange.Max = Base::PairKnots[SegmentIndex + 1].Value;
-		}
-        
-		return SegmentRange;
-	}
-
 	/**
 	 * Maps a local segment parameter [0,1] to global parameter space
 	 * @param SegmentIndex - Index of the segment
@@ -1552,7 +1546,7 @@ private:
 	 */
 	void GenerateBezierKnots()
 	{
-		int32 NumSegments = GetNumSegments();
+		int32 NumSegments = GetNumberOfSegments();
 		// For Bezier curves:
 		// - Each segment has degree 3
 		// - For closed loop, we need n + 2*p + 1 knots where:
@@ -1651,7 +1645,7 @@ private:
         {
             return ValueType();
         }
-  		int32 NumSegments = GetNumSegments();
+  		int32 NumSegments = GetNumberOfSegments();
     	int32 SegmentIndex;
     	float LocalT;
     	MapGlobalParameterToLocalSegment(Parameter, SegmentIndex, LocalT);
diff --git a/Engine/Source/Runtime/GeometryCore/Public/Splines/Spline.h b/Engine/Source/Runtime/GeometryCore/Public/Splines/Spline.h
--- a/Engine/Source/Runtime/GeometryCore/Public/Splines/Spline.h
+++ b/Engine/Source/Runtime/GeometryCore/Public/Splines/Spline.h
@@ -29,6 +29,8 @@ public:
     // ISplineInterface Forwarding
     virtual void Clear() override { InternalSpline.Clear(); }
 	virtual bool IsEqual(const ISplineInterface* OtherSpline) const override { return InternalSpline.IsEqual(OtherSpline); }
+	virtual int32 GetNumberOfSegments() const override { return InternalSpline.GetNumberOfSegments(); }
+	virtual FInterval1f GetSegmentParameterRange(int32 SegmentIndex) const override { return InternalSpline.GetSegmentParameterRange(SegmentIndex); }
     virtual bool Serialize(FArchive& Ar) override
     {
     	// Call parent's Serialize
diff --git a/Engine/Source/Runtime/GeometryCore/Public/Splines/SplineInterfaces.h b/Engine/Source/Runtime/GeometryCore/Public/Splines/SplineInterfaces.h
--- a/Engine/Source/Runtime/GeometryCore/Public/Splines/SplineInterfaces.h
+++ b/Engine/Source/Runtime/GeometryCore/Public/Splines/SplineInterfaces.h
@@ -28,8 +28,8 @@ enum class EOutOfBoundsHandlingMode
 };
 
 /**
- * Interface for all spline types
- * Provides the base functionality independent of value type
+ * Interface for all spline types, where spline is defined as a piecewise continuous curve.
+ * Provides the base functionality independent of value type.
  */
 class ISplineInterface
 {
@@ -51,6 +51,9 @@ public:
      */
     virtual FSplineTypeId::IdType GetTypeId() const = 0;
     
+	virtual int32 GetNumberOfSegments() const = 0;
+	virtual FInterval1f GetSegmentParameterRange(int32 SegmentIndex) const = 0;
+	
     virtual bool Serialize(FArchive& Ar)
     {
         // Simple version marker for future compatibility
diff --git a/Engine/Source/Runtime/GeometryCore/Public/Splines/TangentBezierSpline.h b/Engine/Source/Runtime/GeometryCore/Public/Splines/TangentBezierSpline.h
--- a/Engine/Source/Runtime/GeometryCore/Public/Splines/TangentBezierSpline.h
+++ b/Engine/Source/Runtime/GeometryCore/Public/Splines/TangentBezierSpline.h
@@ -427,7 +427,7 @@ public:
 	    {
 	    	// For an open spline, remove the closing segment first
 	    	// then update the internal flag
-	    	const int32 NumSegments = InternalSpline.GetNumSegments();
+	    	const int32 NumSegments = GetNumberOfSegments();
 	    	if (NumSegments > 0)
 	    	{
 	    		// Remove the last segment
@@ -1242,9 +1242,9 @@ public:
         if (Index == 0)
         {
             // Special case for first point in closed loop
-            if (InternalSpline.IsClosedLoop() && InternalSpline.GetNumSegments() > 0)
+            if (InternalSpline.IsClosedLoop() && GetNumberOfSegments() > 0)
             {
-                const int32 LastSegment = InternalSpline.GetNumSegments() - 1;
+                const int32 LastSegment = GetNumberOfSegments() - 1;
                 ValueType P3 = InternalSpline.GetValue(LastSegment * 4 + 3);
                 ValueType P2 = P3 - (NewTangent / 3.0f);
                 UpdateBezierControlPoint(LastSegment * 4 + 2, P2);
@@ -1291,12 +1291,17 @@ public:
     	else if (NumKeys == 1)
     		return 1;  // Special case when we just have the first point
     	else
-    		return IsClosedLoop() ? InternalSpline.GetNumSegments() : InternalSpline.GetNumSegments() + 1;
+    		return IsClosedLoop() ? GetNumberOfSegments() : GetNumberOfSegments() + 1;
     }
 
-	int32 GetNumSegments() const
+	virtual int32 GetNumberOfSegments() const override
 	{
-		return InternalSpline.GetNumSegments();
+		return InternalSpline.GetNumberOfSegments();
+	}
+
+	virtual FInterval1f GetSegmentParameterRange(int32 SegmentIndex) const override
+	{
+		return InternalSpline.GetSegmentParameterRange(SegmentIndex);
 	}
 
     /** Checks if a point's tangents are automatically computed */
@@ -1368,11 +1373,6 @@ public:
 		
 		return InternalSpline.GetParameterSpace();
 	}
-	
-	FInterval1f GetSegmentParameterRange(int32 SegmentIndex) const
-	{
-		return InternalSpline.GetSegmentParameterRange(SegmentIndex);
-	}
 
 	int32 FindSegmentIndex(float Parameter, float &OutLocalParam) const
 	{
