Fixed dynamic mesh CompactInPlace creating a triple-sized TriangleEdges array + add coverage for this case in CheckValidity, and guard against this kind of error breaking AppendWithOffsets

diff --git a/Engine/Source/Runtime/GeometryCore/Private/DynamicMesh/DynamicMesh3.cpp b/Engine/Source/Runtime/GeometryCore/Private/DynamicMesh/DynamicMesh3.cpp
--- a/Engine/Source/Runtime/GeometryCore/Private/DynamicMesh/DynamicMesh3.cpp
+++ b/Engine/Source/Runtime/GeometryCore/Private/DynamicMesh/DynamicMesh3.cpp
@@ -289,6 +289,14 @@ void FDynamicMesh3::AppendWithOffsets(const FDynamicMesh3& ToAppend, FAppendInfo
 		Triangles[Idx].B += UseAppendInfo->VertexOffset;
 		Triangles[Idx].C += UseAppendInfo->VertexOffset;
 	}
+	// TriangleEdges should be 1:1 with Triangles, so this ensure should not fail ...
+	// however due to a now-fixed bug, it is possible that a serialized mesh will 
+	// have too many TriangleEdges; we can recover by resizing to match before appending
+	if (!ensure(TriangleEdges.Num() == UseAppendInfo->TriangleOffset))
+	{
+		// Resize to recover from a too-large TriangleEdges array
+		TriangleEdges.Resize(UseAppendInfo->TriangleOffset);
+	}
 	TriangleEdges.Add(ToAppend.TriangleEdges);
 	for (int32 Idx = UseAppendInfo->TriangleOffset, N = Triangles.Num(); Idx < N; ++Idx)
 	{
@@ -1232,6 +1240,9 @@ bool FDynamicMesh3::CheckValidity(FValidityOptions ValidityOptions, EValidityChe
 	CheckOrFailF(!VertexRefCounts.IsDense() || (VertexRefCounts.IsDense() && VertexRefCounts.GetCount() == Vertices.Num()));
 	CheckOrFailF(!TriangleRefCounts.IsDense() || (TriangleRefCounts.IsDense() && TriangleRefCounts.GetCount() == Triangles.Num()));
 
+	// TriangleEdges should be 1:1 with Triangles
+	CheckOrFailF(TriangleEdges.Num() == Triangles.Num());
+
 	for (int tID : TriangleIndicesItr())
 	{
 		CheckOrFailF(IsTriangle(tID));
diff --git a/Engine/Source/Runtime/GeometryCore/Private/DynamicMesh/DynamicMesh3_Edits.cpp b/Engine/Source/Runtime/GeometryCore/Private/DynamicMesh/DynamicMesh3_Edits.cpp
--- a/Engine/Source/Runtime/GeometryCore/Private/DynamicMesh/DynamicMesh3_Edits.cpp
+++ b/Engine/Source/Runtime/GeometryCore/Private/DynamicMesh/DynamicMesh3_Edits.cpp
@@ -495,7 +495,7 @@ void FDynamicMesh3::CompactInPlace(FCompactMaps* CompactInfo)
 	// trim triangles data structures
 	TriangleRefCounts.Trim(TriangleCount());
 	Triangles.Resize(TriangleCount());
-	TriangleEdges.Resize(TriangleCount() * 3);
+	TriangleEdges.Resize(TriangleCount());
 	if (HasTriangleGroups())
 	{
 		TriangleGroups->Resize(TriangleCount());
