Made the implicit morphology generator not hit an ensure() when passed an empty mesh (and still return an empty result)

diff --git a/Engine/Source/Runtime/GeometryCore/Public/Implicit/Morphology.h b/Engine/Source/Runtime/GeometryCore/Public/Implicit/Morphology.h
--- a/Engine/Source/Runtime/GeometryCore/Public/Implicit/Morphology.h
+++ b/Engine/Source/Runtime/GeometryCore/Public/Implicit/Morphology.h
@@ -108,9 +108,9 @@ public:
 	const FMeshShapeGenerator& Generate()
 	{
 		MarchingCubes.Reset();
-		if (!ensure(Validate()))
+		if ((Source && Source->TriangleCount() == 0) || !ensure(Validate()))
 		{
-			// give up and return and empty result on invalid parameters
+			// return an empty result if input is empty or parameters are not valid
 			return MarchingCubes;
 		}
 
