Reset the UsedCount on a moved-from RefCountVector, so it is left in a valid state (and e.g. a mesh will not fail CheckValidity after being moved from)

diff --git a/Engine/Source/Runtime/GeometryCore/Public/Util/RefCountVector.h b/Engine/Source/Runtime/GeometryCore/Public/Util/RefCountVector.h
--- a/Engine/Source/Runtime/GeometryCore/Public/Util/RefCountVector.h
+++ b/Engine/Source/Runtime/GeometryCore/Public/Util/RefCountVector.h
@@ -28,9 +28,20 @@ public:
 
 	FRefCountVector() = default;
 	FRefCountVector(const FRefCountVector&) = default;
-	FRefCountVector(FRefCountVector&&) = default;
+	FRefCountVector(FRefCountVector&& From) : 
+		RefCounts(MoveTemp(From.RefCounts)), FreeIndices(MoveTemp(From.FreeIndices)), UsedCount(From.UsedCount)
+	{
+		From.UsedCount = 0;
+	}
 	FRefCountVector& operator=(const FRefCountVector&) = default;
-	FRefCountVector& operator=(FRefCountVector&&) = default;
+	FRefCountVector& operator=(FRefCountVector&& From)
+	{
+		RefCounts = MoveTemp(From.RefCounts);
+		FreeIndices = MoveTemp(From.FreeIndices);
+		UsedCount = From.UsedCount;
+		From.UsedCount = 0;
+		return *this;
+	}
 
 	bool IsEmpty() const
 	{
