Replace some usages of FORCEINLINE with inline in Geometry modules.

For context about this change, see this doc: https://docs.google.com/document/d/1LjpT7NpzTecfXdFNyv5KM6sCu67827BRljNi1iaWVSE/edit?tab=t.0#heading=h.fvtz8gjukz7x

Includes modules: GeometryCollectionEngine, GeometryCore

#rnx

diff --git a/Engine/Source/Runtime/Experimental/GeometryCollectionEngine/Public/GeometryCollection/GeometryCollectionComponent.h b/Engine/Source/Runtime/Experimental/GeometryCollectionEngine/Public/GeometryCollection/GeometryCollectionComponent.h
--- a/Engine/Source/Runtime/Experimental/GeometryCollectionEngine/Public/GeometryCollection/GeometryCollectionComponent.h
+++ b/Engine/Source/Runtime/Experimental/GeometryCollectionEngine/Public/GeometryCollection/GeometryCollectionComponent.h
@@ -373,7 +373,7 @@ struct FGeometryCollectionActivatedCluster
 	}
 };
 
-FORCEINLINE FArchive& operator<<(FArchive& Ar, FGeometryCollectionActivatedCluster& ActivatedCluster)
+inline FArchive& operator<<(FArchive& Ar, FGeometryCollectionActivatedCluster& ActivatedCluster)
 {
 	Ar << ActivatedCluster.ActivatedIndex;
 	Ar << ActivatedCluster.InitialLinearVelocity;
@@ -859,11 +859,11 @@ public:
 	UFUNCTION(BlueprintCallable, Category = "ChaosPhysics")
 	GEOMETRYCOLLECTIONENGINE_API FString GetDebugInfo();
 
-	FORCEINLINE const UGeometryCollection* GetRestCollection() const { return RestCollection; }
+	inline const UGeometryCollection* GetRestCollection() const { return RestCollection; }
 
-	FORCEINLINE FGeometryCollectionEdit EditRestCollection(GeometryCollection::EEditUpdate EditUpdate = GeometryCollection::EEditUpdate::RestPhysicsDynamic, bool bShapeIsUnchanged = false) { return FGeometryCollectionEdit(this, EditUpdate, bShapeIsUnchanged); }
+	inline FGeometryCollectionEdit EditRestCollection(GeometryCollection::EEditUpdate EditUpdate = GeometryCollection::EEditUpdate::RestPhysicsDynamic, bool bShapeIsUnchanged = false) { return FGeometryCollectionEdit(this, EditUpdate, bShapeIsUnchanged); }
 #if WITH_EDITOR
-	FORCEINLINE FScopedColorEdit EditBoneSelection(bool bForceUpdate = false) { return FScopedColorEdit(this, bForceUpdate); }
+	inline FScopedColorEdit EditBoneSelection(bool bForceUpdate = false) { return FScopedColorEdit(this, bForceUpdate); }
 
 	/** Propagate bone selection to embedded geometry components. */
 	GEOMETRYCOLLECTIONENGINE_API void SelectEmbeddedGeometry();
@@ -1165,8 +1165,8 @@ public:
 	}
 	
 #if WITH_EDITORONLY_DATA
-	FORCEINLINE const TArray<int32>& GetSelectedBones() const { return SelectedBones; }
-	FORCEINLINE const TArray<int32>& GetHighlightedBones() const { return HighlightedBones; }
+	inline const TArray<int32>& GetSelectedBones() const { return SelectedBones; }
+	inline const TArray<int32>& GetHighlightedBones() const { return HighlightedBones; }
 #endif
 
 	GEOMETRYCOLLECTIONENGINE_API FPhysScene_Chaos* GetInnerChaosScene() const;
diff --git a/Engine/Source/Runtime/Experimental/GeometryCollectionEngine/Public/GeometryCollection/GeometryCollectionISMPoolComponent.h b/Engine/Source/Runtime/Experimental/GeometryCollectionEngine/Public/GeometryCollection/GeometryCollectionISMPoolComponent.h
--- a/Engine/Source/Runtime/Experimental/GeometryCollectionEngine/Public/GeometryCollection/GeometryCollectionISMPoolComponent.h
+++ b/Engine/Source/Runtime/Experimental/GeometryCollectionEngine/Public/GeometryCollection/GeometryCollectionISMPoolComponent.h
@@ -141,7 +141,7 @@ struct FISMComponentDescription
 	}
 };
 
-FORCEINLINE uint32 GetTypeHash(const FISMComponentDescription& Desc)
+inline uint32 GetTypeHash(const FISMComponentDescription& Desc)
 {
 	uint32 Hash = HashCombineFast(GetTypeHash(Desc.Flags), GetTypeHash(Desc.NumCustomDataFloats));
 	Hash = HashCombineFast(Hash, GetTypeHash(Desc.Position));
@@ -202,7 +202,7 @@ struct FGeometryCollectionStaticMeshInstance
 	}
 };
 
-FORCEINLINE uint32 GetTypeHash(const FGeometryCollectionStaticMeshInstance& MeshInstance)
+inline uint32 GetTypeHash(const FGeometryCollectionStaticMeshInstance& MeshInstance)
 {
 	uint32 CombinedHash = GetTypeHash(MeshInstance.StaticMesh);
 	CombinedHash = HashCombineFast(CombinedHash, GetTypeHash(MeshInstance.MaterialsOverrides.Num()));
diff --git a/Engine/Source/Runtime/GeometryCore/Public/Implicit/GridInterpolant.h b/Engine/Source/Runtime/GeometryCore/Public/Implicit/GridInterpolant.h
--- a/Engine/Source/Runtime/GeometryCore/Public/Implicit/GridInterpolant.h
+++ b/Engine/Source/Runtime/GeometryCore/Public/Implicit/GridInterpolant.h
@@ -218,7 +218,7 @@ public:
 	}
 
 private:
-	FORCEINLINE static RealType GetDim(CellSizeType CellSize, int32 Axis)
+	inline static RealType GetDim(CellSizeType CellSize, int32 Axis)
 	{
 		if constexpr (bScalarCellSize)
 		{
diff --git a/Engine/Source/Runtime/GeometryCore/Public/Implicit/SweepingMeshSDF.h b/Engine/Source/Runtime/GeometryCore/Public/Implicit/SweepingMeshSDF.h
--- a/Engine/Source/Runtime/GeometryCore/Public/Implicit/SweepingMeshSDF.h
+++ b/Engine/Source/Runtime/GeometryCore/Public/Implicit/SweepingMeshSDF.h
@@ -1392,7 +1392,7 @@ public:
 	}
 
 	// Helper methods to access CellSize in either scalar or vector mode
-	FORCEINLINE static double GetDim(CellSizeTyped CellSize, int32 Axis)
+	inline static double GetDim(CellSizeTyped CellSize, int32 Axis)
 	{
 		if constexpr (bScalarCellSize)
 		{
@@ -1403,7 +1403,7 @@ public:
 			return (double)CellSize[Axis];
 		}
 	}
-	FORCEINLINE static float GetDim(CellSizeType CellSize, int32 Axis)
+	inline static float GetDim(CellSizeType CellSize, int32 Axis)
 	{
 		if constexpr (bScalarCellSize)
 		{
@@ -1414,7 +1414,7 @@ public:
 			return (float)CellSize[Axis];
 		}
 	}
-	FORCEINLINE static float GetDiagLength(CellSizeType CellSize)
+	inline static float GetDiagLength(CellSizeType CellSize)
 	{
 		if constexpr (bScalarCellSize)
 		{
@@ -1425,7 +1425,7 @@ public:
 			return CellSize.Length();
 		}
 	}
-	FORCEINLINE static bool IsValid(CellSizeType CellSize)
+	inline static bool IsValid(CellSizeType CellSize)
 	{
 		if constexpr (std::is_floating_point_v<CellSizeType>)
 		{
diff --git a/Engine/Source/Runtime/GeometryCore/Public/IndexTypes.h b/Engine/Source/Runtime/GeometryCore/Public/IndexTypes.h
--- a/Engine/Source/Runtime/GeometryCore/Public/IndexTypes.h
+++ b/Engine/Source/Runtime/GeometryCore/Public/IndexTypes.h
@@ -138,7 +138,7 @@ PRAGMA_ENABLE_DEPRECATION_WARNINGS
 	}
 };
 
-FORCEINLINE uint32 GetTypeHash(const FIndex2i& Index)
+inline uint32 GetTypeHash(const FIndex2i& Index)
 {
 	// (this is how FIntVector and all the other FVectors do their hash functions)
 	// Note: this assumes there's no padding that could contain non compared data.
@@ -277,7 +277,7 @@ PRAGMA_ENABLE_DEPRECATION_WARNINGS
 	}
 };
 
-FORCEINLINE uint32 GetTypeHash(const FIndex3i& Index)
+inline uint32 GetTypeHash(const FIndex3i& Index)
 {
 	// (this is how FIntVector and all the other FVectors do their hash functions)
 	// Note: this assumes there's no padding that could contain uncompared data.
@@ -385,7 +385,7 @@ PRAGMA_ENABLE_DEPRECATION_WARNINGS
 	}
 };
 
-FORCEINLINE uint32 GetTypeHash(const FIndex4i& Index)
+inline uint32 GetTypeHash(const FIndex4i& Index)
 {
 	// (this is how FIntVector and all the other FVectors do their hash functions)
 	// Note: this assumes there's no padding that could contain uncompared data.
diff --git a/Engine/Source/Runtime/GeometryCore/Public/IntVectorTypes.h b/Engine/Source/Runtime/GeometryCore/Public/IntVectorTypes.h
--- a/Engine/Source/Runtime/GeometryCore/Public/IntVectorTypes.h
+++ b/Engine/Source/Runtime/GeometryCore/Public/IntVectorTypes.h
@@ -238,7 +238,7 @@ inline std::ostream& operator<<(std::ostream& os, const FVector2i& Vec)
 	return os;
 }
 
-FORCEINLINE uint32 GetTypeHash(const FVector2i& Vector)
+inline uint32 GetTypeHash(const FVector2i& Vector)
 {
 	return FCrc::MemCrc_DEPRECATED(&Vector, sizeof(FVector2i));
 }
@@ -504,7 +504,7 @@ inline std::ostream& operator<<(std::ostream& os, const FVector3i& Vec)
 	return os;
 }
 
-FORCEINLINE uint32 GetTypeHash(const FVector3i& Vector)
+inline uint32 GetTypeHash(const FVector3i& Vector)
 {
 	return FCrc::MemCrc_DEPRECATED(&Vector, sizeof(FVector3i));
 }
diff --git a/Engine/Source/Runtime/GeometryCore/Public/Intersection/IntrRay3AxisAlignedBox3.h b/Engine/Source/Runtime/GeometryCore/Public/Intersection/IntrRay3AxisAlignedBox3.h
--- a/Engine/Source/Runtime/GeometryCore/Public/Intersection/IntrRay3AxisAlignedBox3.h
+++ b/Engine/Source/Runtime/GeometryCore/Public/Intersection/IntrRay3AxisAlignedBox3.h
@@ -149,7 +149,7 @@ protected:
 
 	// internal functions
 
-	static FORCEINLINE bool DoClipping(RealType& t0, RealType& t1,
+	static inline bool DoClipping(RealType& t0, RealType& t1,
 		const TVector<RealType>& RayOrigin, const TVector<RealType>& RayDirection,
 		const TAxisAlignedBox3<RealType>& Box, bool solid, 
 		int& quantity, TVector<RealType>& Point0, TVector<RealType>& Point1, EIntersectionType& intrType)
@@ -191,7 +191,7 @@ protected:
 		return intrType != EIntersectionType::Empty;
 	}
 
-	static FORCEINLINE bool Clip(RealType denom, RealType numer, RealType& t0, RealType& t1)
+	static inline bool Clip(RealType denom, RealType numer, RealType& t0, RealType& t1)
 	{
 		// Return value is 'true' if line segment intersects the current test
 		// plane.  Otherwise 'false' is returned in which case the line segment
diff --git a/Engine/Source/Runtime/GeometryCore/Public/MeshAdapter.h b/Engine/Source/Runtime/GeometryCore/Public/MeshAdapter.h
--- a/Engine/Source/Runtime/GeometryCore/Public/MeshAdapter.h
+++ b/Engine/Source/Runtime/GeometryCore/Public/MeshAdapter.h
@@ -87,54 +87,54 @@ struct TIndexMeshArrayAdapter
 	{
 	}
 
-	FORCEINLINE bool IsTriangle(int32 Index) const
+	inline bool IsTriangle(int32 Index) const
 	{
 		return SourceTriangles->IsValidIndex(Index * 3);
 	}
 
-	FORCEINLINE bool IsVertex(int32 Index) const
+	inline bool IsVertex(int32 Index) const
 	{
 		return SourceVertices->IsValidIndex(Index);
 	}
 
-	FORCEINLINE int32 MaxTriangleID() const
+	inline int32 MaxTriangleID() const
 	{
 		return SourceTriangles->Num() / 3;
 	}
 
-	FORCEINLINE int32 MaxVertexID() const
+	inline int32 MaxVertexID() const
 	{
 		return SourceVertices->Num();
 	}
 
 	// Counts are same as MaxIDs, because these are compact meshes
-	FORCEINLINE int32 TriangleCount() const
+	inline int32 TriangleCount() const
 	{
 		return SourceTriangles->Num() / 3;
 	}
 
-	FORCEINLINE int32 VertexCount() const
+	inline int32 VertexCount() const
 	{
 		return SourceVertices->Num();
 	}
 
-	FORCEINLINE uint64 GetChangeStamp() const
+	inline uint64 GetChangeStamp() const
 	{
 		return 1; // source data doesn't have a timestamp concept
 	}
 
-	FORCEINLINE FIndex3i GetTriangle(int32 Index) const
+	inline FIndex3i GetTriangle(int32 Index) const
 	{
 		int32 Start = Index * 3;
 		return FIndex3i((int)(*SourceTriangles)[Start], (int)(*SourceTriangles)[Start+1], (int)(*SourceTriangles)[Start+2]);
 	}
 
-	FORCEINLINE TVector<OutRealType> GetVertex(int32 Index) const
+	inline TVector<OutRealType> GetVertex(int32 Index) const
 	{
 		return TVector<OutRealType>((*SourceVertices)[Index]);
 	}
 
-	FORCEINLINE void GetTriVertices(int32 TriIndex, UE::Math::TVector<OutRealType>& V0, UE::Math::TVector<OutRealType>& V1, UE::Math::TVector<OutRealType>& V2) const
+	inline void GetTriVertices(int32 TriIndex, UE::Math::TVector<OutRealType>& V0, UE::Math::TVector<OutRealType>& V1, UE::Math::TVector<OutRealType>& V2) const
 	{
 		int32 Start = TriIndex * 3;
 		V0 = TVector<OutRealType>((*SourceVertices)[(*SourceTriangles)[Start]]);
@@ -169,54 +169,54 @@ struct TIndexVectorMeshArrayAdapter
 	{
 	}
 
-	FORCEINLINE bool IsTriangle(int32 Index) const
+	inline bool IsTriangle(int32 Index) const
 	{
 		return SourceTriangles->IsValidIndex(Index);
 	}
 
-	FORCEINLINE bool IsVertex(int32 Index) const
+	inline bool IsVertex(int32 Index) const
 	{
 		return SourceVertices->IsValidIndex(Index);
 	}
 
-	FORCEINLINE int32 MaxTriangleID() const
+	inline int32 MaxTriangleID() const
 	{
 		return SourceTriangles->Num();
 	}
 
-	FORCEINLINE int32 MaxVertexID() const
+	inline int32 MaxVertexID() const
 	{
 		return SourceVertices->Num();
 	}
 
 	// Counts are same as MaxIDs, because these are compact meshes
-	FORCEINLINE int32 TriangleCount() const
+	inline int32 TriangleCount() const
 	{
 		return SourceTriangles->Num();
 	}
 
-	FORCEINLINE int32 VertexCount() const
+	inline int32 VertexCount() const
 	{
 		return SourceVertices->Num();
 	}
 
-	FORCEINLINE uint64 GetChangeStamp() const
+	inline uint64 GetChangeStamp() const
 	{
 		return 1; // source data doesn't have a timestamp concept
 	}
 
-	FORCEINLINE FIndex3i GetTriangle(int32 Index) const
+	inline FIndex3i GetTriangle(int32 Index) const
 	{
 		const IndexVectorType& Tri = (*SourceTriangles)[Index];
 		return FIndex3i((int)Tri[0], (int)Tri[1], (int)Tri[2]);
 	}
 
-	FORCEINLINE TVector<OutRealType> GetVertex(int32 Index) const
+	inline TVector<OutRealType> GetVertex(int32 Index) const
 	{
 		return TVector<OutRealType>((*SourceVertices)[Index]);
 	}
 
-	FORCEINLINE void GetTriVertices(int32 TriIndex, UE::Math::TVector<OutRealType>& V0, UE::Math::TVector<OutRealType>& V1, UE::Math::TVector<OutRealType>& V2) const
+	inline void GetTriVertices(int32 TriIndex, UE::Math::TVector<OutRealType>& V0, UE::Math::TVector<OutRealType>& V1, UE::Math::TVector<OutRealType>& V2) const
 	{
 		const IndexVectorType& Tri = (*SourceTriangles)[TriIndex];
 		V0 = TVector<OutRealType>((*SourceVertices)[Tri[0]]);
diff --git a/Engine/Source/Runtime/GeometryCore/Public/Tessellation/AdaptiveTessellator.h b/Engine/Source/Runtime/GeometryCore/Public/Tessellation/AdaptiveTessellator.h
--- a/Engine/Source/Runtime/GeometryCore/Public/Tessellation/AdaptiveTessellator.h
+++ b/Engine/Source/Runtime/GeometryCore/Public/Tessellation/AdaptiveTessellator.h
@@ -58,7 +58,7 @@ class FAdaptiveTessellator
 {
 private:
 
-	FORCEINLINE static uint32 HashPosition(const FVector3f& Position)
+	inline static uint32 HashPosition(const FVector3f& Position)
 	{
 		union { float f; uint32 i; } x;
 		union { float f; uint32 i; } y;
@@ -75,7 +75,7 @@ private:
 			});
 	}
 
-	FORCEINLINE static uint32 HashPosition(const FVector3d& Position)
+	inline static uint32 HashPosition(const FVector3d& Position)
 	{
 		union { double f; uint64 i; } x;
 		union { double f; uint64 i; } y;
@@ -94,7 +94,7 @@ private:
 
 
 	template <typename T>
-	FORCEINLINE void SubtriangleBarycentrics(uint32 TriX, uint32 TriY, uint32 FlipTri, uint32 NumSubdivisions, UE::Math::TVector<T> Barycentrics[3])
+	inline void SubtriangleBarycentrics(uint32 TriX, uint32 TriY, uint32 FlipTri, uint32 NumSubdivisions, UE::Math::TVector<T> Barycentrics[3])
 	{
 		/*
 			Vert order:
@@ -124,14 +124,14 @@ private:
 	}
 
 	template <typename T>
-	FORCEINLINE static constexpr T EquilateralArea( T EdgeLength )
+	inline static constexpr T EquilateralArea( T EdgeLength )
 	{
 		constexpr T sqrt3_4 = T(0.4330127018922193);
 		return sqrt3_4 * FMath::Square( EdgeLength * EdgeLength );
 	}
 
 	template <typename T>
-	FORCEINLINE static T SubtriangleArea( const UE::Math::TVector<T>& Barycentrics0, const UE::Math::TVector<T>& Barycentrics1, const UE::Math::TVector<T>& Barycentrics2, T TriangleArea )
+	inline static T SubtriangleArea( const UE::Math::TVector<T>& Barycentrics0, const UE::Math::TVector<T>& Barycentrics1, const UE::Math::TVector<T>& Barycentrics2, T TriangleArea )
 	{
 		// Area * Determinant using triple product
 		return TriangleArea * FMath::Abs( Barycentrics0 | ( Barycentrics1 ^ Barycentrics2 ) );
@@ -139,7 +139,7 @@ private:
 
 	template <typename T>
 	// https://math.stackexchange.com/questions/3748903/closest-point-to-triangle-edge-with-barycentric-coordinates
-	FORCEINLINE static constexpr T DistanceToEdge( T Barycentric, T EdgeLength, T TriangleArea )
+	inline static constexpr T DistanceToEdge( T Barycentric, T EdgeLength, T TriangleArea )
 	{
 		return T(2.) * Barycentric * TriangleArea / EdgeLength;
 	}
@@ -147,8 +147,8 @@ private:
 	using FIndex2i = UE::Geometry::FIndex2i;
 	using FIndex3i = UE::Geometry::FIndex3i;
 	
-  	FORCEINLINE static constexpr int32 NextEdge(int32 e) { return (e+1)%3; }
-	FORCEINLINE static constexpr int32 PrevEdge(int32 e) { return (e+2)%3; }
+  	inline static constexpr int32 NextEdge(int32 e) { return (e+1)%3; }
+	inline static constexpr int32 PrevEdge(int32 e) { return (e+2)%3; }
 
   public:
 
@@ -345,7 +345,7 @@ private:
 	}
 
 private:
-	FORCEINLINE bool CouldFlipEdge( int32 TriIndex, int32 TriEdgeIndex ) const 
+	inline bool CouldFlipEdge( int32 TriIndex, int32 TriEdgeIndex ) const 
 	{
 		// Don't flip boundary edges
 		const int32 AdjTriIndex = Mesh.GetAdjTriangle(TriIndex, TriEdgeIndex);
diff --git a/Engine/Source/Runtime/GeometryCore/Public/Tessellation/ParallelAdaptiveRefinement.h b/Engine/Source/Runtime/GeometryCore/Public/Tessellation/ParallelAdaptiveRefinement.h
--- a/Engine/Source/Runtime/GeometryCore/Public/Tessellation/ParallelAdaptiveRefinement.h
+++ b/Engine/Source/Runtime/GeometryCore/Public/Tessellation/ParallelAdaptiveRefinement.h
@@ -25,7 +25,7 @@ private:
 		std::atomic<uint16> GreenAdjMask { 0 }; //< bitmask: encode which of the neighbors [0..3) are green refinement 
 		std::atomic_flag bClaimed = ATOMIC_FLAG_INIT;
 
-		FORCEINLINE void Reset()
+		inline void Reset()
 		{
 			SplitMask = 0;
 			GreenAdjMask.store(0);
diff --git a/Engine/Source/Runtime/GeometryCore/Public/Util/SparseIndexCollectionTypes.h b/Engine/Source/Runtime/GeometryCore/Public/Util/SparseIndexCollectionTypes.h
--- a/Engine/Source/Runtime/GeometryCore/Public/Util/SparseIndexCollectionTypes.h
+++ b/Engine/Source/Runtime/GeometryCore/Public/Util/SparseIndexCollectionTypes.h
@@ -70,7 +70,7 @@ public:
 	 * @param Index index to check for set membership
 	 * @return true if Index is in set
 	 */
-	FORCEINLINE bool Contains(int Index) const
+	inline bool Contains(int Index) const
 	{
 		if (Dense.IsSet())
 		{
@@ -86,7 +86,7 @@ public:
 	/**
 	 * @param Index index to add to set
 	 */
-	FORCEINLINE void Add(int Index)
+	inline void Add(int Index)
 	{
 		if (Dense.IsSet())
 		{
@@ -107,7 +107,7 @@ public:
 	/**
 	 * @param Index index to remove from set
 	 */
-	FORCEINLINE void Remove(int Index)
+	inline void Remove(int Index)
 	{
 		if (Dense.IsSet())
 		{
@@ -128,7 +128,7 @@ public:
 	/**
 	 *  Returns number of values in set
 	 */
-	FORCEINLINE int Count() const
+	inline int Count() const
 	{
 		if (Dense.IsSet())
 		{
