Fix for crash in Extrude Path tool when creating a totally straight polyline

diff --git a/Engine/Source/Runtime/GeometryCore/Private/Generators/PolygonEdgeMeshGenerator.cpp b/Engine/Source/Runtime/GeometryCore/Private/Generators/PolygonEdgeMeshGenerator.cpp
--- a/Engine/Source/Runtime/GeometryCore/Private/Generators/PolygonEdgeMeshGenerator.cpp
+++ b/Engine/Source/Runtime/GeometryCore/Private/Generators/PolygonEdgeMeshGenerator.cpp
@@ -93,7 +93,11 @@ namespace PolygonEdgeMeshGeneratorLocal
 
 		if (FMath::IsNearlyEqual(FMath::Abs(RayDirection0.Dot(RayDirection1)), 1.0))
 		{
-			OutCorner.Add(P1);
+			// TODO: Make a proper triangle fan, not a bunch of vertices on top of each other
+			for (int I = 0; I < NumArcPoints; ++I)
+			{
+				OutCorner.Add(P1);
+			}
 			return;
 		}
 
