Fix PVS warnings

#rnx

diff --git a/Engine/Plugins/FX/Niagara/Source/Niagara/Private/NiagaraFunctionLibrary.cpp b/Engine/Plugins/FX/Niagara/Source/Niagara/Private/NiagaraFunctionLibrary.cpp
--- a/Engine/Plugins/FX/Niagara/Source/Niagara/Private/NiagaraFunctionLibrary.cpp
+++ b/Engine/Plugins/FX/Niagara/Source/Niagara/Private/NiagaraFunctionLibrary.cpp
@@ -1099,7 +1099,7 @@ struct FVectorKernelFastMatrixToQuaternion
 
 					float InvS = FMath::InvSqrt(s);
 
-					float qt[4];
+					float qt[4] = {};
 					qt[i] = 0.5f * (1.f / InvS);
 
 					s = 0.5f * InvS;
diff --git a/Engine/Source/Runtime/Core/Private/HAL/LowLevelMemTracker.cpp b/Engine/Source/Runtime/Core/Private/HAL/LowLevelMemTracker.cpp
--- a/Engine/Source/Runtime/Core/Private/HAL/LowLevelMemTracker.cpp
+++ b/Engine/Source/Runtime/Core/Private/HAL/LowLevelMemTracker.cpp
@@ -1600,7 +1600,7 @@ public:
 	{
 		if (EnumHasAnyFlags(SizeParams, UE::LLM::ESizeParams::RelativeToSnapshot))
 		{
-			return FMath::Clamp<int64>(TrackedTotal - TrackedTotalInSnapshot, 0, INT64_MAX);
+			return FMath::Clamp<int64>(TrackedTotal - TrackedTotalInSnapshot, 0, INT64_MAX); //-V569
 		}
 		return TrackedTotal;
 	}
diff --git a/Engine/Source/Runtime/Core/Private/HAL/LowLevelMemTrackerPrivate.h b/Engine/Source/Runtime/Core/Private/HAL/LowLevelMemTrackerPrivate.h
--- a/Engine/Source/Runtime/Core/Private/HAL/LowLevelMemTrackerPrivate.h
+++ b/Engine/Source/Runtime/Core/Private/HAL/LowLevelMemTrackerPrivate.h
@@ -203,7 +203,7 @@ namespace LLMPrivate
 			// Note, this will also subtract the snapshotted size from PeakSize if that flag is enabled
 			if (EnumHasAnyFlags(SizeParams, UE::LLM::ESizeParams::RelativeToSnapshot))
 			{
-				CurrentSize = FMath::Clamp<int64>(CurrentSize - SizeInSnapshot, 0, INT64_MAX);
+				CurrentSize = FMath::Clamp<int64>(CurrentSize - SizeInSnapshot, 0, INT64_MAX); //-V569
 			}
 
 			return CurrentSize;
diff --git a/Engine/Source/Runtime/GeometryCore/Public/Generators/GridBoxMeshGenerator.h b/Engine/Source/Runtime/GeometryCore/Public/Generators/GridBoxMeshGenerator.h
--- a/Engine/Source/Runtime/GeometryCore/Public/Generators/GridBoxMeshGenerator.h
+++ b/Engine/Source/Runtime/GeometryCore/Public/Generators/GridBoxMeshGenerator.h
@@ -122,7 +122,7 @@ public:
 							int SubDims[2];
 							SubDims[1 - WhichFace] = EdgeVert;
 							SubDims[WhichFace] = NTri[D[WhichFace][FaceDim]] * DSides[1-WhichFace];
-							int FaceV = ToFaceV(FaceDim, SubDims[0], SubDims[1]);
+							int FaceV = ToFaceV(FaceDim, SubDims[0], SubDims[1]); //-V614
 							FaceVertIndices[FaceInds[WhichFace]][FaceV] = CurrentVertIndex;
 						}
 						CurrentVertIndex++;
