Added option for the Mesh Boolean to take a custom 'inside mesh' function, per mesh.

diff --git a/Engine/Source/Runtime/GeometryCore/Private/Operations/MeshBoolean.cpp b/Engine/Source/Runtime/GeometryCore/Private/Operations/MeshBoolean.cpp
--- a/Engine/Source/Runtime/GeometryCore/Private/Operations/MeshBoolean.cpp
+++ b/Engine/Source/Runtime/GeometryCore/Private/Operations/MeshBoolean.cpp
@@ -206,6 +206,16 @@ using namespace UE::Geometry;
 //}
 
 bool FMeshBoolean::Compute()
+{
+	FCustomInsideMeshTest WindingTest(
+		[this](FVector3d Pt, const FDynamicMeshAABBTree3& Spatial, const TFastWindingTree<FDynamicMesh3>* WindingTree) -> bool
+		{
+			return WindingTree->IsInside(Pt, WindingThreshold);
+		}, true);
+	return ComputeWithCustomInside(WindingTest, WindingTest);
+}
+
+bool FMeshBoolean::ComputeWithCustomInside(FCustomInsideMeshTest InsideMeshA, FCustomInsideMeshTest InsideMeshB)
 {
 	// copy meshes
 	FDynamicMesh3 CutMeshB(*Meshes[1]);
@@ -236,7 +246,10 @@ bool FMeshBoolean::Compute()
 	}
 
 	// build spatial data and use it to find intersections
-	FDynamicMeshAABBTree3 Spatial[2]{ CutMesh[0], CutMesh[1] };
+	// Use a copy of the meshes that won't be cut to build the spatial data, so we can re-use it later w/out
+	// needing to recompute it and without numerical issues from the cut
+	FDynamicMesh3 InitialMeshes[2]{ *CutMesh[0], *CutMesh[1] };
+	FDynamicMeshAABBTree3 Spatial[2]{ &InitialMeshes[0], &InitialMeshes[1] };
 	Spatial[0].SetTolerance(SnapTolerance);
 	Spatial[1].SetTolerance(SnapTolerance);
 	MeshIntersection::FIntersectionsQueryResult Intersections
@@ -364,7 +377,10 @@ bool FMeshBoolean::Compute()
 		}
 		for (int MeshIdx = 0; MeshIdx < NumMeshesToProcess; MeshIdx++)
 		{
-			TFastWindingTree<FDynamicMesh3> Winding(&Spatial[1 - MeshIdx]);
+			FCustomInsideMeshTest& UseInsideTest = MeshIdx == 0 ? InsideMeshB : InsideMeshA;
+			TFastWindingTree<FDynamicMesh3> LocalWinding(&Spatial[1 - MeshIdx], UseInsideTest.bRequiresWinding);
+			TFastWindingTree<FDynamicMesh3>* UseWinding = UseInsideTest.bRequiresWinding ? &LocalWinding : nullptr;
+			
 			FDynamicMeshAABBTree3& OtherSpatial = Spatial[1 - MeshIdx];
 			FDynamicMesh3& ProcessMesh = *CutMesh[MeshIdx];
 			int MaxTriID = ProcessMesh.MaxTriangleID();
@@ -450,8 +466,7 @@ bool FMeshBoolean::Compute()
 					}
 
 					// didn't already return a coplanar result; use the winding number
-					double WindingNum = Winding.FastWindingNumber(Centroid);
-					KeepTri[MeshIdx][TID] = (WindingNum > WindingThreshold) != bRemoveInside;
+					KeepTri[MeshIdx][TID] = (UseInsideTest.IsPointInsideFn(Centroid, OtherSpatial, UseWinding)) != bRemoveInside;
 				});
 		}
 
diff --git a/Engine/Source/Runtime/GeometryCore/Public/Operations/MeshBoolean.h b/Engine/Source/Runtime/GeometryCore/Public/Operations/MeshBoolean.h
--- a/Engine/Source/Runtime/GeometryCore/Public/Operations/MeshBoolean.h
+++ b/Engine/Source/Runtime/GeometryCore/Public/Operations/MeshBoolean.h
@@ -175,11 +175,36 @@ public:
 	}
 
 	/**
-	 * Compute the Boolean operation
+	 * Compute the Boolean operation, using winding number for inside/outside tests
 	 * @return true if operation succeeds
 	 */
 	GEOMETRYCORE_API bool Compute();
 
+	// Function type for custom inside/outside query functions
+	// Note: If implementing a custom function that needs to access the underlying mesh, it is accessible via Spatial.GetMesh()
+	using TestInsideFunctionType = TFunctionRef<bool(const FVector3d& TestPosition, const FDynamicMeshAABBTree3& Spatial, const TFastWindingTree<FDynamicMesh3>* OptionalWindingTree)>;
+	struct FCustomInsideMeshTest
+	{
+		TestInsideFunctionType IsPointInsideFn;
+		bool bRequiresWinding = true;
+
+		/**
+		 * Construct a custom class for testing whether a point is inside a mesh.
+		 * Note: This will be called in the local space of the boolean, which is not the initial / world-space of the input meshes
+		 * 
+		 * @param IsPointInsideFn Custom function reference to test if a given point is inside a mesh.
+		 * @param bRequiresWinding Whether the winding tree needs to be computed for the custom function. If false, the OptionalWindingTree passed to IsPointInsideFn will be null.
+		 */
+		FCustomInsideMeshTest(TestInsideFunctionType IsPointInsideFn, bool bRequiresWinding = true)
+			: IsPointInsideFn(IsPointInsideFn), bRequiresWinding(bRequiresWinding)
+		{}
+	};
+
+	/**
+	 * Compute the Boolean operation with a custom function determining if points are inside the given mesh
+	 */
+	GEOMETRYCORE_API bool ComputeWithCustomInside(FCustomInsideMeshTest TestMeshA, FCustomInsideMeshTest TestMeshB);
+
 protected:
 	/** If this returns true, abort computation.  */
 	virtual bool Cancelled()
