speculative fixed mesh boolean custom inside test

diff --git a/Engine/Source/Runtime/GeometryCore/Private/Operations/MeshBoolean.cpp b/Engine/Source/Runtime/GeometryCore/Private/Operations/MeshBoolean.cpp
--- a/Engine/Source/Runtime/GeometryCore/Private/Operations/MeshBoolean.cpp
+++ b/Engine/Source/Runtime/GeometryCore/Private/Operations/MeshBoolean.cpp
@@ -206,6 +206,16 @@ using namespace UE::Geometry;
 //}
 
 bool FMeshBoolean::Compute()
+{
+	FCustomInsideMeshTest WindingTest(
+		[Threshold = WindingThreshold](FVector3d Pt, const FDynamicMeshAABBTree3& Spatial, const TFastWindingTree<FDynamicMesh3>* WindingTree) -> bool
+		{
+			return WindingTree->IsInside(Pt, Threshold);
+		}, true);
+	return ComputeWithCustomInside(WindingTest, WindingTest);
+}
+
+bool FMeshBoolean::ComputeWithCustomInside(FCustomInsideMeshTest InsideMeshA, FCustomInsideMeshTest InsideMeshB)
 {
 	// copy meshes
 	FDynamicMesh3 CutMeshB(*Meshes[1]);
@@ -364,8 +374,15 @@ bool FMeshBoolean::Compute()
 		}
 		for (int MeshIdx = 0; MeshIdx < NumMeshesToProcess; MeshIdx++)
 		{
-			TFastWindingTree<FDynamicMesh3> Winding(&Spatial[1 - MeshIdx]);
+			FCustomInsideMeshTest& UseInsideTest = MeshIdx == 0 ? InsideMeshB : InsideMeshA;
 			FDynamicMeshAABBTree3& OtherSpatial = Spatial[1 - MeshIdx];
+			if (!OtherSpatial.IsValid(false))
+			{
+				OtherSpatial.Build();
+			}
+			TFastWindingTree<FDynamicMesh3> LocalWinding(&OtherSpatial, UseInsideTest.bRequiresWinding);
+			TFastWindingTree<FDynamicMesh3>* UseWinding = UseInsideTest.bRequiresWinding ? &LocalWinding : nullptr;
+			
 			FDynamicMesh3& ProcessMesh = *CutMesh[MeshIdx];
 			int MaxTriID = ProcessMesh.MaxTriangleID();
 			KeepTri[MeshIdx].SetNumUninitialized(MaxTriID);
@@ -450,8 +467,7 @@ bool FMeshBoolean::Compute()
 					}
 
 					// didn't already return a coplanar result; use the winding number
-					double WindingNum = Winding.FastWindingNumber(Centroid);
-					KeepTri[MeshIdx][TID] = (WindingNum > WindingThreshold) != bRemoveInside;
+					KeepTri[MeshIdx][TID] = (UseInsideTest.IsPointInsideFn(Centroid, OtherSpatial, UseWinding)) != bRemoveInside;
 				});
 		}
 
diff --git a/Engine/Source/Runtime/GeometryCore/Public/Operations/MeshBoolean.h b/Engine/Source/Runtime/GeometryCore/Public/Operations/MeshBoolean.h
--- a/Engine/Source/Runtime/GeometryCore/Public/Operations/MeshBoolean.h
+++ b/Engine/Source/Runtime/GeometryCore/Public/Operations/MeshBoolean.h
@@ -175,11 +175,36 @@ public:
 	}
 
 	/**
-	 * Compute the Boolean operation
+	 * Compute the Boolean operation, using winding number for inside/outside tests
 	 * @return true if operation succeeds
 	 */
 	GEOMETRYCORE_API bool Compute();
 
+	// Function type for custom inside/outside query functions
+	// Note: If implementing a custom function that needs to access the underlying mesh, it is accessible via Spatial.GetMesh()
+	using TestInsideFunctionType = TFunctionRef<bool(const FVector3d& TestPosition, const FDynamicMeshAABBTree3& Spatial, const TFastWindingTree<FDynamicMesh3>* OptionalWindingTree)>;
+	struct FCustomInsideMeshTest
+	{
+		TestInsideFunctionType IsPointInsideFn;
+		bool bRequiresWinding = true;
+
+		/**
+		 * Construct a custom class for testing whether a point is inside a mesh.
+		 * Note: This will be called in the local space of the boolean, which is not the initial / world-space of the input meshes
+		 * 
+		 * @param IsPointInsideFn Custom function reference to test if a given point is inside a mesh.
+		 * @param bRequiresWinding Whether the winding tree needs to be computed for the custom function. If false, the OptionalWindingTree passed to IsPointInsideFn will be null.
+		 */
+		FCustomInsideMeshTest(TestInsideFunctionType IsPointInsideFn, bool bRequiresWinding = true)
+			: IsPointInsideFn(IsPointInsideFn), bRequiresWinding(bRequiresWinding)
+		{}
+	};
+
+	/**
+	 * Compute the Boolean operation with a custom function determining if points are inside the given mesh
+	 */
+	GEOMETRYCORE_API bool ComputeWithCustomInside(FCustomInsideMeshTest TestMeshA, FCustomInsideMeshTest TestMeshB);
+
 protected:
 	/** If this returns true, abort computation.  */
 	virtual bool Cancelled()
