Fix Knot vector warnings when adding first points in an empty spline

diff --git a/Engine/Source/Runtime/GeometryCore/Public/Splines/PolyBezierSpline.h b/Engine/Source/Runtime/GeometryCore/Public/Splines/PolyBezierSpline.h
--- a/Engine/Source/Runtime/GeometryCore/Public/Splines/PolyBezierSpline.h
+++ b/Engine/Source/Runtime/GeometryCore/Public/Splines/PolyBezierSpline.h
@@ -1445,12 +1445,7 @@ public:
 		// Special handling for incomplete first segment
 		if (NumSegments == 0)
 		{
-			if (NumValues == 0)
-				return 0;  // No points at all
-            
-			// For 1-3 points (building toward first segment), return a progressive number
-			// This ensures we maintain sufficient knots during initialization
-			return NumValues * 2;  // Simple progression: 1→2, 2→4, 3→6 knots
+			return 0;
 		}
 		return 4 + (NumSegments - 1) * 3 + 4;
 	}
