Splines - Fix GetTangentIn/Out for closed splines

diff --git a/Engine/Source/Runtime/GeometryCore/Public/Splines/TangentBezierSpline.h b/Engine/Source/Runtime/GeometryCore/Public/Splines/TangentBezierSpline.h
--- a/Engine/Source/Runtime/GeometryCore/Public/Splines/TangentBezierSpline.h
+++ b/Engine/Source/Runtime/GeometryCore/Public/Splines/TangentBezierSpline.h
@@ -1174,10 +1174,10 @@ public:
 		
         if (Index == 0)
         {
-            // Special case for first point
             if (IsClosedLoop())
             {
-                return GetTangentIn(NumPoints - 1);
+            	// The point at NumPoints when the spline is closed is the point that is coincident with point 0.
+                return GetTangentIn(NumPoints);
             }
             else
             {
@@ -1207,17 +1207,10 @@ public:
 			return ValueType();
 		}
 		
-		if (Index == NumPoints - 1)
+		if (Index == NumPoints - 1 && !IsClosedLoop())
 		{
-			if (IsClosedLoop())
-			{
-				return GetTangentOut(0);
-			}
-			else
-			{
-				// If spline is unclosed, out tangent == in tangent at final point.
-				return GetTangentIn(Index);
-			}
+			// If spline is unclosed, out tangent == in tangent at final point.
+			return GetTangentIn(Index);
 		}
 		
 		// Normal case - get from next segment's control points
