Segment Tree:
- add option to build the tree with an array of segments
- make use of the filter function that's passed in as an option

diff --git a/Engine/Source/Runtime/GeometryCore/Public/Spatial/SegmentTree3.h b/Engine/Source/Runtime/GeometryCore/Public/Spatial/SegmentTree3.h
--- a/Engine/Source/Runtime/GeometryCore/Public/Spatial/SegmentTree3.h
+++ b/Engine/Source/Runtime/GeometryCore/Public/Spatial/SegmentTree3.h
@@ -79,7 +79,29 @@ public:
 		BuildTopDown(SegmentIndices, Centers, SegmentList.Num());
 	}
 
+	void Build(const TArray<FSegment3d>& Segments)
+	{
+		SegmentList.SetNumUninitialized(Segments.Num());
+		for (int32 SegmentID = 0; SegmentID < Segments.Num(); ++SegmentID)
+		{
+			SegmentList[SegmentID].ID = SegmentID;
+			SegmentList[SegmentID].Segment = Segments[SegmentID];
+		}
 
+		// Make explicit list of linear indices and centers for the build step, which
+		// will sort/swap these lists in-place. 
+		// Centers could possibly be removed as centers can be fetched from the SegmentList...
+		TArray<int> SegmentIndices;
+		SegmentIndices.Reserve(SegmentList.Num());
+		TArray<FVector3d> Centers;
+		Centers.Reserve(SegmentList.Num());
+		for (int32 k = 0; k < SegmentList.Num(); ++k)
+		{
+			SegmentIndices.Add(k);
+			Centers.Add(SegmentList[k].Segment.Center);
+		}
+		BuildTopDown(SegmentIndices, Centers, SegmentList.Num());
+	}
 
 	/**
 	 * Find the nearest segment to query point P, and return it in NearestSegmentOut
@@ -439,6 +461,12 @@ protected:
 			for (int i = 1; i <= NumSegments; ++i)
 			{
 				int SegmentIdx = IndexList[TreeBox.BoxToIndex + i];
+				
+				if (Options.TriangleFilterF != nullptr && Options.TriangleFilterF(SegmentIdx) == false)
+				{
+					continue;
+				}
+
 				double SegmentDistSqr = SegmentList[SegmentIdx].Segment.DistanceSquared(P);
 				if (SegmentDistSqr < NearestDistSqr)
 				{
