Fix append/prepend to closed splines

diff --git a/Engine/Source/Runtime/Engine/Private/Curves/Spline.cpp b/Engine/Source/Runtime/Engine/Private/Curves/Spline.cpp
--- a/Engine/Source/Runtime/Engine/Private/Curves/Spline.cpp
+++ b/Engine/Source/Runtime/Engine/Private/Curves/Spline.cpp
@@ -1880,21 +1880,11 @@ void FNewSpline::ResetScale()
 
 void FNewSpline::AddPoint(const FSplinePoint& Point)
 {
-	// todo: this is a hack to support appending points when the spline is closed. remove when appending to a closed spline is fixed.
-	bool bWasClosed = IsClosedLoop();
-	SetClosedLoop(false);
-	
 	ValidateRotScale();
 	ON_SCOPE_EXIT
 	{
 		SetAttributeChannelRange(ScaleAttrName, FInterval1f(0.f, 1.f), EMappingRangeSpace::Normalized);
 		SetAttributeChannelRange(RotationAttrName, FInterval1f(0.f, 1.f), EMappingRangeSpace::Normalized);
-
-		if (bWasClosed)
-		{
-			SetClosedLoop(true);
-		}
-		
 		ValidateRotScale();
 	};
 	
@@ -1955,21 +1945,11 @@ void FNewSpline::AddPoint(const FSplinePoint& Point)
 
 void FNewSpline::InsertPoint(const FSplinePoint& Point, int32 Index)
 {
-	// todo: this is a hack to support appending points when the spline is closed. remove when appending to a closed spline is fixed.
-	bool bWasClosed = IsClosedLoop();
-	SetClosedLoop(false);
-	
 	ValidateRotScale();
 	ON_SCOPE_EXIT
 	{
 		SetAttributeChannelRange(ScaleAttrName, FInterval1f(0.f, 1.f), EMappingRangeSpace::Normalized);
 		SetAttributeChannelRange(RotationAttrName, FInterval1f(0.f, 1.f), EMappingRangeSpace::Normalized);
-
-		if (bWasClosed)
-		{
-			SetClosedLoop(true);
-		}
-		
 		ValidateRotScale();
 	};
 	
diff --git a/Engine/Source/Runtime/GeometryCore/Public/Splines/PolyBezierSpline.h b/Engine/Source/Runtime/GeometryCore/Public/Splines/PolyBezierSpline.h
--- a/Engine/Source/Runtime/GeometryCore/Public/Splines/PolyBezierSpline.h
+++ b/Engine/Source/Runtime/GeometryCore/Public/Splines/PolyBezierSpline.h
@@ -3,6 +3,8 @@
 #pragma once
 
 #include "BSpline.h"
+#include "Misc/ScopeExit.h"
+
 namespace UE
 {
 namespace Geometry
@@ -641,6 +643,16 @@ public:
         const ValueType& P3, 
         EParameterizationPolicy ParameterizationPolicy = EParameterizationPolicy::Centripetal)
     {
+		bool bWasClosed = Base::IsClosedLoop();
+		SetClosedLoop(false);
+		ON_SCOPE_EXIT
+		{
+			if (bWasClosed)
+			{
+				SetClosedLoop(true);
+			}
+		};
+		
 		// Get P0 from the last point of the previous segment
 		ValueType P0 = Base::GetValue(Base::NumKeys() - 1);
 		int32 NumSegments = Base::NumKeys() / 4;
@@ -665,9 +677,18 @@ public:
 		const ValueType& P2, 
 		EParameterizationPolicy ParameterizationPolicy = EParameterizationPolicy::Centripetal)
 	{
+		bool bWasClosed = Base::IsClosedLoop();
+		SetClosedLoop(false);
+		ON_SCOPE_EXIT
+		{
+			if (bWasClosed)
+			{
+				SetClosedLoop(true);
+			}
+		};
+		
 		// Get P3 from the first point of the existing spline
 		ValueType P3 = Base::GetValue(0);
-   
 		// Add segment with prepend behavior
 		return AddBezierSegmentInternal(0, P0, P1, P2, P3,
 										ParameterizationPolicy,
@@ -1306,13 +1327,13 @@ public:
 		Base::bIsClosedLoop = bClosed;
 	}
 	
-	virtual void SetClosedLoop(bool bClosed) override
+	virtual void SetClosedLoop(bool bShouldClose) override
 	{
 		// Skip if state isn't changing
-		if (bClosed == this->IsClosedLoop())
+		if (bShouldClose == this->IsClosedLoop())
 			return;
         
-		if (bClosed)
+		if (bShouldClose)
 		{
 			// When closing a loop, add a segment to connect the last point back to the first
 			const int32 NumSegments = GetNumberOfSegments();
@@ -1348,7 +1369,7 @@ public:
 		}
     
 		// Set the flag
-		SetClosedLoopFlag(bClosed);
+		SetClosedLoopFlag(bShouldClose);
 	}
 
 	/**
@@ -1760,6 +1781,16 @@ private:
 		const ValueType& Position,
 		EParameterizationPolicy ParameterizationPolicy = EParameterizationPolicy::Centripetal)
 	{
+		bool bWasClosed = Base::IsClosedLoop();
+		SetClosedLoop(false);
+		ON_SCOPE_EXIT
+		{
+			if (bWasClosed)
+			{
+				SetClosedLoop(true);
+			}
+		};
+		
 		// Get current parameter range start
 		FInterval1f ParameterRange = GetParameterSpace();
     
@@ -1789,6 +1820,16 @@ private:
 		const ValueType& Position,
 		EParameterizationPolicy ParameterizationPolicy)
 	{
+		bool bWasClosed = Base::IsClosedLoop();
+    	SetClosedLoop(false);
+    	ON_SCOPE_EXIT
+    	{
+    		if (bWasClosed)
+    		{
+    			SetClosedLoop(true);
+    		}
+    	};
+    		
 		// Append - create a point at the end
 		int32 LastIndex = Base::NumKeys() - 1;
 		float StartParam = Base::GetParameter(LastIndex);
@@ -1802,6 +1843,7 @@ private:
 		int32 NumSegments = Base::NumKeys() / 4;
 		// Add the segment
 		int32 NewIndex = AddBezierSegmentInternal(NumSegments, P0, P1, P2, Position, ParameterizationPolicy/*, StartParam*/);
+
 		return NewIndex;
 	}
 	
diff --git a/Engine/Source/Runtime/GeometryCore/Public/Splines/TangentBezierSpline.h b/Engine/Source/Runtime/GeometryCore/Public/Splines/TangentBezierSpline.h
--- a/Engine/Source/Runtime/GeometryCore/Public/Splines/TangentBezierSpline.h
+++ b/Engine/Source/Runtime/GeometryCore/Public/Splines/TangentBezierSpline.h
@@ -842,6 +842,16 @@ public:
 	 */
 	void PrependPoint(const FTangentBezierControlPoint& ControlPoint)
 	{
+		bool bWasClosed = IsClosedLoop();
+		SetClosedLoop(false);
+		ON_SCOPE_EXIT
+		{
+			if (bWasClosed)
+			{
+				SetClosedLoop(true);
+			}
+		};
+		
     	int32 OutNewIndex;
 
     	// Special case: handle first two points
@@ -889,6 +899,16 @@ public:
      */
     void AppendPoint(const FTangentBezierControlPoint& ControlPoint)
     {
+		bool bWasClosed = IsClosedLoop();
+		SetClosedLoop(false);
+		ON_SCOPE_EXIT
+		{
+			if (bWasClosed)
+			{
+				SetClosedLoop(true);
+			}
+		};
+		
     	int32 NewIndex = GetNumPoints();
 
     	// Special case: handle first two points
@@ -896,7 +916,7 @@ public:
     	{
     		return;
     	}
-    	
+		
     	
 	    // Normal case - add segment to existing spline
         ValueType P0 = GetValue(NewIndex - 1);
