Added option for geometrycore's MeshBoolean to report the source mesh of each result triangle

diff --git a/Engine/Source/Runtime/GeometryCore/Private/Operations/MeshBoolean.cpp b/Engine/Source/Runtime/GeometryCore/Private/Operations/MeshBoolean.cpp
--- a/Engine/Source/Runtime/GeometryCore/Private/Operations/MeshBoolean.cpp
+++ b/Engine/Source/Runtime/GeometryCore/Private/Operations/MeshBoolean.cpp
@@ -727,6 +727,15 @@ bool FMeshBoolean::ComputeWithCustomInside(FCustomInsideMeshTest InsideMeshA, FC
 		FMeshIndexMappings IndexMaps;
 		Editor.AppendMesh(CutMesh[1], IndexMaps);
 
+		if (TrackPerTriangleSourceMesh.IsSet())
+		{
+			TrackPerTriangleSourceMesh->SetNumZeroed(Result->MaxTriangleID());
+			for (const TPair<int32, int32> TriSrcDst : IndexMaps.GetTriangleMap().GetForwardMap())
+			{
+				(*TrackPerTriangleSourceMesh)[TriSrcDst.Value] = 1;
+			}
+		}
+
 		if (bPopulateSecondMeshGroupMap)
 		{
 			SecondMeshGroupMap = IndexMaps.GetGroupMap();
diff --git a/Engine/Source/Runtime/GeometryCore/Public/Operations/MeshBoolean.h b/Engine/Source/Runtime/GeometryCore/Public/Operations/MeshBoolean.h
--- a/Engine/Source/Runtime/GeometryCore/Public/Operations/MeshBoolean.h
+++ b/Engine/Source/Runtime/GeometryCore/Public/Operations/MeshBoolean.h
@@ -138,6 +138,11 @@ public:
 	bool bPopulateSecondMeshGroupMap = false;
 	FIndexMapi SecondMeshGroupMap;
 
+	// If set, will be populated with an array of source mesh indices per TriangleID in the result mesh,
+	// indicating if that triangle was from the Meshes[0] or Meshes[1] source.
+	// Note: For non-compact meshes, invalid triangle IDs will have arbitrary indices.
+	TOptional<TArray<int8>> TrackPerTriangleSourceMesh;
+
 public:
 
 	/**
