FEdgeSpan: add a function to get subspans from a given span, broken up by a given angle threhold

diff --git a/Engine/Source/Runtime/GeometryCore/Private/EdgeSpan.cpp b/Engine/Source/Runtime/GeometryCore/Private/EdgeSpan.cpp
--- a/Engine/Source/Runtime/GeometryCore/Private/EdgeSpan.cpp
+++ b/Engine/Source/Runtime/GeometryCore/Private/EdgeSpan.cpp
@@ -278,4 +278,142 @@ void FEdgeSpan::VertexSpanToEdgeSpan(const FDynamicMesh3* Mesh, const TArray<int
 		int v1 = VertexSpan[i + 1];
 		OutEdgeSpan[i] = Mesh->FindEdge(v0, v1);
 	}
-}
\ No newline at end of file
+}
+
+
+void FEdgeSpan::GetSubspansByAngle(double AngleThresholdDeg, int32 MinSpanSize, TArray<FEdgeSpan>& OutSpans) const
+{
+	auto IsEdgeAngleSharp = [this](int32 SharedVid, const FIndex2i& AttachedEids, double DotThreshold) -> bool
+	{
+		if (!ensure(Mesh && Mesh->IsEdge(AttachedEids.A) && Mesh->IsEdge(AttachedEids.B)))
+		{
+			return false;
+		}
+
+		// Gets vector pointing from the Vid along the edge.
+		auto GetEdgeUnitVector = [this, SharedVid](int32 Eid, FVector3d& VectorOut)->bool
+		{
+			FIndex2i EdgeVids = Mesh->GetEdgeV(Eid);
+			// Make sure that the Vid is at EdgeVids.A
+			if (EdgeVids.B == SharedVid)
+			{
+				Swap(EdgeVids.A, EdgeVids.B);
+			}
+			VectorOut = Mesh->GetVertex(EdgeVids.B) - Mesh->GetVertex(EdgeVids.A);
+			return VectorOut.Normalize(KINDA_SMALL_NUMBER);
+		};
+
+		FVector Edge1, Edge2;
+		if (!GetEdgeUnitVector(AttachedEids.A, Edge1) || !GetEdgeUnitVector(AttachedEids.B, Edge2))
+		{
+			// If either edge was degenerate, we don't want to consider the angle "sharp"
+			return false;
+		}
+
+		return Edge1.Dot(Edge2) >= DotThreshold;
+	};
+
+	const double CosAngleThresholdDeg = FMath::Cos(FMath::DegreesToRadians(AngleThresholdDeg));
+
+	// First just identify corners
+
+	TArray<int32> SpanCorners;		// These are indices of Vertices, not vertex indices
+	for (int32 SpanVertexIndex = 1; SpanVertexIndex < Vertices.Num() - 1; ++SpanVertexIndex)
+	{
+		const int32 PreviousEdgeID = Edges[SpanVertexIndex - 1];
+		const int32 VertexID = Vertices[SpanVertexIndex];
+		check(SpanVertexIndex < Edges.Num());
+		const int32 NextEdgeID = Edges[SpanVertexIndex];
+
+		const bool bIsSharp = IsEdgeAngleSharp(VertexID, FIndex2i{ PreviousEdgeID, NextEdgeID }, CosAngleThresholdDeg);
+
+		if (bIsSharp)
+		{
+			SpanCorners.Add(SpanVertexIndex);
+		}
+	}
+
+	// Filter corners too close together
+
+	if (SpanCorners.Num() > 1)
+	{
+		TArray<int32> FilteredSpanCorners;
+		FilteredSpanCorners.Add(SpanCorners[0]);
+
+		for (int32 CornerIndex = 1; CornerIndex < SpanCorners.Num(); ++CornerIndex)
+		{
+			int32 Distance = SpanCorners[CornerIndex] - FilteredSpanCorners.Last();
+			Distance = FMath::Min(Distance, Vertices.Num() - Distance);	 // check wraparound distance as well	
+
+			if (Distance >= MinSpanSize)
+			{
+				FilteredSpanCorners.Add(SpanCorners[CornerIndex]);
+			}
+		}
+
+		SpanCorners = MoveTemp(FilteredSpanCorners);
+	}
+
+	// Now assemble output spans
+
+	if (SpanCorners.Num() == 0)
+	{
+		OutSpans.Add(*this);
+		return;
+	}
+
+	int32 NewSpanStart = 0;
+	for (int32 NewSpanEnd : SpanCorners)
+	{
+		TArray<int32> NewSpanVertices;
+		TArray<int32> NewSpanEdges;
+
+		for (int32 NewSpanVertexIndex = NewSpanStart; NewSpanVertexIndex < NewSpanEnd; ++NewSpanVertexIndex)
+		{
+			NewSpanVertices.Add(Vertices[NewSpanVertexIndex]);
+			ensure(NewSpanVertexIndex < Edges.Num());
+			NewSpanEdges.Add(Edges[NewSpanVertexIndex]);
+		}
+		NewSpanVertices.Add(Vertices[NewSpanEnd]);		// final vertex in this span
+
+		if (ensure(NewSpanVertices.Num() > 0))
+		{
+			OutSpans.Add(FEdgeSpan(Mesh, NewSpanVertices, NewSpanEdges));
+		}
+		NewSpanStart = NewSpanEnd;
+	}
+
+	// Final span: prepend it to the first one if they are connected and there is no corner between them
+
+	TArray<int32> FinalSpanVertices;
+	TArray<int32> FinalSpanEdges;
+	for (int32 FinalSpanVertexIndex = NewSpanStart; FinalSpanVertexIndex < Vertices.Num(); ++FinalSpanVertexIndex)
+	{
+		FinalSpanVertices.Add(Vertices[FinalSpanVertexIndex]);
+		if (FinalSpanVertexIndex < Edges.Num())
+		{
+			FinalSpanEdges.Add(Edges[FinalSpanVertexIndex]);
+		}
+	}
+
+	const bool bSpanIsLoop = (IndexUtil::FindSharedEdgeVertex(Mesh->GetEdgeV(Edges[0]), Mesh->GetEdgeV(Edges.Last())) != IndexConstants::InvalidID);
+	const bool bSharpCorner = IsEdgeAngleSharp(Vertices[0], FIndex2i{ Edges[0], Edges.Last() }, CosAngleThresholdDeg);
+	const bool bShouldPrepend = bSpanIsLoop && !bSharpCorner && (OutSpans.Num() > 0);
+
+	if (bShouldPrepend)
+	{
+		if (ensure(FinalSpanVertices.Last() == OutSpans[0].Vertices[0]))
+		{
+			FinalSpanVertices.Pop();
+		}
+
+		OutSpans[0].Vertices.Insert(FinalSpanVertices, 0);
+		OutSpans[0].Edges.Insert(FinalSpanEdges, 0);
+	}
+	else
+	{
+		// Don't add to the first span, create a new final output span
+		OutSpans.Add(FEdgeSpan(Mesh, FinalSpanVertices, FinalSpanEdges));
+	}
+}
+
diff --git a/Engine/Source/Runtime/GeometryCore/Public/EdgeSpan.h b/Engine/Source/Runtime/GeometryCore/Public/EdgeSpan.h
--- a/Engine/Source/Runtime/GeometryCore/Public/EdgeSpan.h
+++ b/Engine/Source/Runtime/GeometryCore/Public/EdgeSpan.h
@@ -208,6 +208,15 @@ public:
 	 */
 	static GEOMETRYCORE_API void VertexSpanToEdgeSpan(const FDynamicMesh3* Mesh, const TArray<int>& VertexSpan, TArray<int>& OutEdgeSpan);
 
+	/**
+	 * Split a single EdgeSpan into multiple subspans based on an edge angle threshold. Subspans are copied out, the original span is not modified.
+	 * @param AngleThresholdDeg max angle between adjacant edges to mark a vertex as a "corner" (in degrees). Angle is measured between edge vectors extending from their
+	 *							common vertex, so a value of 180 means all vertices would be labelled as corners.
+	 * @param MinSpanSize minimum number of edges necessary for a subspan, regardless of angles between edges. Subspans that would be too small are merged to adjacent subspans.
+	 * @param OutSpan output subspans
+	 */
+	GEOMETRYCORE_API void GetSubspansByAngle(double AngleThresholdDeg, int32 MinSpanSize, TArray<FEdgeSpan>& OutSpans) const;
+
 };
 
 
