Change from TFunctionRef to TFunction to fix lifetime issues for the mesh boolean custom inside test method

diff --git a/Engine/Source/Runtime/GeometryCore/Public/Operations/MeshBoolean.h b/Engine/Source/Runtime/GeometryCore/Public/Operations/MeshBoolean.h
--- a/Engine/Source/Runtime/GeometryCore/Public/Operations/MeshBoolean.h
+++ b/Engine/Source/Runtime/GeometryCore/Public/Operations/MeshBoolean.h
@@ -187,7 +187,7 @@ public:
 
 	// Function type for custom inside/outside query functions
 	// Note: If implementing a custom function that needs to access the underlying mesh, it is accessible via Spatial.GetMesh()
-	using TestInsideFunctionType = TFunctionRef<bool(const FVector3d& TestPosition, const FDynamicMeshAABBTree3& Spatial, const TFastWindingTree<FDynamicMesh3>* OptionalWindingTree)>;
+	using TestInsideFunctionType = TFunction<bool(const FVector3d& TestPosition, const FDynamicMeshAABBTree3& Spatial, const TFastWindingTree<FDynamicMesh3>* OptionalWindingTree)>;
 	struct FCustomInsideMeshTest
 	{
 		TestInsideFunctionType IsPointInsideFn;
