Modeless Tools - Spline Element Selection Support - Adds pieces of functionality which are unused.

Includes a spline selector as an example which demonstrates automatic selector injection by the object visualizer tools subsystem. The spline selector is associated with a spline point transformer.

A followup change will utilize this by integrating the geometry selection manager into the modeless tools system

diff --git a/Engine/Source/Runtime/GeometryCore/Public/Splines/SplineMath.h b/Engine/Source/Runtime/GeometryCore/Public/Splines/SplineMath.h
--- a/Engine/Source/Runtime/GeometryCore/Public/Splines/SplineMath.h
+++ b/Engine/Source/Runtime/GeometryCore/Public/Splines/SplineMath.h
@@ -1,7 +1,9 @@
 // Copyright Epic Games, Inc. All Rights Reserved.
 
 #pragma once
+
 #include "CoreMinimal.h"
+#include "SplineInterfaces.h"
 #include "ParameterizedTypes.h"
 
 namespace UE::Geometry::Spline
@@ -673,6 +675,38 @@ float FindNearestPoint_Cubic(
 	OutSquaredDistance = BestDistSq;
 	return BestParam;
 }
+
+template <typename ValueType>
+ValueType ComputeFirstDerivativeCentralDifference(const TSplineInterface<ValueType>& Spline, double Param, double Step = UE_KINDA_SMALL_NUMBER)
+{
+	FInterval1f ParameterSpace = Spline.GetParameterSpace();
+	Step = FMath::Max(Step, UE_KINDA_SMALL_NUMBER);
+
+	const double ParamRight = FMath::Clamp(Param + Step, ParameterSpace.Min, ParameterSpace.Max);
+	const double ParamLeft = FMath::Clamp(Param - Step, ParameterSpace.Min, ParameterSpace.Max);
+	
+	const ValueType ValueRight  = Spline.Evaluate(ParamRight);
+	const ValueType ValueLeft = Spline.Evaluate(ParamLeft);
+
+	return (ValueRight - ValueLeft) / (2.0 * Step);
+}
+
+template <typename ValueType>
+ValueType ComputeSecondDerivativeCentralDifference(const TSplineInterface<ValueType>& Spline, double Param, double Step = 0.1)
+{
+	FInterval1f ParameterSpace = Spline.GetParameterSpace();
+	Step = FMath::Max(Step, UE_KINDA_SMALL_NUMBER);
+
+	const double ParamRight = FMath::Clamp(Param + Step, ParameterSpace.Min, ParameterSpace.Max);
+	Param = FMath::Clamp(Param, ParameterSpace.Min, ParameterSpace.Max);
+	const double ParamLeft = FMath::Clamp(Param - Step, ParameterSpace.Min, ParameterSpace.Max);
+	
+	const ValueType ValueRight  = Spline.Evaluate(ParamRight);
+	const ValueType Value  = Spline.Evaluate(Param);
+	const ValueType ValueLeft = Spline.Evaluate(ParamLeft);
+
+	return (ValueRight - 2.0 * Value + ValueLeft) / (Step * Step);
+}
 	
 } // end namespace UE::Geometry::Spline::Math
 } // end namespace UE::Geometry::Spline
