Fix Spline warning when trying to set 0-1 point

diff --git a/Engine/Source/Runtime/GeometryCore/Public/Splines/TangentBezierSpline.h b/Engine/Source/Runtime/GeometryCore/Public/Splines/TangentBezierSpline.h
--- a/Engine/Source/Runtime/GeometryCore/Public/Splines/TangentBezierSpline.h
+++ b/Engine/Source/Runtime/GeometryCore/Public/Splines/TangentBezierSpline.h
@@ -453,12 +453,19 @@ public:
 	bool SetControlPoints(const TArray<FTangentBezierControlPoint>& Points)
 	{
 	    const int32 NumPoints = Points.Num();
-	    
-	    // Require at least 2 points to form a valid segment
+		
+		Clear();
+		
+	    // Special case for less than 2 points
 	    if (NumPoints < 2)
 	    {
-	        UE_LOG(LogSpline, Warning, TEXT("SetControlPoints requires at least 2 points to form a valid segment. Got %d points."), NumPoints);
-	        return false;
+	    	if (NumPoints == 1)
+	    	{
+	    		// Initialize the spline with a single point
+	    		int32 OutNewIndex;
+	    		return InitializeFirstSegment(Points[0], OutNewIndex);
+	    	}
+	    	return true;
 	    }
 	    
 	    // Convert to cubic Bezier control points for the internal spline
