Added comment clarifying that we intentionally don't reset the nbr boundary edge array when switching vertices in MeshSelfUnion

#rnx

diff --git a/Engine/Source/Runtime/GeometryCore/Private/Operations/MeshSelfUnion.cpp b/Engine/Source/Runtime/GeometryCore/Private/Operations/MeshSelfUnion.cpp
--- a/Engine/Source/Runtime/GeometryCore/Private/Operations/MeshSelfUnion.cpp
+++ b/Engine/Source/Runtime/GeometryCore/Private/Operations/MeshSelfUnion.cpp
@@ -400,6 +400,8 @@ bool FMeshSelfUnion::Compute()
 						// old VID is swapped in as the one to consider as unmatched
 						// it will now be matched below
 						BoundaryVID = OldVID;
+						// Note we don't reset BoundaryNbrEdges here, since we still want to filter
+						// the edges near the (now-matched) initial BoundaryVID
 						Mesh->GetAllVtxBoundaryEdges(BoundaryVID, BoundaryNbrEdges);
 						Pos = Mesh->GetVertex(BoundaryVID);
 						DSq = OldDSq;
