Made FDynamicMesh3 queries optionally take an inline-allocated array for local vertex-neighborhood mesh queries

diff --git a/Engine/Source/Runtime/GeometryCore/Private/DynamicMesh/DynamicMesh3_Queries.cpp b/Engine/Source/Runtime/GeometryCore/Private/DynamicMesh/DynamicMesh3_Queries.cpp
--- a/Engine/Source/Runtime/GeometryCore/Private/DynamicMesh/DynamicMesh3_Queries.cpp
+++ b/Engine/Source/Runtime/GeometryCore/Private/DynamicMesh/DynamicMesh3_Queries.cpp
@@ -50,8 +50,8 @@ int FDynamicMesh3::GetVtxBoundaryEdges(int vID, int& e0, int& e1) const
 	return 0;
 }
 
-
-int FDynamicMesh3::GetAllVtxBoundaryEdges(int vID, TArray<int>& EdgeListOut) const
+template<typename ArrayType>
+int FDynamicMesh3::GetAllVtxBoundaryEdges(int vID, ArrayType& EdgeListOut) const
 {
 	if (VertexRefCounts.IsValid(vID))
 	{
@@ -69,6 +69,8 @@ int FDynamicMesh3::GetAllVtxBoundaryEdges(int vID, TArray<int>& EdgeListOut) con
 	return 0;
 }
 
+template GEOMETRYCORE_API int FDynamicMesh3::GetAllVtxBoundaryEdges<TArray<int32>>(int vID, TArray<int32>& EdgeListOut) const;
+template GEOMETRYCORE_API int FDynamicMesh3::GetAllVtxBoundaryEdges<FDynamicMesh3::FLocalIntArray>(int vID, FDynamicMesh3::FLocalIntArray& EdgeListOut) const;
 
 
 void FDynamicMesh3::GetVtxNbrhood(int eID, int vID, int& vOther, int& oppV1, int& oppV2, int& t1, int& t2) const
@@ -129,8 +131,8 @@ int FDynamicMesh3::GetVtxSingleTriangle(int VertexID) const
 	return FDynamicMesh3::InvalidID;
 }
 
-
-EMeshResult FDynamicMesh3::GetVtxTriangles(int vID, TArray<int>& TrianglesOut) const
+template<typename ArrayType>
+EMeshResult FDynamicMesh3::GetVtxTriangles(int vID, ArrayType& TrianglesOut) const
 {
 	if (!IsVertex(vID))
 	{
@@ -168,9 +170,11 @@ EMeshResult FDynamicMesh3::GetVtxTriangles(int vID, TArray<int>& TrianglesOut) c
 	return EMeshResult::Ok;
 }
 
+template GEOMETRYCORE_API EMeshResult FDynamicMesh3::GetVtxTriangles<TArray<int32>>(int vID, TArray<int32>& TrianglesOut) const;
+template GEOMETRYCORE_API EMeshResult FDynamicMesh3::GetVtxTriangles<FDynamicMesh3::FLocalIntArray>(int vID, FDynamicMesh3::FLocalIntArray& TrianglesOut) const;
 
-
-EMeshResult FDynamicMesh3::GetVtxContiguousTriangles(int VertexID, TArray<int>& TrianglesOut, TArray<int>& SpanLengths, TArray<bool>& IsLoop) const
+template<typename IntArray, typename BoolArray>
+EMeshResult FDynamicMesh3::GetVtxContiguousTriangles(int VertexID, IntArray& TrianglesOut, IntArray& SpanLengths, BoolArray& IsLoop) const
 {
 	TrianglesOut.Reset();
 	SpanLengths.Reset();
@@ -190,7 +194,7 @@ EMeshResult FDynamicMesh3::GetVtxContiguousTriangles(int VertexID, TArray<int>&
 	// The number of StartEdgeIDs for a vertex is NumBoundaryEdges + NumRingsWithoutBoundary
 	// This is only higher than 2 at non-manifold (e.g., "bowtie") vertices, and
 	// it should be very rare for it to exceed this inline allocation of 8
-	TArray<int, TInlineAllocator<8>> StartEdgeIDs;
+	FLocalIntArray StartEdgeIDs;
 	// initial starting edge candidates == boundary edges
 	for (int EID : VertexEdgeLists.Values(VertexID))
 	{
@@ -283,6 +287,11 @@ EMeshResult FDynamicMesh3::GetVtxContiguousTriangles(int VertexID, TArray<int>&
 	return ensure(SpanLengths.Num() == IsLoop.Num()) ? EMeshResult::Ok : EMeshResult::Failed_InvalidNeighbourhood;
 }
 
+template GEOMETRYCORE_API EMeshResult FDynamicMesh3::GetVtxContiguousTriangles<TArray<int32>, TArray<bool>>
+	(int VertexID, TArray<int32>& TrianglesOut, TArray<int32>& SpanLengths, TArray<bool>& IsLoop) const;
+template GEOMETRYCORE_API EMeshResult FDynamicMesh3::GetVtxContiguousTriangles<FDynamicMesh3::FLocalIntArray, FDynamicMesh3::FLocalBoolArray>
+	(int VertexID, FDynamicMesh3::FLocalIntArray& TrianglesOut, FDynamicMesh3::FLocalIntArray& SpanLengths, FDynamicMesh3::FLocalBoolArray& IsLoop) const;
+
 bool FDynamicMesh3::IsBoundaryVertex(int vID) const
 {
 	checkSlow(IsVertex(vID));
@@ -459,8 +468,8 @@ bool FDynamicMesh3::GetVertexGroups(int vID, FIndex4i& groups) const
 }
 
 
-
-bool FDynamicMesh3::GetAllVertexGroups(int vID, TArray<int>& GroupsOut) const
+template<typename ArrayType>
+bool FDynamicMesh3::GetAllVertexGroups(int vID, ArrayType& GroupsOut) const
 {
 	if (!HasTriangleGroups()) return false;
 
@@ -481,6 +490,8 @@ bool FDynamicMesh3::GetAllVertexGroups(int vID, TArray<int>& GroupsOut) const
 	return true;
 }
 
+template GEOMETRYCORE_API bool FDynamicMesh3::GetAllVertexGroups<TArray<int32>>(int vID, TArray<int32>& GroupsOut) const;
+template GEOMETRYCORE_API bool FDynamicMesh3::GetAllVertexGroups<FDynamicMesh3::FLocalIntArray>(int vID, FDynamicMesh3::FLocalIntArray& GroupsOut) const;
 
 
 
@@ -489,7 +500,6 @@ bool FDynamicMesh3::GetAllVertexGroups(int vID, TArray<int>& GroupsOut) const
  */
 bool FDynamicMesh3::IsBowtieVertex(int vID) const
 {
-	checkSlow(VertexRefCounts.IsValid(vID));
 	if (VertexRefCounts.IsValid(vID) == false)
 	{
 		return false;
diff --git a/Engine/Source/Runtime/GeometryCore/Public/DynamicMesh/DynamicMesh3.h b/Engine/Source/Runtime/GeometryCore/Public/DynamicMesh/DynamicMesh3.h
--- a/Engine/Source/Runtime/GeometryCore/Public/DynamicMesh/DynamicMesh3.h
+++ b/Engine/Source/Runtime/GeometryCore/Public/DynamicMesh/DynamicMesh3.h
@@ -121,6 +121,10 @@ public:
 	using FPokeTriangleInfo = DynamicMeshInfo::FPokeTriangleInfo;
 	using FVertexSplitInfo  = DynamicMeshInfo::FVertexSplitInfo;
 
+	// Inline-allocator array types optionally used for mesh queries, to reduce heap allocations
+	using FLocalIntArray = TArray<int32, TInlineAllocator<8>>;
+	using FLocalBoolArray = TArray<bool, TInlineAllocator<8>>;
+
 	struct FEdge
 	{
 		FIndex2i Vert;
@@ -1094,8 +1098,10 @@ public:
 	 * @param vID Vertex ID
 	 * @param EdgeListOut boundary edge IDs are appended to this list
 	 * @return count of number of elements of e that were filled
+	 * Note: ArrayType must by TArray<int> or FLocalIntArray
 	 */
-	GEOMETRYCORE_API int GetAllVtxBoundaryEdges(int VertexID, TArray<int>& EdgeListOut) const;
+	template<typename ArrayType = FLocalIntArray>
+	GEOMETRYCORE_API int GetAllVtxBoundaryEdges(int VertexID, ArrayType& EdgeListOut) const;
 
 	/**
 	 * return # of triangles attached to vID, or -1 if invalid vertex
@@ -1104,8 +1110,11 @@ public:
 
 	/**
 	 * Get triangle one-ring at vertex.
+	 * Note: ArrayType must by TArray<int> or FLocalIntArray
 	 */
-	GEOMETRYCORE_API EMeshResult GetVtxTriangles(int VertexID, TArray<int>& TrianglesOut) const;
+	template<typename ArrayType = FLocalIntArray>
+	GEOMETRYCORE_API EMeshResult GetVtxTriangles(int VertexID, ArrayType& TrianglesOut) const;
+	
 
 	/**
 	 * @return Triangle ID for a single triangle connected to VertexID, or InvalidID if VertexID does not exist or has no attached triangles
@@ -1118,9 +1127,11 @@ public:
 	* @param TrianglesOut All triangles connected to the vertex, in contiguous order; if there are multiple contiguous groups they are packed one after another
 	* @param ContiguousGroupLengths Lengths of contiguous groups packed into TrianglesOut (if not a bowtie, this will just be a length-one array w/ {TrianglesOut.Num()})
 	* @param GroupIsLoop Indicates whether each contiguous group is a loop (first triangle connected to last) or not
+	* Note: ArrayTypes must by TArray<int>/<bool> or FLocalIntArray/FLocalBoolArry
 	*/
-	GEOMETRYCORE_API EMeshResult GetVtxContiguousTriangles(int VertexID, TArray<int>& TrianglesOut, TArray<int>& ContiguousGroupLengths,
-	                                      TArray<bool>& GroupIsLoop) const;
+	template<typename IntArrayType = FLocalIntArray, typename BoolArrayType = FLocalBoolArray>
+	GEOMETRYCORE_API EMeshResult GetVtxContiguousTriangles(int VertexID, IntArrayType& TrianglesOut, IntArrayType& ContiguousGroupLengths,
+	                                      BoolArrayType& GroupIsLoop) const;
 
 	/** Returns true if the two triangles connected to edge have different group IDs */
 	GEOMETRYCORE_API bool IsGroupBoundaryEdge(int EdgeID) const;
@@ -1134,8 +1145,9 @@ public:
 	/** Returns up to 4 group IDs at vertex. Returns false if > 4 encountered */
 	GEOMETRYCORE_API bool GetVertexGroups(int VertexID, FIndex4i& GroupsOut) const;
 
-	/** Returns all group IDs at vertex */
-	GEOMETRYCORE_API bool GetAllVertexGroups(int VertexID, TArray<int>& GroupsOut) const;
+	/** Returns all group IDs at vertex. ArrayType must by TArray<int> or FLocalIntArray */
+	template<typename ArrayType = FLocalIntArray>
+	GEOMETRYCORE_API bool GetAllVertexGroups(int VertexID, ArrayType& GroupsOut) const;
 
 	/** returns true if vID is a "bowtie" vertex, ie multiple disjoint triangle sets in one-ring */
 	GEOMETRYCORE_API bool IsBowtieVertex(int VertexID) const;
@@ -1770,7 +1782,8 @@ protected:
 		VertexEdgeLists.AllocateAt(VertexID);
 	}
 
-	void GetVertexEdgesList(int VertexID, TArray<int>& EdgesOut) const
+	template<typename ArrayType = FLocalIntArray>
+	void GetVertexEdgesList(int VertexID, ArrayType& EdgesOut) const
 	{
 		for (int eid : VertexEdgeLists.Values(VertexID))
 		{
