Improved SplitBowties performance and fixed correctness of IsBowtieInOverlay method

diff --git a/Engine/Source/Runtime/GeometryCore/Private/DynamicMesh/DynamicMeshOverlay.cpp b/Engine/Source/Runtime/GeometryCore/Private/DynamicMesh/DynamicMeshOverlay.cpp
--- a/Engine/Source/Runtime/GeometryCore/Private/DynamicMesh/DynamicMeshOverlay.cpp
+++ b/Engine/Source/Runtime/GeometryCore/Private/DynamicMesh/DynamicMeshOverlay.cpp
@@ -3,6 +3,9 @@
 #include "DynamicMesh/DynamicMeshOverlay.h"
 #include "DynamicMesh/DynamicMesh3.h"
 
+#include "Async/ParallelFor.h"
+#include "Async/Mutex.h"
+
 using namespace UE::Geometry;
 
 
@@ -85,9 +88,9 @@ template<typename RealType, int ElementSize>
 void TDynamicMeshOverlay<RealType, ElementSize>::CreateFromPredicate(TFunctionRef<bool(int ParentVertexIdx, int TriIDA, int TriIDB)> TrisCanShareVertexPredicate, RealType InitElementValue)
 {
 	ClearElements(); // deletes all elements and initializes triangles to be 1:1 w/ parentmesh IDs
-	TArray<int> TrisActiveSubGroup, AppendedElements;
-	TArray<int> TriangleIDs, TriangleContigGroupLens;
-	TArray<bool> GroupIsLoop;
+	FDynamicMesh3::FLocalIntArray TrisActiveSubGroup, AppendedElements;
+	FDynamicMesh3::FLocalIntArray TriangleIDs, TriangleContigGroupLens;
+	FDynamicMesh3::FLocalBoolArray GroupIsLoop;
 	for (int VertexID : ParentMesh->VertexIndicesItr())
 	{
 		
@@ -293,13 +296,39 @@ int TDynamicMeshOverlay<RealType, ElementSize>::SplitElementWithNewParent(int El
 
 
 template<typename RealType, int ElementSize>
-void TDynamicMeshOverlay<RealType, ElementSize>::SplitBowties()
+void TDynamicMeshOverlay<RealType, ElementSize>::SplitBowties(bool bParallel)
 {
 	TRACE_CPUPROFILER_EVENT_SCOPE(DynamicMeshOverlay_SplitBowties);
 
-	for (int VertexID : ParentMesh->VertexIndicesItr())
+	if (bParallel)
 	{
-		SplitBowtiesAtVertex(VertexID);
+		// Bowties are typically rare, so we can test for bowties in parallel and only run the split on the (hopefully) few needed cases
+		TArray<int32> BowtieVerts;
+		FMutex BowtieCS;
+		ParallelFor(ParentMesh->MaxVertexID(), [this, &BowtieVerts, &BowtieCS](int32 VertexID)
+			{
+				if (IsBowtieInOverlay(VertexID))
+				{
+					BowtieCS.Lock();
+					BowtieVerts.Add(VertexID);
+					BowtieCS.Unlock();
+				}
+			}
+		);
+		for (int32 VertexID : BowtieVerts)
+		{
+			SplitBowtiesAtVertex(VertexID);
+		}
+	}
+	else
+	{
+		for (int VertexID : ParentMesh->VertexIndicesItr())
+		{
+			if (IsBowtieInOverlay(VertexID))
+			{
+				SplitBowtiesAtVertex(VertexID);
+			}
+		}
 	}
 }
 
@@ -308,12 +337,12 @@ template<typename RealType, int ElementSize>
 void TDynamicMeshOverlay<RealType, ElementSize>::SplitBowtiesAtVertex(int32 VertexID, TArray<int32>* NewElementIDs)
 {
 	// arrays for storing contiguous triangle groups from parentmesh
-	TArray<int> TrianglesOut, ContiguousGroupLengths;
-	TArray<bool> GroupIsLoop;
+	FDynamicMesh3::FLocalIntArray TrianglesOut, ContiguousGroupLengths;
+	FDynamicMesh3::FLocalBoolArray GroupIsLoop;
 
 	// per-vertex element group tracking data, reused in loop below
-	TSet<int> ElementIDSeen;
-	TArray<int> GroupElementIDs, // stores element IDs of this vertex, 1:1 w/ contiguous triangles in the parent mesh
+	FDynamicMesh3::FLocalIntArray ElementIDSeen;
+	FDynamicMesh3::FLocalIntArray GroupElementIDs, // stores element IDs of this vertex, 1:1 w/ contiguous triangles in the parent mesh
 		SubGroupID, // mapping from GroupElementIDs indices into SubGroupElementIDs indices, giving subgroup membership per triangle
 		SubGroupElementIDs; // 1:1 w/ 'subgroups' in the group (e.g. if all triangles in the group all had the same element ID, this would be an array of length 1, just containing that element ID)
 
@@ -390,7 +419,7 @@ void TDynamicMeshOverlay<RealType, ElementSize>::SplitBowtiesAtVertex(int32 Vert
 			// split needed the *second* time we see a sub-group using a given ElementID
 			if (ElementIDSeen.Contains(ElementID))
 			{
-				TArray<int> ConnectedTris;
+				FDynamicMesh3::FLocalIntArray ConnectedTris;
 				for (int TriSubIdx = TriSubStart; TriSubIdx < TriSubEnd; TriSubIdx++)
 				{
 					if (SubID == SubGroupID[TriSubIdx - TriSubStart])
@@ -404,7 +433,10 @@ void TDynamicMeshOverlay<RealType, ElementSize>::SplitBowtiesAtVertex(int32 Vert
 					NewElementIDs->Add(NewElementID);
 				}
 			}
-			ElementIDSeen.Add(ElementID);
+			else
+			{
+				ElementIDSeen.Add(ElementID);
+			}
 		}
 
 		TriSubStart = TriSubEnd;
@@ -796,32 +828,129 @@ bool TDynamicMeshOverlay<RealType, ElementSize>::IsSeamIntersectionVertex(int32
 	return SeamCount == 1 || SeamCount > 2;
 }
 
+
 template<typename RealType, int ElementSize>
 bool TDynamicMeshOverlay<RealType, ElementSize>::IsBowtieInOverlay(int32 VertexID) const
 {
-	// This is a bit tricky but seems to be correct. If we think of a mesh boundary edge as "one" border edge,
-	// and a UV-seam as "two" border edges (one on each side), then we have a non-bowtie configuration if:
-	//  1) NumElements == 1 && BorderEdgeCount == 0
-	//  2) BorderEdgeCount == 2*NumElements
-	// otherwise we have bowties
-	// (todo: validate this with brute-force algorithm that finds uv-connected-components of one-ring?)
-
-	int32 NumBoundary = 0;
-	int32 NumSeams = 0;
-	for (int32 EdgeID : ParentMesh->VtxEdgesItr(VertexID))
+	// arrays for storing contiguous triangle groups from parentmesh
+	FDynamicMesh3::FLocalIntArray TrianglesOut, ContiguousGroupLengths;
+	FDynamicMesh3::FLocalBoolArray GroupIsLoop;
+
+	// per-vertex element group tracking data, reused in loop below
+	FDynamicMesh3::FLocalIntArray ElementIDSeen;
+	FDynamicMesh3::FLocalIntArray GroupElementIDs, // stores element IDs of this vertex, 1:1 w/ contiguous triangles in the parent mesh
+		SubGroupID, // mapping from GroupElementIDs indices into SubGroupElementIDs indices, giving subgroup membership per triangle
+		SubGroupElementIDs; // 1:1 w/ 'subgroups' in the group (e.g. if all triangles in the group all had the same element ID, this would be an array of length 1, just containing that element ID)
+
+	ensure(EMeshResult::Ok == ParentMesh->GetVtxContiguousTriangles(VertexID, TrianglesOut, ContiguousGroupLengths, GroupIsLoop));
+	int32 NumTris = TrianglesOut.Num();
+
+	auto ElementIDFromTriangle = [VertexID, this](int32 TriID) -> int32
 	{
-		if (ParentMesh->IsBoundaryEdge(EdgeID))
+		FIndex3i TriVIDs = ParentMesh->GetTriangle(TriID);
+		FIndex3i TriEIDs = GetTriangle(TriID);
+		int SubIdx = TriVIDs.IndexOf(VertexID);
+		return TriEIDs[SubIdx];
+	};
+
+	// Handle the easy case of 1 element, 1 contiguous group
+	if (NumTris > 0 && ContiguousGroupLengths.Num() == 1)
+	{
+		int32 FirstElSeen = ElementIDFromTriangle(TrianglesOut[0]);
+		bool bSingleElement = true;
+		for (int32 Idx = 1; Idx < TrianglesOut.Num(); ++Idx)
+		{
+			if (FirstElSeen != ElementIDFromTriangle(TrianglesOut[Idx]))
+			{
+				bSingleElement = false;
+			}
+		}
+		if (bSingleElement)
+		{
+			return false;
+		}
+	}
+
+	// More complex case: Iterate through each contiguous group looking for a re-used element across groups
+	ElementIDSeen.Reset();
+	// per contiguous group of triangles around vertex in ParentMesh, find contiguous sub-groups in overlay
+	for (int32 GroupIdx = 0, NumGroups = ContiguousGroupLengths.Num(), TriSubStart = 0; GroupIdx < NumGroups; GroupIdx++)
+	{
+		bool bIsLoop = GroupIsLoop[GroupIdx];
+		int TriInGroupNum = ContiguousGroupLengths[GroupIdx];
+		if (ensure(TriInGroupNum > 0) == false)
+		{
+			continue;
+		}
+		int TriSubEnd = TriSubStart + TriInGroupNum;
+
+		GroupElementIDs.Reset();
+		for (int TriSubIdx = TriSubStart; TriSubIdx < TriSubEnd; TriSubIdx++)
+		{
+			int TriID = TrianglesOut[TriSubIdx];
+			GroupElementIDs.Add(ElementIDFromTriangle(TriID));
+		}
+
+		auto IsConnected = [this, &GroupElementIDs, &TrianglesOut, &TriSubStart](int TriOutIdxA, int TriOutIdxB)
 		{
-			NumBoundary++;
-		} 
-		else if (IsSeamEdge(EdgeID))
+			if (GroupElementIDs[TriOutIdxA - TriSubStart] != GroupElementIDs[TriOutIdxB - TriSubStart])
+			{
+				return false;
+			}
+			int EdgeID = ParentMesh->FindEdgeFromTriPair(TrianglesOut[TriOutIdxA], TrianglesOut[TriOutIdxB]);
+			return EdgeID >= 0 && !IsSeamEdge(EdgeID);
+		};
+
+		SubGroupID.Reset(); SubGroupID.SetNum(TriInGroupNum);
+		SubGroupElementIDs.Reset();
+		int MaxSubID = 0;
+		SubGroupID[0] = 0;
+		SubGroupElementIDs.Add(GroupElementIDs[0]);
+
+		// Iterate through tris in current group, except last one
+		for (int TriSubIdx = TriSubStart; TriSubIdx + 1 < TriSubEnd; TriSubIdx++)
 		{
-			NumSeams++;
+			if (!IsConnected(TriSubIdx, TriSubIdx + 1))
+			{
+				SubGroupElementIDs.Add(GroupElementIDs[TriSubIdx + 1 - TriSubStart]);
+				MaxSubID++;
+			}
+			SubGroupID[TriSubIdx - TriSubStart + 1] = MaxSubID;
 		}
+		// if group was a loop, need to check if the last sub-group and first sub-group were actually the same group
+		if (bIsLoop && MaxSubID > 0 && IsConnected(TriSubStart, TriSubStart + TriInGroupNum - 1))
+		{
+			int LastGroupID = SubGroupID.Last();
+			for (int32 Idx = SubGroupID.Num() - 1; Idx >= 0 && SubGroupID[Idx] == LastGroupID; Idx--)
+			{
+				SubGroupID[Idx] = 0;
+			}
+			MaxSubID--;
+			SubGroupElementIDs.Pop(EAllowShrinking::No);
+		}
+
+		for (int SubID = 0; SubID < SubGroupElementIDs.Num(); SubID++)
+		{
+			int ElementID = SubGroupElementIDs[SubID];
+			if (ElementID < 0)
+			{
+				continue;		// skip if this is an invalid ElementID (eg from an invalid triangle)
+			}
+			// split needed the *second* time we see a sub-group using a given ElementID
+			if (ElementIDSeen.Contains(ElementID))
+			{
+				return true;
+			}
+			else
+			{
+				ElementIDSeen.Add(ElementID);
+			}
+		}
+
+		TriSubStart = TriSubEnd;
 	}
-	int32 NumElements = CountVertexElements(VertexID);
-	int32 BorderEdgeCount = NumBoundary + 2 * NumSeams;
-	return ! ( (BorderEdgeCount == 0 && NumElements == 1) || (BorderEdgeCount == 2*NumElements) );
+
+	return false;
 }
 
 
@@ -866,7 +995,7 @@ void TDynamicMeshOverlay<RealType, ElementSize>::GetVertexElements(int vid, TArr
 template<typename RealType, int ElementSize>
 int TDynamicMeshOverlay<RealType, ElementSize>::CountVertexElements(int vid, bool bBruteForce) const
 {
-	TArray<int> VertexElements;
+	FDynamicMesh3::FLocalIntArray VertexElements;
 	FIndex3i Triangle;
 	if (bBruteForce) 
 	{
diff --git a/Engine/Source/Runtime/GeometryCore/Public/DynamicMesh/DynamicMeshOverlay.h b/Engine/Source/Runtime/GeometryCore/Public/DynamicMesh/DynamicMeshOverlay.h
--- a/Engine/Source/Runtime/GeometryCore/Public/DynamicMesh/DynamicMeshOverlay.h
+++ b/Engine/Source/Runtime/GeometryCore/Public/DynamicMesh/DynamicMeshOverlay.h
@@ -427,9 +427,10 @@ public:
 	void SplitBowtiesAtVertex(int32 Vid, TArray<int32>* NewElementIDs = nullptr);
 
 	/**
-	* Refine an existing overlay topology by splitting any bow ties
+	* Refine an existing overlay topology by splitting any bowties
+	* @param bParallel Whether to run bowtie detection in parallel
 	*/
-	void SplitBowties();
+	void SplitBowties(bool bParallel = true);
 
 
 	//
