add early out to SplitBowtiesAtVertex for the most common non-bowtie case

diff --git a/Engine/Source/Runtime/GeometryCore/Private/DynamicMesh/DynamicMeshOverlay.cpp b/Engine/Source/Runtime/GeometryCore/Private/DynamicMesh/DynamicMeshOverlay.cpp
--- a/Engine/Source/Runtime/GeometryCore/Private/DynamicMesh/DynamicMeshOverlay.cpp
+++ b/Engine/Source/Runtime/GeometryCore/Private/DynamicMesh/DynamicMeshOverlay.cpp
@@ -324,10 +324,7 @@ void TDynamicMeshOverlay<RealType, ElementSize>::SplitBowties(bool bParallel)
 	{
 		for (int VertexID : ParentMesh->VertexIndicesItr())
 		{
-			if (IsBowtieInOverlay(VertexID))
-			{
-				SplitBowtiesAtVertex(VertexID);
-			}
+			SplitBowtiesAtVertex(VertexID);
 		}
 	}
 }
@@ -349,6 +346,32 @@ void TDynamicMeshOverlay<RealType, ElementSize>::SplitBowtiesAtVertex(int32 Vert
 	ensure(EMeshResult::Ok == ParentMesh->GetVtxContiguousTriangles(VertexID, TrianglesOut, ContiguousGroupLengths, GroupIsLoop));
 	int32 NumTris = TrianglesOut.Num();
 
+	auto ElementIDFromTriangle = [VertexID, this](int32 TriID) -> int32
+	{
+		FIndex3i TriVIDs = ParentMesh->GetTriangle(TriID);
+		FIndex3i TriEIDs = GetTriangle(TriID);
+		int SubIdx = TriVIDs.IndexOf(VertexID);
+		return TriEIDs[SubIdx];
+	};
+
+	// Early out in the easy non-bowtie case of 1 element, 1 contiguous group
+	if (NumTris > 0 && ContiguousGroupLengths.Num() == 1)
+	{
+		int32 FirstElSeen = ElementIDFromTriangle(TrianglesOut[0]);
+		bool bSingleElement = true;
+		for (int32 Idx = 1; Idx < TrianglesOut.Num(); ++Idx)
+		{
+			if (FirstElSeen != ElementIDFromTriangle(TrianglesOut[Idx]))
+			{
+				bSingleElement = false;
+			}
+		}
+		if (bSingleElement)
+		{
+			return;
+		}
+	}
+
 	ElementIDSeen.Reset();
 	// per contiguous group of triangles around vertex in ParentMesh, find contiguous sub-groups in overlay
 	for (int32 GroupIdx = 0, NumGroups = ContiguousGroupLengths.Num(), TriSubStart = 0; GroupIdx < NumGroups; GroupIdx++)
@@ -365,10 +388,7 @@ void TDynamicMeshOverlay<RealType, ElementSize>::SplitBowtiesAtVertex(int32 Vert
 		for (int TriSubIdx = TriSubStart; TriSubIdx < TriSubEnd; TriSubIdx++)
 		{
 			int TriID = TrianglesOut[TriSubIdx];
-			FIndex3i TriVIDs = ParentMesh->GetTriangle(TriID);
-			FIndex3i TriEIDs = GetTriangle(TriID);
-			int SubIdx = TriVIDs.IndexOf(VertexID);
-			GroupElementIDs.Add(TriEIDs[SubIdx]);
+			GroupElementIDs.Add(ElementIDFromTriangle(TriID));
 		}
 
 		auto IsConnected = [this, &GroupElementIDs, &TrianglesOut, &TriSubStart](int TriOutIdxA, int TriOutIdxB)
