move the geometrycore sweepgenerator method implementations to the cpp file
+ move local cpp helper methods to static and a file-specific namespace

diff --git a/Engine/Source/Runtime/GeometryCore/Private/Generators/SweepGenerator.cpp b/Engine/Source/Runtime/GeometryCore/Private/Generators/SweepGenerator.cpp
--- a/Engine/Source/Runtime/GeometryCore/Private/Generators/SweepGenerator.cpp
+++ b/Engine/Source/Runtime/GeometryCore/Private/Generators/SweepGenerator.cpp
@@ -4,46 +4,770 @@
 
 #include "Async/ParallelFor.h"
 
-// TODO: Other functions from SweepGenerator.h should go here too.
-
 using namespace UE::Geometry;
 
-// FProfileSweepGenerator
+void FSweepGeneratorBase::ConstructMeshTopology(const FPolygon2d& CrossSection,
+	const TArrayView<const int32>& UVSections,
+	const TArrayView<const int32>& NormalSections,
+	const TArrayView<const int32>& SharpNormalsAlongLength,
+	bool bEvenlySpaceUVs,
+	const TArrayView<const FVector3d>& Path,
+	int32 NumCrossSections,
+	bool bLoop,
+	const ECapType Caps[2],
+	FVector2f SectionsUVScale,
+	FVector2f CapUVScale,
+	FVector2f CapUVOffset,
+	const TArrayView<const float>& CustomCrossSectionTexCoord,
+	const TArrayView<const float>& CustomPathTexCoord)
+{
+	// per cross section
+	const int32 XVerts = CrossSection.VertexCount();
+	const int32 XSegments = bProfileCurveIsClosed ? XVerts : XVerts - 1;
+	const int32 XNormals = XVerts + NormalSections.Num();
+	const int32 XUVs = XSegments + UVSections.Num() + 1;
 
-// Various indexing utility functions. See the comment in Generate() for a description of vertex/triangle/uv layout.
-int32 GetVertIndex(bool VertIsWelded, int32 SweepIndex, int32 ProfileIndex, int32 NumWelded, int32 NumNonWelded, const TArray<int32> &VertPositionOffsets)
+	float TotalPerimeter = 0.0f, TotalPathLength = 0.0f;
+	TArray<float> CrossSectionTexCoord, PathTexCoord;
+
+	// Compute texture coordinates along the cross section (U coordinates)
+	if (bProfileCurveIsClosed && CustomCrossSectionTexCoord.Num() >= XVerts + 1)
+	{
+		CrossSectionTexCoord = CustomCrossSectionTexCoord;
+	}
+	else if (bProfileCurveIsClosed && CustomCrossSectionTexCoord.Num() == XVerts)
+	{
+		CrossSectionTexCoord = CustomCrossSectionTexCoord;
+
+		// If the cross section curve is closed and we are missing texture coordinate for the 
+		// last element we wrap araound and use the coordinate of the first element
+		CrossSectionTexCoord.Add(CrossSectionTexCoord[0]);
+	}
+	else if (bProfileCurveIsClosed == false && CustomCrossSectionTexCoord.Num() >= XVerts)
+	{
+		CrossSectionTexCoord = CustomCrossSectionTexCoord;
+	}
+	else if (bEvenlySpaceUVs)
+	{
+		CrossSectionTexCoord.Add(0.0f);
+		for (int Idx = 0; Idx < XSegments; Idx++)
+		{
+			float SegLen = float(Distance(CrossSection[Idx], CrossSection[(Idx + 1) % XSegments]));
+			TotalPerimeter += SegLen;
+			CrossSectionTexCoord.Add(TotalPerimeter);
+		}
+		TotalPerimeter = FMath::Max(TotalPerimeter, FMathf::ZeroTolerance);
+		for (int Idx = 0; Idx < CrossSectionTexCoord.Num(); Idx++)
+		{
+			CrossSectionTexCoord[Idx] /= TotalPerimeter;
+			CrossSectionTexCoord[Idx] = 1.0f - CrossSectionTexCoord[Idx];
+		}
+	}
+	else
+	{
+		for (int Idx = 0; Idx < XSegments; Idx++)
+		{
+			float U = float(Idx) / float(XSegments);
+			CrossSectionTexCoord.Add(1.0f - U);
+		}
+		CrossSectionTexCoord.Add(0.0f);
+	}
+
+	// Compute texture coordinates along the path (V coordinates)
+	if (CustomPathTexCoord.Num() >= NumCrossSections)
+	{
+		PathTexCoord = CustomPathTexCoord;
+	}
+	else if (bLoop && CustomPathTexCoord.Num() == NumCrossSections - 1)
+	{
+		PathTexCoord = CustomPathTexCoord;
+
+		// If the path curve is closed and we are missing texture coordinate for the 
+		// last element we wrap araound and use the coordinate of the first element
+		PathTexCoord.Add(PathTexCoord[0]);
+	}
+	else if (bEvenlySpaceUVs)
+	{
+		PathTexCoord.Add(0.0f);
+		int NumPathSegs = bLoop ? Path.Num() : Path.Num() - 1;
+		for (int Idx = 0; Idx < NumPathSegs; Idx++)
+		{
+			float SegLen = float(Distance(Path[Idx], Path[(Idx + 1) % Path.Num()]));
+			TotalPathLength += SegLen;
+			PathTexCoord.Add(TotalPathLength);
+		}
+		TotalPathLength = FMath::Max(TotalPathLength, FMathf::ZeroTolerance);
+		for (int Idx = 0; Idx < PathTexCoord.Num(); Idx++)
+		{
+			PathTexCoord[Idx] /= TotalPathLength;
+			PathTexCoord[Idx] = 1.0f - PathTexCoord[Idx];
+		}
+	}
+	else
+	{
+		for (int32 Idx = 0; Idx < NumCrossSections; Idx++)
+		{
+			float V = float(Idx) / float(NumCrossSections - 1);
+			PathTexCoord.Add(1.0f - V);
+		}
+	}
+
+	int32 NumVerts = XVerts * NumCrossSections - (bLoop ? XVerts : 0);
+	int32 NumNormals = NumCrossSections > 1 ? (XNormals * NumCrossSections - (bLoop ? XNormals : 0)) : 0;
+	NumNormals += XNormals * SharpNormalsAlongLength.Num();
+	int32 NumUVs = NumCrossSections > 1 ? XUVs * NumCrossSections : 0;
+	int32 NumPolygons = (NumCrossSections - 1) * XSegments;
+	int32 NumTriangles = NumPolygons * 2;
+
+	TArray<FIndex3i> OutTriangles;
+
+	// doesn't make sense to have cap types if the sweep is a loop
+	ensure(!bLoop || (Caps[0] == ECapType::None && Caps[1] == ECapType::None));
+
+	if (!bLoop)
+	{
+		for (int32 CapIdx = 0; !bLoop && CapIdx < 2; CapIdx++)
+		{
+			CapVertStart[CapIdx] = NumVerts;
+			CapNormalStart[CapIdx] = NumNormals;
+			CapUVStart[CapIdx] = NumUVs;
+			CapTriangleStart[CapIdx] = NumTriangles;
+			CapPolygonStart[CapIdx] = NumPolygons;
+
+			if (Caps[CapIdx] == ECapType::FlatTriangulation)
+			{
+				NumTriangles += XVerts - 2;
+				NumPolygons++;
+				NumUVs += XVerts;
+				NumNormals += XVerts;
+			}
+			else if (Caps[CapIdx] == ECapType::FlatMidpointFan)
+			{
+				NumTriangles += XSegments;
+				NumPolygons++;
+				NumUVs += XVerts + 1;
+				NumNormals += XVerts + 1;
+				NumVerts += 1;
+			}
+			// TODO: support more cap type; e.g.:
+			//else if (Caps[CapIdx] == ECapType::Cone)
+			//{
+			//	NumTriangles += XVerts;
+			//	NumPolygons += XVerts;
+			//	NumUVs += XVerts + 1;
+			//	NumNormals += XVerts * 2;
+			//	NumVerts += 1;
+			//}
+		}
+	}
+
+	SetBufferSizes(NumVerts, NumTriangles, NumUVs, NumNormals);
+
+	if (!bLoop)
+	{
+		for (int32 CapIdx = 0; CapIdx < 2; CapIdx++)
+		{
+			if (Caps[CapIdx] == ECapType::FlatTriangulation)
+			{
+				int32 VertOffset = CapIdx * (XVerts * (NumCrossSections - 1));
+
+				PolygonTriangulation::TriangulateSimplePolygon(CrossSection.GetVertices(), OutTriangles);
+				int32 TriIdx = CapTriangleStart[CapIdx];
+				int32 PolyIdx = CapPolygonStart[CapIdx];
+				for (const FIndex3i& Triangle : OutTriangles)
+				{
+					bool Flipped = CapIdx == 0;
+					SetTriangle(TriIdx,
+						Triangle.A + VertOffset, Triangle.B + VertOffset, Triangle.C + VertOffset,
+						Flipped);
+					SetTriangleUVs(TriIdx,
+						Triangle.A + CapUVStart[CapIdx],
+						Triangle.B + CapUVStart[CapIdx],
+						Triangle.C + CapUVStart[CapIdx],
+						Flipped);
+					SetTriangleNormals(TriIdx,
+						Triangle.A + CapNormalStart[CapIdx],
+						Triangle.B + CapNormalStart[CapIdx],
+						Triangle.C + CapNormalStart[CapIdx],
+						Flipped);
+					SetTrianglePolygon(TriIdx, PolyIdx);
+					TriIdx++;
+				}
+				float SideScale = float(2 * CapIdx - 1);
+				for (int32 Idx = 0; Idx < XVerts; Idx++)
+				{
+					FVector2f CenteredVert = (FVector2f)CrossSection.GetVertices()[Idx] * CapUVScale + CapUVOffset;
+					SetUV(CapUVStart[CapIdx] + Idx, FVector2f(CenteredVert.X * SideScale, CenteredVert.Y), VertOffset + Idx);
+
+					// correct normal to be filled by subclass
+					SetNormal(CapNormalStart[CapIdx] + Idx, FVector3f::Zero(), VertOffset + Idx);
+				}
+			}
+			else if (Caps[CapIdx] == ECapType::FlatMidpointFan)
+			{
+				int32 VertOffset = CapIdx * (XVerts * (NumCrossSections - 1));
+				int32 CapVertStartIdx = CapVertStart[CapIdx];
+				int32 TriIdx = CapTriangleStart[CapIdx];
+				int32 PolyIdx = CapPolygonStart[CapIdx];
+				for (int32 VertIdx = 0; VertIdx < XSegments; VertIdx++)
+				{
+					bool Flipped = CapIdx == 0;
+					SetTriangle(TriIdx,
+						VertOffset + VertIdx,
+						CapVertStartIdx,
+						VertOffset + (VertIdx + 1) % XVerts,
+						Flipped);
+					SetTriangleUVs(TriIdx,
+						CapUVStart[CapIdx] + VertIdx,
+						CapUVStart[CapIdx] + XVerts,
+						CapUVStart[CapIdx] + (VertIdx + 1) % XVerts,
+						Flipped);
+					SetTriangleNormals(TriIdx,
+						CapNormalStart[CapIdx] + VertIdx,
+						CapNormalStart[CapIdx] + XVerts,
+						CapNormalStart[CapIdx] + (VertIdx + 1) % XVerts,
+						Flipped);
+					SetTrianglePolygon(TriIdx, PolyIdx);
+					TriIdx++;
+				}
+
+				// Set cap midpoint UV & Normal
+				// (correct normal to be filled by subclass)
+				SetUV(CapUVStart[CapIdx] + XVerts, FVector2f::Zero() + CapUVOffset, CapVertStartIdx);
+				SetNormal(CapNormalStart[CapIdx] + XVerts, FVector3f::Zero(), CapVertStartIdx);
+
+				// Set cap profile UVs & Normal
+				for (int32 Idx = 0; Idx < XVerts; Idx++)
+				{
+					FVector2f CenteredVert = (FVector2f)CrossSection.GetVertices()[Idx] * CapUVScale + CapUVOffset;
+					SetUV(CapUVStart[CapIdx] + Idx, FVector2f(CenteredVert.X, CenteredVert.Y), VertOffset + Idx);
+					SetNormal(CapNormalStart[CapIdx] + Idx, FVector3f::Zero(), VertOffset + Idx);
+				}
+			}
+		}
+	}
+
+	// fill in UVs and triangles along length
+	int MinValidCrossSections = bLoop ? 3 : 2;
+	int CurFaceGroupIndex = NumPolygons;
+	if (NumCrossSections >= MinValidCrossSections)
+	{
+		int32 UVSection = 0, UVSubIdx = 0;
+
+		int CrossSectionsMod = NumCrossSections;
+		if (bLoop)
+		{
+			CrossSectionsMod--; // last cross section becomes the first
+		}
+		int NormalCrossSectionsMod = CrossSectionsMod + SharpNormalsAlongLength.Num();
+
+		int32 NumSections = UVSections.Num();
+		int32 NextDupVertIdx = UVSection < NumSections ? UVSections[UVSection] : -1;
+		for (int32 VertSubIdx = 0; VertSubIdx < XSegments; UVSubIdx++)
+		{
+			for (int32 XIdx = 0; XIdx < NumCrossSections; XIdx++)
+			{
+				FVector2f UV = FVector2f(CrossSectionTexCoord[VertSubIdx], PathTexCoord[XIdx]);
+				SetUV(XIdx * XUVs + UVSubIdx, UV * SectionsUVScale, (XIdx % CrossSectionsMod) * XVerts + VertSubIdx);
+			}
+
+			if (VertSubIdx == NextDupVertIdx)
+			{
+				NextDupVertIdx = UVSection < NumSections ? UVSections[UVSection] : -1;
+			}
+			else
+			{
+				for (int32 XIdx = 0; XIdx + 1 < NumCrossSections; XIdx++)
+				{
+					SetTriangleUVs(
+						XSegments * 2 * XIdx + 2 * VertSubIdx,
+						XIdx * XUVs + UVSubIdx,
+						XIdx * XUVs + UVSubIdx + 1,
+						(XIdx + 1) * XUVs + UVSubIdx, true);
+					SetTriangleUVs(
+						XSegments * 2 * XIdx + 2 * VertSubIdx + 1,
+						(XIdx + 1) * XUVs + UVSubIdx + 1,
+						(XIdx + 1) * XUVs + UVSubIdx,
+						XIdx * XUVs + UVSubIdx + 1, true);
+				}
+				VertSubIdx++;
+			}
+		}
+		{
+			// final UV
+			int32 VertSubIdx = bProfileCurveIsClosed ? 0 : XSegments;
+			for (int32 XIdx = 0; XIdx < NumCrossSections; XIdx++)
+			{
+				FVector2f UV = FVector2f(CrossSectionTexCoord.Last(), PathTexCoord[XIdx]);
+				SetUV(XIdx * XUVs + UVSubIdx, UV * SectionsUVScale, (XIdx % CrossSectionsMod) * XVerts + VertSubIdx);
+			}
+		}
+		NumSections = NormalSections.Num();
+		int32 NormalSection = 0;
+		NextDupVertIdx = NormalSection < NumSections ? NormalSections[NormalSection] : -1;
+		check(NextDupVertIdx < XVerts);
+		for (int32 VertSubIdx = 0, NormalSubIdx = 0; VertSubIdx < XVerts; NormalSubIdx++)
+		{
+			int SharpNormalIdx = 0;
+			for (int32 XIdx = 0, NormalXIdx = 0; XIdx < NumCrossSections; XIdx++, NormalXIdx++)
+			{
+				// just set the normal parent; don't compute normal yet
+				SetNormal((NormalXIdx % NormalCrossSectionsMod) * XNormals + NormalSubIdx, FVector3f(0, 0, 0), (XIdx % CrossSectionsMod) * XVerts + VertSubIdx);
+				// duplicate normals for cross sections that are 'sharp'
+				if (SharpNormalIdx < SharpNormalsAlongLength.Num() && XIdx == SharpNormalsAlongLength[SharpNormalIdx])
+				{
+					NormalXIdx++;
+					SetNormal((NormalXIdx % NormalCrossSectionsMod) * XNormals + NormalSubIdx, FVector3f(0, 0, 0), (XIdx % CrossSectionsMod) * XVerts + VertSubIdx);
+					SharpNormalIdx++;
+				}
+			}
+
+			if (VertSubIdx == NextDupVertIdx)
+			{
+				NextDupVertIdx = NormalSection < NumSections ? NormalSections[NormalSection] : -1;
+				check(NextDupVertIdx < XVerts);
+			}
+			else
+			{
+				if (VertSubIdx < XSegments) // if bProfileCurveIsClosed == false skip the last triangle strip generation
+				{
+					int32 WrappedNextNormalSubIdx = (NormalSubIdx + 1) % XNormals;
+					int32 WrappedNextVertexSubIdx = (VertSubIdx + 1) % XVerts;
+					SharpNormalIdx = 0;
+					for (int32 XIdx = 0, NXIdx = 0; XIdx + 1 < NumCrossSections; XIdx++, NXIdx++)
+					{
+						int32 T0Idx = XSegments * 2 * XIdx + 2 * VertSubIdx;
+						int32 T1Idx = T0Idx + 1;
+						int32 PIdx = XSegments * XIdx + VertSubIdx;
+						int32 NextXIdx = (XIdx + 1) % CrossSectionsMod;
+						int32 NextNXIdx = (NXIdx + 1) % NormalCrossSectionsMod;
+						SetTrianglePolygon(T0Idx, (bPolygroupPerQuad) ? PIdx : (CurFaceGroupIndex + XIdx));
+						SetTrianglePolygon(T1Idx, (bPolygroupPerQuad) ? PIdx : (CurFaceGroupIndex + XIdx));
+						SetTriangle(T0Idx,
+							XIdx * XVerts + VertSubIdx,
+							XIdx * XVerts + WrappedNextVertexSubIdx,
+							NextXIdx * XVerts + VertSubIdx, true);
+						SetTriangle(T1Idx,
+							NextXIdx * XVerts + WrappedNextVertexSubIdx,
+							NextXIdx * XVerts + VertSubIdx,
+							XIdx * XVerts + WrappedNextVertexSubIdx, true);
+						SetTriangleNormals(
+							T0Idx,
+							NXIdx * XNormals + NormalSubIdx,
+							NXIdx * XNormals + WrappedNextNormalSubIdx,
+							NextNXIdx * XNormals + NormalSubIdx, true);
+						SetTriangleNormals(
+							T1Idx,
+							NextNXIdx * XNormals + WrappedNextNormalSubIdx,
+							NextNXIdx * XNormals + NormalSubIdx,
+							NXIdx * XNormals + WrappedNextNormalSubIdx, true);
+						if (SharpNormalIdx < SharpNormalsAlongLength.Num() && XIdx + 1 == SharpNormalsAlongLength[SharpNormalIdx])
+						{
+							NXIdx++;
+							SharpNormalIdx++;
+						}
+					}
+				}
+				VertSubIdx++;
+			}
+		}
+	}
+}
+
+float FVerticalCylinderGeneratorBase::ComputeSegLengths(const TArrayView<float>& Radii, const TArrayView<float>& Heights, TArray<float>& AlongPercents)
+{
+	float LenAlong = 0;
+	int32 NumX = Radii.Num();
+	AlongPercents.SetNum(NumX);
+	AlongPercents[0] = 0;
+	for (int XIdx = 0; XIdx + 1 < NumX; XIdx++)
+	{
+		double Dist = Distance(FVector2d(Radii[XIdx], Heights[XIdx]), FVector2d(Radii[XIdx + 1], Heights[XIdx + 1]));
+		LenAlong += float(Dist);
+		AlongPercents[XIdx + 1] = LenAlong;
+	}
+	for (int XIdx = 0; XIdx + 1 < NumX; XIdx++)
+	{
+		AlongPercents[XIdx + 1] /= LenAlong;
+	}
+	return LenAlong;
+}
+
+bool FVerticalCylinderGeneratorBase::GenerateVerticalCircleSweep(const TArrayView<float>& Radii, const TArrayView<float>& Heights, const TArrayView<int>& SharpNormalsAlongLength)
+{
+	FPolygon2d X = FPolygon2d::MakeCircle(1.0, AngleSamples);
+	const TArray<FVector2d>& XVerts = X.GetVertices();
+	ECapType Caps[2] = { ECapType::None, ECapType::None };
+
+	if (bCapped)
+	{
+		Caps[0] = CapType;
+		Caps[1] = CapType;
+	}
+
+	int NumX = Radii.Num();
+	if (!ensure(NumX == Heights.Num()))
+	{
+		return false;
+	}
+	// first and last cross sections can't be sharp, so can't have more than NumX-2 sharp normal indices
+	if (!ensure(SharpNormalsAlongLength.Num() + 2 <= NumX))
+	{
+		return false;
+	}
+
+	TArray<float> AlongPercents;
+	float LenAlong = ComputeSegLengths(Radii, Heights, AlongPercents);
+
+	ConstructMeshTopology(X, {}, {}, SharpNormalsAlongLength, false, {}, NumX, false, Caps, FVector2f(1, 1), FVector2f(.5f, .5f), FVector2f(.5f, .5f));
+
+	TArray<FVector2d> NormalSides; NormalSides.SetNum(NumX - 1);
+	for (int XIdx = 0; XIdx + 1 < NumX; XIdx++)
+	{
+		FVector2d Vec = FVector2d(Radii[XIdx + 1], Heights[XIdx + 1]) - FVector2d(Radii[XIdx], Heights[XIdx]);
+		NormalSides[XIdx] = Normalized(PerpCW(Vec));
+	}
+	TArray<FVector2d> SmoothedNormalSides; SmoothedNormalSides.SetNum(NumX);
+	// smooth internal normals
+	SmoothedNormalSides[0] = NormalSides[0];
+	SmoothedNormalSides.Last() = NormalSides.Last();
+	for (int XIdx = 1; XIdx + 1 < NumX; XIdx++)
+	{
+		SmoothedNormalSides[XIdx] = Normalized(NormalSides[XIdx] + NormalSides[XIdx - 1]);
+	}
+
+
+	// set vertex positions and normals for all cross sections along length
+	for (int SubIdx = 0; SubIdx < X.VertexCount(); SubIdx++)
+	{
+		int SharpNormalIdx = 0;
+		for (int XIdx = 0, NormalXIdx = 0; XIdx < NumX; ++XIdx, ++NormalXIdx)
+		{
+			double AlongRadius = Radii[XIdx];
+			Vertices[SubIdx + XIdx * AngleSamples] =
+				FVector3d(XVerts[SubIdx].X * AlongRadius, XVerts[SubIdx].Y * AlongRadius, Heights[XIdx]);
+			if (SharpNormalIdx < SharpNormalsAlongLength.Num() && XIdx == SharpNormalsAlongLength[SharpNormalIdx])
+			{
+				// write sharp normals
+				if (ensure(XIdx > 0)) // very first index cannot be sharp
+				{
+					Normals[SubIdx + NormalXIdx * AngleSamples] = FVector3f(float(XVerts[SubIdx].X * NormalSides[XIdx - 1].X), float(XVerts[SubIdx].Y * NormalSides[XIdx - 1].X), float(NormalSides[XIdx - 1].Y));
+				}
+				NormalXIdx++;
+				if (ensure(XIdx + 1 < NumX)) // very last index cannot be sharp
+				{
+					Normals[SubIdx + NormalXIdx * AngleSamples] = FVector3f(float(XVerts[SubIdx].X * NormalSides[XIdx].X), float(XVerts[SubIdx].Y * NormalSides[XIdx].X), float(NormalSides[XIdx].Y));
+				}
+				SharpNormalIdx++;
+			}
+			else
+			{
+				// write smoothed normal
+				Normals[SubIdx + NormalXIdx * AngleSamples] = FVector3f(float(XVerts[SubIdx].X * SmoothedNormalSides[XIdx].X), float(XVerts[SubIdx].Y * SmoothedNormalSides[XIdx].X), float(SmoothedNormalSides[XIdx].Y));
+			}
+		}
+	}
+
+	if (bCapped)
+	{
+		// if capped, set vertices.
+		for (int CapIdx = 0; CapIdx < 2; CapIdx++)
+		{
+			if (Caps[CapIdx] == ECapType::FlatMidpointFan)
+			{
+				Vertices[CapVertStart[CapIdx]] = FVector3d::UnitZ() * (double)Heights[CapIdx * (Heights.Num() - 1)];
+			}
+		}
+
+		// if capped, set top/bottom normals
+		for (int SubIdx = 0; SubIdx < X.VertexCount(); SubIdx++)
+		{
+			for (int XBotTop = 0; XBotTop < 2; ++XBotTop)
+			{
+				Normals[CapNormalStart[XBotTop] + SubIdx] = FVector3f(0.f, 0.f, float(2 * XBotTop - 1));
+			}
+		}
+		for (int CapIdx = 0; CapIdx < 2; CapIdx++)
+		{
+			if (Caps[CapIdx] == ECapType::FlatMidpointFan)
+			{
+				Normals[CapNormalStart[CapIdx] + X.VertexCount()] = FVector3f(0.f, 0.f, float(2 * CapIdx - 1));
+			}
+		}
+	}
+
+	if (bUVScaleMatchSidesAndCaps)
+	{
+		float MaxAbsRad = FMathf::Abs(Radii[0]);
+		for (int XIdx = 0; XIdx < NumX; XIdx++)
+		{
+			MaxAbsRad = FMathf::Max(FMathf::Abs(Radii[XIdx]), MaxAbsRad);
+		}
+		float AbsHeight = LenAlong;
+		float MaxAbsCircumference = MaxAbsRad * FMathf::TwoPi;
+
+		// scales to put each differently-scaled UV coordinate into the same space
+		float ThetaScale = MaxAbsCircumference;
+		float HeightScale = AbsHeight;
+		float CapScale = MaxAbsRad * 2;
+
+		float MaxScale = FMathf::Max3(ThetaScale, HeightScale, CapScale);
+		ThetaScale /= MaxScale;
+		HeightScale /= MaxScale;
+		CapScale /= MaxScale;
+		for (int UVIdx = 0; UVIdx < CapUVStart[0]; UVIdx++)
+		{
+			UVs[UVIdx].X *= ThetaScale;
+			UVs[UVIdx].Y *= HeightScale;
+		}
+		for (int UVIdx = CapUVStart[0]; UVIdx < UVs.Num(); UVIdx++)
+		{
+			UVs[UVIdx] *= CapScale;
+		}
+	}
+
+	return true;
+}
+
+FMeshShapeGenerator& FCylinderGenerator::Generate()
 {
-	return VertIsWelded ? 
-		VertPositionOffsets[ProfileIndex] 
-		: NumWelded + SweepIndex * NumNonWelded + VertPositionOffsets[ProfileIndex];
+	if (bCapped)
+	{
+		AngleSamples = FMath::Max(AngleSamples, 3);
+	}
+
+	TArray<float> Radii, Heights;
+
+	Radii.Add(Radius[0]);
+	Heights.Add(0);
+	for (int ExtraIdx = 0; ExtraIdx < LengthSamples; ExtraIdx++)
+	{
+		float Along = float(ExtraIdx + 1) / float(LengthSamples + 1);
+		Radii.Add(FMath::Lerp(Radius[0], Radius[1], Along));
+		Heights.Add(Height * Along);
+	}
+	Radii.Add(Radius[1]);
+	Heights.Add(Height);
+
+	GenerateVerticalCircleSweep(Radii, Heights, {});
+
+	return *this;
 }
-int32 GetTriangleIndex(int32 SweepIndex, int32 ProfileIndex, int32 NumTrisPerSweepSegment, const TArray<int32> &VertTriangleOffsets)
+
+FMeshShapeGenerator& FArrowGenerator::Generate()
 {
-	return (SweepIndex * NumTrisPerSweepSegment) + VertTriangleOffsets[ProfileIndex];
+	const float SrcRadii[]{ StickRadius, StickRadius, HeadBaseRadius, HeadTipRadius };
+	const float SrcHeights[]{ 0, StickLength, StickLength, StickLength + HeadLength };
+
+	TArray<float> Radii, Heights;
+	const int NumVerts = 4 + AdditionalLengthSamples[0] + AdditionalLengthSamples[1] + AdditionalLengthSamples[2];
+	Radii.SetNumUninitialized(NumVerts);
+	Heights.SetNumUninitialized(NumVerts);
+
+	int VertIdx = 0;
+	auto SetVert = [&Radii, &Heights, &VertIdx](float Radius, float Height)
+		{
+			Radii[VertIdx] = Radius;
+			Heights[VertIdx] = Height;
+			++VertIdx;
+		};
+
+	for (int SegIdx = 0; true; ++SegIdx)
+	{
+		SetVert(SrcRadii[SegIdx], SrcHeights[SegIdx]);
+
+		if (SegIdx == 3)
+		{
+			break;
+		}
+
+		for (int ExtraSeg = 1, NumExtraSegs = AdditionalLengthSamples[SegIdx] + 1; ExtraSeg < NumExtraSegs; ++ExtraSeg)
+		{
+			const float Along = float(ExtraSeg) / float(NumExtraSegs);
+			SetVert(FMath::Lerp(SrcRadii[SegIdx], SrcRadii[SegIdx + 1], Along),
+				FMath::Lerp(SrcHeights[SegIdx], SrcHeights[SegIdx + 1], Along));
+		}
+	}
+
+	TArray<int> SharpNormalsAlongLength{ 1 + AdditionalLengthSamples[0], 2 + AdditionalLengthSamples[0] + AdditionalLengthSamples[1] };
+
+	GenerateVerticalCircleSweep(Radii, Heights, SharpNormalsAlongLength);
+
+	return *this;
 }
-int32 GetUvIndex(int32 SweepIndex, int32 ProfileIndex, int32 NumUvColumns)
+
+FMeshShapeGenerator& FGeneralizedCylinderGenerator::Generate()
 {
-	return (SweepIndex * NumUvColumns) + ProfileIndex;
+	const TArray<FVector2d>& XVerts = CrossSection.GetVertices();
+	ECapType Caps[2] = { ECapType::None, ECapType::None };
+
+	if (bCapped && !bLoop)
+	{
+		Caps[0] = CapType;
+		Caps[1] = CapType;
+	}
+	int PathNum = Path.Num();
+
+	bool bHavePathScaling = (PathScales.Num() == PathNum);
+	bool bApplyScaling = (bHavePathScaling || (StartScale != 1.0) || (EndScale != 1.0)) && (bLoop == false);
+	bool bNeedArcLength = (bApplyScaling || bUVScaleRelativeWorld);
+	double TotalPathArcLength = (bNeedArcLength) ? UE::Geometry::CurveUtil::ArcLength<double, FVector3d>(Path, bLoop) : 1.0;
+
+	FAxisAlignedBox2f Bounds = (FAxisAlignedBox2f)CrossSection.Bounds();
+	double BoundsMaxDimInv = 1.0 / FMathd::Max(Bounds.MaxDim(), .001);
+	FVector2f SectionScale(1.f, 1.f), CapScale((float)BoundsMaxDimInv, (float)BoundsMaxDimInv);
+	if (bUVScaleRelativeWorld)
+	{
+		double Perimeter = CrossSection.Perimeter();
+		SectionScale.X = float(Perimeter / UnitUVInWorldCoordinates);
+		SectionScale.Y = float(TotalPathArcLength / UnitUVInWorldCoordinates);
+		CapScale.X = CapScale.Y = 1.0f / UnitUVInWorldCoordinates;
+	}
+	ConstructMeshTopology(CrossSection, {}, {}, {}, true, Path, PathNum + (bLoop ? 1 : 0), bLoop, Caps, SectionScale, CapScale, Bounds.Center(), CrossSectionTexCoord, PathTexCoord);
+
+	int XNum = CrossSection.VertexCount();
+	TArray<FVector2d> XNormals; XNormals.SetNum(XNum);
+	for (int Idx = 0; Idx < XNum; Idx++)
+	{
+		XNormals[Idx] = CrossSection.GetNormal_FaceAvg(Idx);
+	}
+
+	double AccumArcLength = 0;
+	FFrame3d CrossSectionFrame = InitialFrame;
+	bool bHaveExplicitFrames = (PathFrames.Num() == Path.Num());
+	for (int PathIdx = 0; PathIdx < PathNum; ++PathIdx)
+	{
+		FVector3d C = Path[PathIdx];
+		FVector3d X, Y;
+		FMatrix2d MiterPosScale = FMatrix2d::Identity();
+		if (bHaveExplicitFrames)
+		{
+			CrossSectionFrame = PathFrames[PathIdx];
+			C = PathFrames[PathIdx].Origin;
+		}
+		if (!bHaveExplicitFrames || bAlignFramesToSampledTangents)
+		{
+			FVector3d ToPrev, ToNext;
+			UE::Geometry::CurveUtil::GetVectorsToPrevNext<double, FVector3d>(Path, PathIdx, ToPrev, ToNext, true, bLoop);
+
+			FVector3d Tangent = Normalized(-ToPrev + ToNext);
+			CrossSectionFrame.AlignAxis(2, Tangent);
+			if (MiterLimit > 1)
+			{
+				// only miter if neither neighbor vector is exactly zero
+				if (!ToPrev.IsZero() && !ToNext.IsZero())
+				{
+					FVector3d TurnDir = ToPrev + ToNext;
+					FVector2d ScaleAlong = (FVector2d)CrossSectionFrame.ToFrameVector(TurnDir);
+					double ScaleAmount = 1.0;
+					if (ScaleAlong.Normalize())
+					{
+						double CosTwoTheta = ToPrev.Dot(ToNext);
+						double SqrScale = 2.0 / (1.0 - CosTwoTheta);
+						if (!FMath::IsFinite(SqrScale) || SqrScale > MiterLimit * MiterLimit)
+						{
+							ScaleAmount = MiterLimit;
+						}
+						else
+						{
+							ScaleAmount = FMath::Sqrt(SqrScale);
+						}
+						MiterPosScale = FMatrix2d::AxisScale(ScaleAlong, ScaleAmount, false /* axis was already normalized */);
+					}
+				}
+			}
+		}
+		X = CrossSectionFrame.X();
+		Y = CrossSectionFrame.Y();
+
+		double T = FMathd::Clamp((AccumArcLength / TotalPathArcLength), 0.0, 1.0);
+		double UniformScale = (bApplyScaling) ? FMathd::Lerp(StartScale, EndScale, T) : 1.0;
+		FVector2d PathScaling = (bHavePathScaling) ? PathScales[PathIdx] : FVector2d::One();
+
+		for (int SubIdx = 0; SubIdx < XNum; SubIdx++)
+		{
+			FVector2d XP = MiterPosScale * (UniformScale * PathScaling * CrossSection[SubIdx]);
+			FVector2d XN = XNormals[SubIdx];
+			// Note: Arguably, one could apply an inverse mitering scale to the normals, and then re-normalize, here ...
+			// However, it's better to think about these normals an average from the two unscaled cross sections that are meeting at the cross section plane,
+			// and in that model they won't be affected by the miter scale
+			Vertices[SubIdx + PathIdx * XNum] = C + X * XP.X + Y * XP.Y;
+			Normals[SubIdx + PathIdx * XNum] = (FVector3f)(X * XN.X + Y * XN.Y);
+		}
+
+		if (PathIdx < PathNum - 1)
+		{
+			AccumArcLength += Distance(C, Path[PathIdx + 1]);
+		}
+	}
+	if (bCapped && !bLoop)
+	{
+		// if capped, set vertices.
+		for (int CapIdx = 0; CapIdx < 2; CapIdx++)
+		{
+			if (Caps[CapIdx] == ECapType::FlatMidpointFan)
+			{
+				Vertices[CapVertStart[CapIdx]] = Path[CapIdx * (Path.Num() - 1)];
+			}
+		}
+
+		for (int CapIdx = 0; CapIdx < 2; CapIdx++)
+		{
+			FVector3d Normal = CurveUtil::Tangent<double, FVector3d>(Path, CapIdx * (PathNum - 1), bLoop)* (double)(CapIdx * 2 - 1);
+			for (int SubIdx = 0; SubIdx < XNum; SubIdx++)
+			{
+				Normals[CapNormalStart[CapIdx] + SubIdx] = (FVector3f)Normal;
+			}
+		}
+	}
+
+	for (int k = 0; k < Normals.Num(); ++k)
+	{
+		Normalize(Normals[k]);
+	}
+
+	return *this;
 }
-int32 GetPolygonGroup(EProfileSweepPolygonGrouping PolygonGroupingMode, int32 SweepIndex, int32 ProfileIndex, int32 NumProfileSegments)
+
+// FProfileSweepGenerator
+
+// Various indexing utility functions. See the comment in Generate() for a description of vertex/triangle/uv layout.
+namespace UE::SweepGeneratorLocals
 {
-	int32 PolygonId = -1;
-	switch (PolygonGroupingMode)
-	{
-	case EProfileSweepPolygonGrouping::Single:
-		PolygonId = 0;
-		break;
-	case EProfileSweepPolygonGrouping::PerFace:
-		PolygonId = SweepIndex * NumProfileSegments + ProfileIndex;
-		break;
-	case EProfileSweepPolygonGrouping::PerProfileSegment:
-		PolygonId = ProfileIndex;
-		break;
-	case EProfileSweepPolygonGrouping::PerSweepSegment:
-		PolygonId = SweepIndex;
-		break;
-	}
-	return PolygonId;
+	static int32 GetVertIndex(bool VertIsWelded, int32 SweepIndex, int32 ProfileIndex, int32 NumWelded, int32 NumNonWelded, const TArray<int32>& VertPositionOffsets)
+	{
+		return VertIsWelded ?
+			VertPositionOffsets[ProfileIndex]
+			: NumWelded + SweepIndex * NumNonWelded + VertPositionOffsets[ProfileIndex];
+	}
+	static int32 GetTriangleIndex(int32 SweepIndex, int32 ProfileIndex, int32 NumTrisPerSweepSegment, const TArray<int32>& VertTriangleOffsets)
+	{
+		return (SweepIndex * NumTrisPerSweepSegment) + VertTriangleOffsets[ProfileIndex];
+	}
+	static int32 GetUvIndex(int32 SweepIndex, int32 ProfileIndex, int32 NumUvColumns)
+	{
+		return (SweepIndex * NumUvColumns) + ProfileIndex;
+	}
+	static int32 GetPolygonGroup(EProfileSweepPolygonGrouping PolygonGroupingMode, int32 SweepIndex, int32 ProfileIndex, int32 NumProfileSegments)
+	{
+		int32 PolygonId = -1;
+		switch (PolygonGroupingMode)
+		{
+		case EProfileSweepPolygonGrouping::Single:
+			PolygonId = 0;
+			break;
+		case EProfileSweepPolygonGrouping::PerFace:
+			PolygonId = SweepIndex * NumProfileSegments + ProfileIndex;
+			break;
+		case EProfileSweepPolygonGrouping::PerProfileSegment:
+			PolygonId = ProfileIndex;
+			break;
+		case EProfileSweepPolygonGrouping::PerSweepSegment:
+			PolygonId = SweepIndex;
+			break;
+		}
+		return PolygonId;
+	}
 }
 
 /**
@@ -88,6 +812,8 @@ void FProfileSweepGenerator::AdjustNormalsForTriangle(int32 TriIndex, int32 Firs
 
 FMeshShapeGenerator& FProfileSweepGenerator::Generate()
 {
+	using namespace UE::SweepGeneratorLocals;
+
 	if (Progress && Progress->Cancelled())
 	{
 		return *this;
@@ -437,6 +1163,8 @@ FMeshShapeGenerator& FProfileSweepGenerator::Generate()
 void FProfileSweepGenerator::InitializeUvBuffer(const TArray<int32>& VertPositionOffsets,
 	int32& NumUvRowsOut, int32& NumUvColumnsOut)
 {
+	using namespace UE::SweepGeneratorLocals;
+
 	// Convenience variables
 	int32 NumProfilePoints = ProfileCurve.Num();
 	int32 NumProfileSegments = bProfileCurveIsClosed ? NumProfilePoints : NumProfilePoints - 1;
diff --git a/Engine/Source/Runtime/GeometryCore/Public/Generators/SweepGenerator.h b/Engine/Source/Runtime/GeometryCore/Public/Generators/SweepGenerator.h
--- a/Engine/Source/Runtime/GeometryCore/Public/Generators/SweepGenerator.h
+++ b/Engine/Source/Runtime/GeometryCore/Public/Generators/SweepGenerator.h
@@ -37,7 +37,7 @@ namespace Geometry
 /**
  * ECapType indicates the type of cap to use on a sweep
  */
-enum class /*GEOMETRYCORE_API*/ ECapType
+enum class ECapType
 {
 	None = 0,
 	FlatTriangulation = 1,
@@ -45,7 +45,7 @@ enum class /*GEOMETRYCORE_API*/ ECapType
 	// TODO: Cone, other caps ...
 };
 
-class /*GEOMETRYCORE_API*/ FSweepGeneratorBase : public FMeshShapeGenerator
+class FSweepGeneratorBase : public FMeshShapeGenerator
 {
 public:
 	virtual ~FSweepGeneratorBase()
@@ -65,7 +65,7 @@ protected:
 	 * Shared logic for creating vertex buffers and triangulations across all sweep primitives
 	 * Note: Does not set vertex positions or normals; a separate call must do that.
 	 */
-	void ConstructMeshTopology(const FPolygon2d& CrossSection,
+	GEOMETRYCORE_API void ConstructMeshTopology(const FPolygon2d& CrossSection,
 							   const TArrayView<const int32>& UVSections,
 							   const TArrayView<const int32>& NormalSections,
 							   const TArrayView<const int32>& SharpNormalsAlongLength,
@@ -78,366 +78,13 @@ protected:
 							   FVector2f CapUVScale, 
 							   FVector2f CapUVOffset,
 							   const TArrayView<const float>& CustomCrossSectionTexCoord = {}, // if specified and valid, we use custom UVs instead of automatically generating them
-							   const TArrayView<const float>& CustomPathTexCoord = {})
-	{
-		// per cross section
-		const int32 XVerts = CrossSection.VertexCount();
-		const int32 XSegments = bProfileCurveIsClosed ? XVerts : XVerts - 1;
-		const int32 XNormals = XVerts + NormalSections.Num();
-		const int32 XUVs = XSegments + UVSections.Num() + 1;
-
-		float TotalPerimeter = 0.0f, TotalPathLength = 0.0f;
-		TArray<float> CrossSectionTexCoord, PathTexCoord;
-
-		// Compute texture coordinates along the cross section (U coordinates)
-		if (bProfileCurveIsClosed && CustomCrossSectionTexCoord.Num() >= XVerts + 1)
-		{
-			CrossSectionTexCoord = CustomCrossSectionTexCoord;
-		}
-		else if (bProfileCurveIsClosed && CustomCrossSectionTexCoord.Num() == XVerts)
-		{
-			CrossSectionTexCoord = CustomCrossSectionTexCoord;
-
-			// If the cross section curve is closed and we are missing texture coordinate for the 
-			// last element we wrap araound and use the coordinate of the first element
-			CrossSectionTexCoord.Add(CrossSectionTexCoord[0]);
-		}
-		else if (bProfileCurveIsClosed == false && CustomCrossSectionTexCoord.Num() >= XVerts) 
-		{
-			CrossSectionTexCoord = CustomCrossSectionTexCoord;
-		}
-		else if (bEvenlySpaceUVs)
-		{
-			CrossSectionTexCoord.Add(0.0f);
-			for (int Idx = 0; Idx < XSegments; Idx++)
-			{
-				float SegLen = float(Distance(CrossSection[Idx], CrossSection[(Idx + 1) % XSegments]));
-				TotalPerimeter += SegLen;
-				CrossSectionTexCoord.Add(TotalPerimeter);
-			}
-			TotalPerimeter = FMath::Max(TotalPerimeter, FMathf::ZeroTolerance);
-			for (int Idx = 0; Idx < CrossSectionTexCoord.Num(); Idx++)
-			{
-				CrossSectionTexCoord[Idx] /= TotalPerimeter;
-				CrossSectionTexCoord[Idx] = 1.0f - CrossSectionTexCoord[Idx];
-			}
-		}
-		else 
-		{	
-			for (int Idx = 0; Idx < XSegments; Idx++)
-			{
-				float U = float(Idx) / float(XSegments);
-				CrossSectionTexCoord.Add(1.0f - U);
-			}
-			CrossSectionTexCoord.Add(0.0f);
-		}
-		
-		// Compute texture coordinates along the path (V coordinates)
-		if (CustomPathTexCoord.Num() >= NumCrossSections)
-		{
-			PathTexCoord = CustomPathTexCoord;
-		}
-		else if (bLoop && CustomPathTexCoord.Num() == NumCrossSections - 1)
-		{
-			PathTexCoord = CustomPathTexCoord;
-			
-			// If the path curve is closed and we are missing texture coordinate for the 
-			// last element we wrap araound and use the coordinate of the first element
-			PathTexCoord.Add(PathTexCoord[0]);
-		}
-		else if (bEvenlySpaceUVs) 
-		{
-			PathTexCoord.Add(0.0f);
-			int NumPathSegs = bLoop ? Path.Num() : Path.Num() - 1;
-			for (int Idx = 0; Idx < NumPathSegs; Idx++)
-			{
-				float SegLen = float(Distance(Path[Idx], Path[(Idx + 1) % Path.Num()]));
-				TotalPathLength += SegLen;
-				PathTexCoord.Add(TotalPathLength);
-			}
-			TotalPathLength = FMath::Max(TotalPathLength, FMathf::ZeroTolerance);
-			for (int Idx = 0; Idx < PathTexCoord.Num(); Idx++)
-			{
-				PathTexCoord[Idx] /= TotalPathLength;
-				PathTexCoord[Idx] = 1.0f - PathTexCoord[Idx];
-			}
-		}
-		else 
-		{
-			for (int32 Idx = 0; Idx < NumCrossSections; Idx++)
-			{
-				float V = float(Idx) / float(NumCrossSections - 1);
-				PathTexCoord.Add(1.0f - V);
-			}
-		}
-		
-		int32 NumVerts = XVerts * NumCrossSections - (bLoop ? XVerts : 0);
-		int32 NumNormals = NumCrossSections > 1 ? (XNormals * NumCrossSections - (bLoop ? XNormals : 0)) : 0;
-		NumNormals += XNormals * SharpNormalsAlongLength.Num();
-		int32 NumUVs = NumCrossSections > 1 ? XUVs * NumCrossSections : 0;
-		int32 NumPolygons = (NumCrossSections - 1) * XSegments;
-		int32 NumTriangles = NumPolygons * 2;
-
-		TArray<FIndex3i> OutTriangles;
-
-		// doesn't make sense to have cap types if the sweep is a loop
-		ensure(!bLoop || (Caps[0] == ECapType::None && Caps[1] == ECapType::None));
-		
-		if (!bLoop)
-		{
-			for (int32 CapIdx = 0; !bLoop && CapIdx < 2; CapIdx++)
-			{
-				CapVertStart[CapIdx] = NumVerts;
-				CapNormalStart[CapIdx] = NumNormals;
-				CapUVStart[CapIdx] = NumUVs;
-				CapTriangleStart[CapIdx] = NumTriangles;
-				CapPolygonStart[CapIdx] = NumPolygons;
-
-				if (Caps[CapIdx] == ECapType::FlatTriangulation)
-				{
-					NumTriangles += XVerts - 2;
-					NumPolygons++;
-					NumUVs += XVerts;
-					NumNormals += XVerts;
-				}
-				else if (Caps[CapIdx] == ECapType::FlatMidpointFan)
-				{
-					NumTriangles += XSegments;
-					NumPolygons++;
-					NumUVs += XVerts + 1;
-					NumNormals += XVerts + 1;
-					NumVerts += 1;
-				}
-				// TODO: support more cap type; e.g.:
-				//else if (Caps[CapIdx] == ECapType::Cone)
-				//{
-				//	NumTriangles += XVerts;
-				//	NumPolygons += XVerts;
-				//	NumUVs += XVerts + 1;
-				//	NumNormals += XVerts * 2;
-				//	NumVerts += 1;
-				//}
-			}
-		}
-
-		SetBufferSizes(NumVerts, NumTriangles, NumUVs, NumNormals);
-
-		if (!bLoop)
-		{
-			for (int32 CapIdx = 0; CapIdx < 2; CapIdx++)
-			{
-				if (Caps[CapIdx] == ECapType::FlatTriangulation)
-				{
-					int32 VertOffset = CapIdx * (XVerts * (NumCrossSections - 1));
-
-					PolygonTriangulation::TriangulateSimplePolygon(CrossSection.GetVertices(), OutTriangles);
-					int32 TriIdx = CapTriangleStart[CapIdx];
-					int32 PolyIdx = CapPolygonStart[CapIdx];
-					for (const FIndex3i& Triangle : OutTriangles)
-					{
-						bool Flipped = CapIdx == 0;
-						SetTriangle(TriIdx,
-							Triangle.A + VertOffset, Triangle.B + VertOffset, Triangle.C + VertOffset,
-							Flipped);
-						SetTriangleUVs(TriIdx,
-							Triangle.A + CapUVStart[CapIdx],
-							Triangle.B + CapUVStart[CapIdx],
-							Triangle.C + CapUVStart[CapIdx],
-							Flipped);
-						SetTriangleNormals(TriIdx,
-							Triangle.A + CapNormalStart[CapIdx],
-							Triangle.B + CapNormalStart[CapIdx],
-							Triangle.C + CapNormalStart[CapIdx],
-							Flipped);
-						SetTrianglePolygon(TriIdx, PolyIdx);
-						TriIdx++;
-					}
-					float SideScale = float(2 * CapIdx - 1);
-					for (int32 Idx = 0; Idx < XVerts; Idx++)
-					{
-						FVector2f CenteredVert = (FVector2f)CrossSection.GetVertices()[Idx] * CapUVScale + CapUVOffset;
-						SetUV(CapUVStart[CapIdx] + Idx, FVector2f(CenteredVert.X * SideScale, CenteredVert.Y), VertOffset + Idx);
-
-						// correct normal to be filled by subclass
-						SetNormal(CapNormalStart[CapIdx] + Idx, FVector3f::Zero(), VertOffset + Idx);
-					}
-				}
-				else if (Caps[CapIdx] == ECapType::FlatMidpointFan)
-				{
-					int32 VertOffset = CapIdx * (XVerts * (NumCrossSections - 1));
-					int32 CapVertStartIdx = CapVertStart[CapIdx];
-					int32 TriIdx = CapTriangleStart[CapIdx];
-					int32 PolyIdx = CapPolygonStart[CapIdx];
-					for (int32 VertIdx = 0; VertIdx < XSegments; VertIdx++)
-					{
-						bool Flipped = CapIdx == 0;
-						SetTriangle(TriIdx,
-							VertOffset + VertIdx,
-							CapVertStartIdx,
-							VertOffset + (VertIdx + 1) % XVerts,
-							Flipped);
-						SetTriangleUVs(TriIdx,
-							CapUVStart[CapIdx] + VertIdx,
-							CapUVStart[CapIdx] + XVerts,
-							CapUVStart[CapIdx] + (VertIdx + 1) % XVerts,
-							Flipped);
-						SetTriangleNormals(TriIdx,
-							CapNormalStart[CapIdx] + VertIdx,
-							CapNormalStart[CapIdx] + XVerts,
-							CapNormalStart[CapIdx] + (VertIdx + 1) % XVerts,
-							Flipped);
-						SetTrianglePolygon(TriIdx, PolyIdx);
-						TriIdx++;
-					}
-
-					// Set cap midpoint UV & Normal
-					// (correct normal to be filled by subclass)
-					SetUV(CapUVStart[CapIdx] + XVerts, FVector2f::Zero() + CapUVOffset, CapVertStartIdx);
-					SetNormal(CapNormalStart[CapIdx] + XVerts, FVector3f::Zero(), CapVertStartIdx);
-
-					// Set cap profile UVs & Normal
-					for (int32 Idx = 0; Idx < XVerts; Idx++)
-					{
-						FVector2f CenteredVert = (FVector2f)CrossSection.GetVertices()[Idx] * CapUVScale + CapUVOffset;
-						SetUV(CapUVStart[CapIdx] + Idx, FVector2f(CenteredVert.X, CenteredVert.Y), VertOffset + Idx);
-						SetNormal(CapNormalStart[CapIdx] + Idx, FVector3f::Zero(), VertOffset + Idx);
-					}
-				}
-			}
-		}
-
-		// fill in UVs and triangles along length
-		int MinValidCrossSections = bLoop ? 3 : 2;
-		int CurFaceGroupIndex = NumPolygons;
-		if (NumCrossSections >= MinValidCrossSections)
-		{
-			int32 UVSection = 0, UVSubIdx = 0;
-
-			int CrossSectionsMod = NumCrossSections;
-			if (bLoop)
-			{
-				CrossSectionsMod--; // last cross section becomes the first
-			}
-			int NormalCrossSectionsMod = CrossSectionsMod + SharpNormalsAlongLength.Num();
-
-			int32 NumSections = UVSections.Num();
-			int32 NextDupVertIdx = UVSection < NumSections ? UVSections[UVSection] : -1;
-			for (int32 VertSubIdx = 0; VertSubIdx < XSegments; UVSubIdx++)
-			{
-				for (int32 XIdx = 0; XIdx < NumCrossSections; XIdx++)
-				{
-					FVector2f UV = FVector2f(CrossSectionTexCoord[VertSubIdx], PathTexCoord[XIdx]);
-					SetUV(XIdx * XUVs + UVSubIdx, UV * SectionsUVScale, (XIdx % CrossSectionsMod) * XVerts + VertSubIdx);
-				}
-
-				if (VertSubIdx == NextDupVertIdx)
-				{
-					NextDupVertIdx = UVSection < NumSections ? UVSections[UVSection] : -1;
-				}
-				else
-				{
-					for (int32 XIdx = 0; XIdx + 1 < NumCrossSections; XIdx++)
-					{
-						SetTriangleUVs(
-							XSegments * 2 * XIdx + 2 * VertSubIdx,
-							XIdx * XUVs + UVSubIdx,
-							XIdx * XUVs + UVSubIdx + 1,
-							(XIdx + 1) * XUVs + UVSubIdx, true);
-						SetTriangleUVs(
-							XSegments * 2 * XIdx + 2 * VertSubIdx + 1,
-							(XIdx + 1) * XUVs + UVSubIdx + 1,
-							(XIdx + 1) * XUVs + UVSubIdx,
-							XIdx * XUVs + UVSubIdx + 1, true);
-					}
-					VertSubIdx++;
-				}
-			}
-			{
-				// final UV
-				int32 VertSubIdx = bProfileCurveIsClosed ? 0 : XSegments;
-				for (int32 XIdx = 0; XIdx < NumCrossSections; XIdx++)
-				{
-					FVector2f UV = FVector2f(CrossSectionTexCoord.Last(), PathTexCoord[XIdx]);
-					SetUV(XIdx * XUVs + UVSubIdx, UV * SectionsUVScale, (XIdx % CrossSectionsMod) * XVerts + VertSubIdx);
-				}
-			}
-			NumSections = NormalSections.Num();
-			int32 NormalSection = 0;
-			NextDupVertIdx = NormalSection < NumSections ? NormalSections[NormalSection] : -1;
-			check(NextDupVertIdx < XVerts);
-			for (int32 VertSubIdx = 0, NormalSubIdx = 0; VertSubIdx < XVerts; NormalSubIdx++)
-			{
-				int SharpNormalIdx = 0;
-				for (int32 XIdx = 0, NormalXIdx = 0; XIdx < NumCrossSections; XIdx++, NormalXIdx++)
-				{
-					// just set the normal parent; don't compute normal yet
-					SetNormal((NormalXIdx % NormalCrossSectionsMod) * XNormals + NormalSubIdx, FVector3f(0, 0, 0), (XIdx % CrossSectionsMod) * XVerts + VertSubIdx);
-					// duplicate normals for cross sections that are 'sharp'
-					if (SharpNormalIdx < SharpNormalsAlongLength.Num() && XIdx == SharpNormalsAlongLength[SharpNormalIdx])
-					{
-						NormalXIdx++;
-						SetNormal((NormalXIdx % NormalCrossSectionsMod) * XNormals + NormalSubIdx, FVector3f(0, 0, 0), (XIdx % CrossSectionsMod) * XVerts + VertSubIdx);
-						SharpNormalIdx++;
-					}
-				}
-
-				if (VertSubIdx == NextDupVertIdx)
-				{
-					NextDupVertIdx = NormalSection < NumSections ? NormalSections[NormalSection] : -1;
-					check(NextDupVertIdx < XVerts);
-				}
-				else
-				{
-					if (VertSubIdx < XSegments) // if bProfileCurveIsClosed == false skip the last triangle strip generation
-					{
-						int32 WrappedNextNormalSubIdx = (NormalSubIdx + 1) % XNormals;
-						int32 WrappedNextVertexSubIdx = (VertSubIdx + 1) % XVerts;
-						SharpNormalIdx = 0;
-						for (int32 XIdx = 0, NXIdx = 0; XIdx + 1 < NumCrossSections; XIdx++, NXIdx++)
-						{
-							int32 T0Idx = XSegments * 2 * XIdx + 2 * VertSubIdx;
-							int32 T1Idx = T0Idx + 1;
-							int32 PIdx = XSegments * XIdx + VertSubIdx;
-							int32 NextXIdx = (XIdx + 1) % CrossSectionsMod;
-							int32 NextNXIdx = (NXIdx + 1) % NormalCrossSectionsMod;
-							SetTrianglePolygon(T0Idx, (bPolygroupPerQuad) ? PIdx : (CurFaceGroupIndex + XIdx));
-							SetTrianglePolygon(T1Idx, (bPolygroupPerQuad) ? PIdx : (CurFaceGroupIndex + XIdx));
-							SetTriangle(T0Idx,
-								XIdx * XVerts + VertSubIdx,
-								XIdx * XVerts + WrappedNextVertexSubIdx,
-								NextXIdx * XVerts + VertSubIdx, true);
-							SetTriangle(T1Idx,
-								NextXIdx * XVerts + WrappedNextVertexSubIdx,
-								NextXIdx * XVerts + VertSubIdx,
-								XIdx * XVerts + WrappedNextVertexSubIdx, true);
-							SetTriangleNormals(
-								T0Idx,
-								NXIdx * XNormals + NormalSubIdx,
-								NXIdx * XNormals + WrappedNextNormalSubIdx,
-								NextNXIdx * XNormals + NormalSubIdx, true);
-							SetTriangleNormals(
-								T1Idx,
-								NextNXIdx * XNormals + WrappedNextNormalSubIdx,
-								NextNXIdx * XNormals + NormalSubIdx,
-								NXIdx * XNormals + WrappedNextNormalSubIdx, true);
-							if (SharpNormalIdx < SharpNormalsAlongLength.Num() && XIdx + 1 == SharpNormalsAlongLength[SharpNormalIdx])
-							{
-								NXIdx++;
-								SharpNormalIdx++;
-							}
-						}
-					}
-					VertSubIdx++;
-				}
-			}
-		}
-	}
+							   const TArrayView<const float>& CustomPathTexCoord = {});
 };
 
 /**
 * Generate a cylinder with optional end caps
 */
-class /*GEOMETRYCORE_API*/ FVerticalCylinderGeneratorBase : public FSweepGeneratorBase
+class FVerticalCylinderGeneratorBase : public FSweepGeneratorBase
 {
 public:
 	int AngleSamples = 16;
@@ -445,166 +92,15 @@ public:
 	bool bUVScaleMatchSidesAndCaps = true;
 	ECapType CapType = ECapType::FlatMidpointFan;
 
-	static float ComputeSegLengths(const TArrayView<float>& Radii, const TArrayView<float>& Heights, TArray<float>& AlongPercents)
-	{
-		float LenAlong = 0;
-		int32 NumX = Radii.Num();
-		AlongPercents.SetNum(NumX);
-		AlongPercents[0] = 0;
-		for (int XIdx = 0; XIdx+1 < NumX; XIdx++)
-		{
-			double Dist = Distance( FVector2d(Radii[XIdx], Heights[XIdx]), FVector2d(Radii[XIdx + 1], Heights[XIdx + 1]) );
-			LenAlong += float(Dist);
-			AlongPercents[XIdx + 1] = LenAlong;
-		}
-		for (int XIdx = 0; XIdx+1 < NumX; XIdx++)
-		{
-			AlongPercents[XIdx+1] /= LenAlong;
-		}
-		return LenAlong;
-	}
-
-	bool GenerateVerticalCircleSweep(const TArrayView<float>& Radii, const TArrayView<float>& Heights, const TArrayView<int>& SharpNormalsAlongLength)
-	{
-		FPolygon2d X = FPolygon2d::MakeCircle(1.0, AngleSamples);
-		const TArray<FVector2d>& XVerts = X.GetVertices();
-		ECapType Caps[2] = {ECapType::None, ECapType::None};
+	GEOMETRYCORE_API static float ComputeSegLengths(const TArrayView<float>& Radii, const TArrayView<float>& Heights, TArray<float>& AlongPercents);;
 
-		if (bCapped)
-		{
-			Caps[0] = CapType;
-			Caps[1] = CapType;
-		}
-
-		int NumX = Radii.Num();
-		if (!ensure(NumX == Heights.Num()))
-		{
-			return false;
-		}
-		// first and last cross sections can't be sharp, so can't have more than NumX-2 sharp normal indices
-		if (!ensure(SharpNormalsAlongLength.Num() + 2 <= NumX))
-		{
-			return false;
-		}
-
-		TArray<float> AlongPercents;
-		float LenAlong = ComputeSegLengths(Radii, Heights, AlongPercents);
-
-		ConstructMeshTopology(X, {}, {}, SharpNormalsAlongLength, false, {}, NumX, false, Caps, FVector2f(1, 1), FVector2f(.5f, .5f), FVector2f(.5f, .5f));
-
-		TArray<FVector2d> NormalSides; NormalSides.SetNum(NumX - 1);
-		for (int XIdx = 0; XIdx+1 < NumX; XIdx++)
-		{
-			FVector2d Vec = FVector2d(Radii[XIdx + 1], Heights[XIdx + 1]) - FVector2d(Radii[XIdx], Heights[XIdx]);
-			NormalSides[XIdx] = Normalized(PerpCW(Vec));
-		}
-		TArray<FVector2d> SmoothedNormalSides; SmoothedNormalSides.SetNum(NumX);
-		// smooth internal normals
-		SmoothedNormalSides[0] = NormalSides[0];
-		SmoothedNormalSides.Last() = NormalSides.Last();
-		for (int XIdx = 1; XIdx + 1 < NumX; XIdx++)
-		{
-			SmoothedNormalSides[XIdx] = Normalized(NormalSides[XIdx] + NormalSides[XIdx - 1]);
-		}
-
-
-		// set vertex positions and normals for all cross sections along length
-		for (int SubIdx = 0; SubIdx < X.VertexCount(); SubIdx++)
-		{
-			int SharpNormalIdx = 0;
-			for (int XIdx = 0, NormalXIdx = 0; XIdx < NumX; ++XIdx, ++NormalXIdx)
-			{
-				double AlongRadius = Radii[XIdx];
-				Vertices[SubIdx + XIdx * AngleSamples] =
-					FVector3d(XVerts[SubIdx].X * AlongRadius, XVerts[SubIdx].Y * AlongRadius, Heights[XIdx]);
-				if (SharpNormalIdx < SharpNormalsAlongLength.Num() && XIdx == SharpNormalsAlongLength[SharpNormalIdx])
-				{
-					// write sharp normals
-					if (ensure(XIdx > 0)) // very first index cannot be sharp
-					{
-						Normals[SubIdx + NormalXIdx * AngleSamples] = FVector3f(float(XVerts[SubIdx].X*NormalSides[XIdx-1].X), float(XVerts[SubIdx].Y*NormalSides[XIdx-1].X), float(NormalSides[XIdx-1].Y));
-					}
-					NormalXIdx++;
-					if (ensure(XIdx + 1 < NumX)) // very last index cannot be sharp
-					{
-						Normals[SubIdx + NormalXIdx * AngleSamples] = FVector3f(float(XVerts[SubIdx].X*NormalSides[XIdx].X), float(XVerts[SubIdx].Y*NormalSides[XIdx].X), float(NormalSides[XIdx].Y));
-					}
-					SharpNormalIdx++;
-				}
-				else
-				{
-					// write smoothed normal
-					Normals[SubIdx + NormalXIdx * AngleSamples] = FVector3f(float(XVerts[SubIdx].X*SmoothedNormalSides[XIdx].X), float(XVerts[SubIdx].Y*SmoothedNormalSides[XIdx].X), float(SmoothedNormalSides[XIdx].Y));
-				}
-			}
-		}
-		
-		if (bCapped)
-		{
-			// if capped, set vertices.
-			for (int CapIdx = 0; CapIdx < 2; CapIdx++)
-			{
-				if (Caps[CapIdx] == ECapType::FlatMidpointFan)
-				{
-					Vertices[CapVertStart[CapIdx]] = FVector3d::UnitZ() * (double)Heights[CapIdx * (Heights.Num() - 1)];
-				}
-			}
-
-			// if capped, set top/bottom normals
-			for (int SubIdx = 0; SubIdx < X.VertexCount(); SubIdx++)
-			{
-				for (int XBotTop = 0; XBotTop < 2; ++XBotTop)
-				{
-					Normals[CapNormalStart[XBotTop] + SubIdx] = FVector3f(0.f, 0.f, float(2 * XBotTop - 1));
-				}
-			}
-			for (int CapIdx = 0; CapIdx < 2; CapIdx++)
-			{
-				if (Caps[CapIdx] == ECapType::FlatMidpointFan)
-				{
-					Normals[CapNormalStart[CapIdx] + X.VertexCount()] = FVector3f(0.f, 0.f, float(2 * CapIdx - 1));
-				}
-			}
-		}
-
-		if (bUVScaleMatchSidesAndCaps)
-		{
-			float MaxAbsRad = FMathf::Abs(Radii[0]);
-			for (int XIdx = 0; XIdx < NumX; XIdx++)
-			{
-				MaxAbsRad = FMathf::Max(FMathf::Abs(Radii[XIdx]), MaxAbsRad);
-			}
-			float AbsHeight = LenAlong;
-			float MaxAbsCircumference = MaxAbsRad * FMathf::TwoPi;
-
-			// scales to put each differently-scaled UV coordinate into the same space
-			float ThetaScale = MaxAbsCircumference;
-			float HeightScale = AbsHeight;
-			float CapScale = MaxAbsRad*2;
-
-			float MaxScale = FMathf::Max3(ThetaScale, HeightScale, CapScale);
-			ThetaScale /= MaxScale;
-			HeightScale /= MaxScale;
-			CapScale /= MaxScale;
-			for (int UVIdx = 0; UVIdx < CapUVStart[0]; UVIdx++)
-			{
-				UVs[UVIdx].X *= ThetaScale;
-				UVs[UVIdx].Y *= HeightScale;
-			}
-			for (int UVIdx = CapUVStart[0]; UVIdx < UVs.Num(); UVIdx++)
-			{
-				UVs[UVIdx] *= CapScale;
-			}
-		}
-
-		return true;
-	}
+	GEOMETRYCORE_API bool GenerateVerticalCircleSweep(const TArrayView<float>& Radii, const TArrayView<float>& Heights, const TArrayView<int>& SharpNormalsAlongLength);
 };
 
 /**
  * Generate a cylinder with optional end caps
  */
-class /*GEOMETRYCORE_API*/ FCylinderGenerator : public FVerticalCylinderGeneratorBase
+class FCylinderGenerator : public FVerticalCylinderGeneratorBase
 {
 public:
 	float Radius[2] = {1.0f, 1.0f};
@@ -613,36 +109,13 @@ public:
 
 public:
 	/** Generate the mesh */
-	virtual FMeshShapeGenerator& Generate() override
-	{
-		if (bCapped)
-		{
-			AngleSamples = FMath::Max(AngleSamples, 3);
-		}
-
-		TArray<float> Radii, Heights;
-
-		Radii.Add(Radius[0]);
-		Heights.Add(0);
-		for (int ExtraIdx = 0; ExtraIdx < LengthSamples; ExtraIdx++)
-		{
-			float Along = float(ExtraIdx + 1) / float(LengthSamples + 1);
-			Radii.Add(FMath::Lerp(Radius[0], Radius[1], Along));
-			Heights.Add(Height * Along);
-		}
-		Radii.Add(Radius[1]);
-		Heights.Add(Height);
-
-		GenerateVerticalCircleSweep(Radii, Heights, {});
-
-		return *this;
-	}
+	GEOMETRYCORE_API virtual FMeshShapeGenerator& Generate() override;
 };
 
 /**
 * Generate a 3D arrow
 */
-class /*GEOMETRYCORE_API*/ FArrowGenerator : public FVerticalCylinderGeneratorBase
+class FArrowGenerator : public FVerticalCylinderGeneratorBase
 {
 public:
 	float StickRadius = 0.5f;
@@ -667,47 +140,7 @@ public:
 
 public:
 	/** Generate the mesh */
-	virtual FMeshShapeGenerator& Generate() override
-	{
-		const float SrcRadii[]{StickRadius, StickRadius, HeadBaseRadius, HeadTipRadius};
-		const float SrcHeights[]{0, StickLength, StickLength, StickLength + HeadLength};
-
-		TArray<float> Radii, Heights;
-		const int NumVerts = 4 + AdditionalLengthSamples[0] + AdditionalLengthSamples[1] + AdditionalLengthSamples[2];
-		Radii.SetNumUninitialized(NumVerts);
-		Heights.SetNumUninitialized(NumVerts);
-
-		int VertIdx = 0;
-		auto SetVert = [&Radii, &Heights, &VertIdx](float Radius, float Height)
-		{
-			Radii[VertIdx] = Radius;
-			Heights[VertIdx] = Height;
-			++VertIdx;
-		};
-
-		for (int SegIdx = 0; true; ++SegIdx)
-		{
-			SetVert(SrcRadii[SegIdx], SrcHeights[SegIdx]);
-
-			if (SegIdx == 3)
-			{
-				break;
-			}
-
-			for (int ExtraSeg = 1, NumExtraSegs = AdditionalLengthSamples[SegIdx] + 1; ExtraSeg < NumExtraSegs; ++ExtraSeg)
-			{
-				const float Along = float(ExtraSeg) / float(NumExtraSegs);
-				SetVert(FMath::Lerp(SrcRadii[SegIdx], SrcRadii[SegIdx + 1], Along),
-						FMath::Lerp(SrcHeights[SegIdx], SrcHeights[SegIdx + 1], Along));
-			}
-		}
-
-		TArray<int> SharpNormalsAlongLength{1 + AdditionalLengthSamples[0], 2 + AdditionalLengthSamples[0] + AdditionalLengthSamples[1]};
-
-		GenerateVerticalCircleSweep(Radii, Heights, SharpNormalsAlongLength);
-
-		return *this;
-	}
+	GEOMETRYCORE_API virtual FMeshShapeGenerator& Generate() override;
 };
 
 
@@ -718,7 +151,7 @@ public:
  * TODO: 
  *  - a custom variant for toruses specifically (would be faster)
  */
-class /*GEOMETRYCORE_API*/ FGeneralizedCylinderGenerator : public FSweepGeneratorBase
+class FGeneralizedCylinderGenerator : public FSweepGeneratorBase
 {
 public:
 	FPolygon2d CrossSection;
@@ -798,138 +231,7 @@ public:
 
 public:
 	/** Generate the mesh */
-	virtual FMeshShapeGenerator& Generate() override
-	{
-		const TArray<FVector2d>& XVerts = CrossSection.GetVertices();
-		ECapType Caps[2] = {ECapType::None, ECapType::None};
-
-		if (bCapped && !bLoop)
-		{
-			Caps[0] = CapType;
-			Caps[1] = CapType;
-		}
-		int PathNum = Path.Num();
-		
-		bool bHavePathScaling = (PathScales.Num() == PathNum);
-		bool bApplyScaling = (bHavePathScaling|| (StartScale != 1.0) || (EndScale != 1.0)) && (bLoop == false);
-		bool bNeedArcLength = (bApplyScaling || bUVScaleRelativeWorld);
-		double TotalPathArcLength = (bNeedArcLength) ? UE::Geometry::CurveUtil::ArcLength<double, FVector3d>(Path, bLoop) : 1.0;
-
-		FAxisAlignedBox2f Bounds = (FAxisAlignedBox2f)CrossSection.Bounds();
-		double BoundsMaxDimInv = 1.0 / FMathd::Max(Bounds.MaxDim(), .001);
-		FVector2f SectionScale(1.f, 1.f), CapScale((float)BoundsMaxDimInv, (float)BoundsMaxDimInv);
-		if (bUVScaleRelativeWorld)
-		{
-			double Perimeter = CrossSection.Perimeter();
-			SectionScale.X = float( Perimeter / UnitUVInWorldCoordinates );
-			SectionScale.Y = float( TotalPathArcLength / UnitUVInWorldCoordinates);
-			CapScale.X = CapScale.Y = 1.0f / UnitUVInWorldCoordinates;
-		}
-		ConstructMeshTopology(CrossSection, {}, {}, {}, true, Path, PathNum + (bLoop ? 1 : 0), bLoop, Caps, SectionScale, CapScale, Bounds.Center(), CrossSectionTexCoord, PathTexCoord);
-
-		int XNum = CrossSection.VertexCount();
-		TArray<FVector2d> XNormals; XNormals.SetNum(XNum);
-		for (int Idx = 0; Idx < XNum; Idx++)
-		{
-			XNormals[Idx] = CrossSection.GetNormal_FaceAvg(Idx);
-		}
-
-		double AccumArcLength = 0;
-		FFrame3d CrossSectionFrame = InitialFrame;
-		bool bHaveExplicitFrames = (PathFrames.Num() == Path.Num());
-		for (int PathIdx = 0; PathIdx < PathNum; ++PathIdx)
-		{
-			FVector3d C = Path[PathIdx];
-			FVector3d X, Y;
-			FMatrix2d MiterPosScale = FMatrix2d::Identity();
-			if (bHaveExplicitFrames)
-			{
-				CrossSectionFrame = PathFrames[PathIdx];
-				C = PathFrames[PathIdx].Origin;
-			}
-			if (!bHaveExplicitFrames || bAlignFramesToSampledTangents)
-			{
-				FVector3d ToPrev, ToNext;
-				UE::Geometry::CurveUtil::GetVectorsToPrevNext<double, FVector3d>(Path, PathIdx, ToPrev, ToNext, true, bLoop);
-
-				FVector3d Tangent = Normalized(-ToPrev + ToNext);
-				CrossSectionFrame.AlignAxis(2, Tangent);
-				if (MiterLimit > 1)
-				{
-					// only miter if neither neighbor vector is exactly zero
-					if (!ToPrev.IsZero() && !ToNext.IsZero())
-					{
-						FVector3d TurnDir = ToPrev + ToNext;
-						FVector2d ScaleAlong = (FVector2d)CrossSectionFrame.ToFrameVector(TurnDir);
-						double ScaleAmount = 1.0;
-						if (ScaleAlong.Normalize())
-						{
-							double CosTwoTheta = ToPrev.Dot(ToNext);
-							double SqrScale = 2.0 / (1.0 - CosTwoTheta);
-							if (!FMath::IsFinite(SqrScale) || SqrScale > MiterLimit * MiterLimit)
-							{
-								ScaleAmount = MiterLimit;
-							}
-							else
-							{
-								ScaleAmount = FMath::Sqrt(SqrScale);
-							}
-							MiterPosScale = FMatrix2d::AxisScale(ScaleAlong, ScaleAmount, false /* axis was already normalized */);
-						}
-					}
-				}
-			}
-			X = CrossSectionFrame.X();
-			Y = CrossSectionFrame.Y();
-
-			double T = FMathd::Clamp((AccumArcLength / TotalPathArcLength), 0.0, 1.0);
-			double UniformScale = (bApplyScaling) ? FMathd::Lerp(StartScale, EndScale, T) : 1.0;
-			FVector2d PathScaling = (bHavePathScaling) ? PathScales[PathIdx] : FVector2d::One();
-
-			for (int SubIdx = 0; SubIdx < XNum; SubIdx++)
-			{
-				FVector2d XP = MiterPosScale * (UniformScale * PathScaling * CrossSection[SubIdx]);
-				FVector2d XN = XNormals[SubIdx];
-				// Note: Arguably, one could apply an inverse mitering scale to the normals, and then re-normalize, here ...
-				// However, it's better to think about these normals an average from the two unscaled cross sections that are meeting at the cross section plane,
-				// and in that model they won't be affected by the miter scale
-				Vertices[SubIdx + PathIdx * XNum] = C + X * XP.X + Y * XP.Y;
-				Normals[SubIdx + PathIdx * XNum] = (FVector3f)(X * XN.X + Y * XN.Y);
-			}
-
-			if (PathIdx < PathNum - 1)
-			{
-				AccumArcLength += Distance(C, Path[PathIdx + 1]);
-			}
-		}
-		if (bCapped && !bLoop)
-		{
-			// if capped, set vertices.
-			for (int CapIdx = 0; CapIdx < 2; CapIdx++)
-			{
-				if (Caps[CapIdx] == ECapType::FlatMidpointFan)
-				{
-					Vertices[CapVertStart[CapIdx]] = Path[CapIdx * (Path.Num() - 1)];
-				}
-			}
-
-			for (int CapIdx = 0; CapIdx < 2; CapIdx++)
-			{
-				FVector3d Normal = CurveUtil::Tangent<double, FVector3d>(Path, CapIdx * (PathNum - 1), bLoop) * (double)(CapIdx * 2 - 1);
-				for (int SubIdx = 0; SubIdx < XNum; SubIdx++)
-				{
-					Normals[CapNormalStart[CapIdx] + SubIdx] = (FVector3f)Normal;
-				}
-			}
-		}
-
-		for (int k = 0; k < Normals.Num(); ++k)
-		{
-			Normalize(Normals[k]);
-		}
-
-		return *this;
-	}
+	GEOMETRYCORE_API virtual FMeshShapeGenerator& Generate() override;
 };
 
 enum class EProfileSweepPolygonGrouping : uint8
