clean up geometrycore sweepgenerator includes

diff --git a/Engine/Plugins/Runtime/GeometryScripting/Source/GeometryScriptingCore/Private/MeshPrimitiveFunctions.cpp b/Engine/Plugins/Runtime/GeometryScripting/Source/GeometryScriptingCore/Private/MeshPrimitiveFunctions.cpp
--- a/Engine/Plugins/Runtime/GeometryScripting/Source/GeometryScriptingCore/Private/MeshPrimitiveFunctions.cpp
+++ b/Engine/Plugins/Runtime/GeometryScripting/Source/GeometryScriptingCore/Private/MeshPrimitiveFunctions.cpp
@@ -22,6 +22,7 @@
 #include "Arrangement2d.h"
 #include "CompGeom/Delaunay2.h"
 #include "CompGeom/ConvexDecomposition3.h"
+#include "CompGeom/PolygonTriangulation.h"
 
 #include "Physics/ComponentCollisionUtil.h"
 
diff --git a/Engine/Plugins/Runtime/MeshModelingToolset/Source/ModelingOperatorsEditorOnly/Private/CuttingOps/EmbedPolygonsOp.cpp b/Engine/Plugins/Runtime/MeshModelingToolset/Source/ModelingOperatorsEditorOnly/Private/CuttingOps/EmbedPolygonsOp.cpp
--- a/Engine/Plugins/Runtime/MeshModelingToolset/Source/ModelingOperatorsEditorOnly/Private/CuttingOps/EmbedPolygonsOp.cpp
+++ b/Engine/Plugins/Runtime/MeshModelingToolset/Source/ModelingOperatorsEditorOnly/Private/CuttingOps/EmbedPolygonsOp.cpp
@@ -13,6 +13,7 @@
 #include "MeshSimplification.h"
 #include "Operations/MeshPlaneCut.h"
 #include "Operations/MeshBoolean.h"
+#include "CompGeom/PolygonTriangulation.h"
 #include "Curve/GeneralPolygon2.h"
 #include "Curve/PolygonIntersectionUtils.h"
 
diff --git a/Engine/Source/Runtime/GeometryCore/Private/Generators/SweepGenerator.cpp b/Engine/Source/Runtime/GeometryCore/Private/Generators/SweepGenerator.cpp
--- a/Engine/Source/Runtime/GeometryCore/Private/Generators/SweepGenerator.cpp
+++ b/Engine/Source/Runtime/GeometryCore/Private/Generators/SweepGenerator.cpp
@@ -3,6 +3,10 @@
 #include "Generators/SweepGenerator.h"
 
 #include "Async/ParallelFor.h"
+#include "CompGeom/PolygonTriangulation.h"
+#include "Curve/CurveUtil.h"
+#include "MathUtil.h"
+#include "Misc/AssertionMacros.h"
 
 using namespace UE::Geometry;
 
diff --git a/Engine/Source/Runtime/GeometryCore/Public/Generators/SweepGenerator.h b/Engine/Source/Runtime/GeometryCore/Public/Generators/SweepGenerator.h
--- a/Engine/Source/Runtime/GeometryCore/Public/Generators/SweepGenerator.h
+++ b/Engine/Source/Runtime/GeometryCore/Public/Generators/SweepGenerator.h
@@ -4,26 +4,15 @@
 
 #pragma once
 
-#include "BoxTypes.h"
-#include "CompGeom/PolygonTriangulation.h"
+
 #include "Containers/Array.h"
 #include "Containers/ArrayView.h"
 #include "Containers/Set.h"
-#include "CoreMinimal.h"
-#include "Curve/CurveUtil.h"
 #include "FrameTypes.h"
-#include "HAL/PlatformCrt.h"
-#include "IndexTypes.h"
-#include "Math/UnrealMathSSE.h"
-#include "Math/UnrealMathUtility.h"
 #include "Math/Vector.h"
 #include "Math/Vector2D.h"
-#include "MathUtil.h"
-#include "MatrixTypes.h"
 #include "MeshShapeGenerator.h"
-#include "Misc/AssertionMacros.h"
 #include "Polygon2.h"
-#include "Templates/UnrealTemplate.h"
 #include "Util/ProgressCancel.h"
 #include "VectorTypes.h"
 
@@ -354,3 +343,18 @@ protected:
 
 } // end namespace UE::Geometry
 } // end namespace UE
+
+#if UE_ENABLE_INCLUDE_ORDER_DEPRECATED_IN_5_7
+#include "BoxTypes.h"
+#include "IndexTypes.h"
+#include "MatrixTypes.h"
+#include "Math/UnrealMathSSE.h"
+#include "Math/UnrealMathUtility.h"
+#include "Templates/UnrealTemplate.h"
+#include "HAL/PlatformCrt.h"
+#include "CoreMinimal.h"
+#include "CompGeom/PolygonTriangulation.h"
+#include "Curve/CurveUtil.h"
+#include "MathUtil.h"
+#include "Misc/AssertionMacros.h"
+#endif
