fix crash on split overlay bowties of non-compact mesh

diff --git a/Engine/Source/Runtime/GeometryCore/Private/DynamicMesh/DynamicMeshOverlay.cpp b/Engine/Source/Runtime/GeometryCore/Private/DynamicMesh/DynamicMeshOverlay.cpp
--- a/Engine/Source/Runtime/GeometryCore/Private/DynamicMesh/DynamicMeshOverlay.cpp
+++ b/Engine/Source/Runtime/GeometryCore/Private/DynamicMesh/DynamicMeshOverlay.cpp
@@ -307,7 +307,7 @@ void TDynamicMeshOverlay<RealType, ElementSize>::SplitBowties(bool bParallel)
 		FMutex BowtieCS;
 		ParallelFor(ParentMesh->MaxVertexID(), [this, &BowtieVerts, &BowtieCS](int32 VertexID)
 			{
-				if (IsBowtieInOverlay(VertexID))
+				if (ParentMesh->IsVertex(VertexID) && IsBowtieInOverlay(VertexID))
 				{
 					BowtieCS.Lock();
 					BowtieVerts.Add(VertexID);
