Add missing include
#rnx

diff --git a/Engine/Source/Runtime/GeometryCore/Private/Generators/RevolveGenerator.cpp b/Engine/Source/Runtime/GeometryCore/Private/Generators/RevolveGenerator.cpp
--- a/Engine/Source/Runtime/GeometryCore/Private/Generators/RevolveGenerator.cpp
+++ b/Engine/Source/Runtime/GeometryCore/Private/Generators/RevolveGenerator.cpp
@@ -3,6 +3,7 @@
 #include "Generators/RevolveGenerator.h"
 #include "Util/RevolveUtil.h"
 #include "Async/ParallelFor.h"
+#include "CompGeom/PolygonTriangulation.h"
 #include "DynamicMesh/DynamicMeshAttributeSet.h"
 #include "DynamicMesh/MeshNormals.h"
 
