Add name-matching weight channel transfer/matching support, for carrying weights through booleans and similar use cases

diff --git a/Engine/Source/Runtime/GeometryCore/Private/DynamicMesh/DynamicMeshAttributeSet.cpp b/Engine/Source/Runtime/GeometryCore/Private/DynamicMesh/DynamicMeshAttributeSet.cpp
--- a/Engine/Source/Runtime/GeometryCore/Private/DynamicMesh/DynamicMeshAttributeSet.cpp
+++ b/Engine/Source/Runtime/GeometryCore/Private/DynamicMesh/DynamicMeshAttributeSet.cpp
@@ -769,6 +769,62 @@ const FDynamicMeshWeightAttribute* FDynamicMeshAttributeSet::GetWeightLayer(int
 	return &WeightLayers[Index];
 }
 
+void FDynamicMeshAttributeSet::EnableMatchingWeightLayersByNames(const FDynamicMeshAttributeSet* ToMatch, bool bDiscardUnmatched)
+{
+	if (!ToMatch)
+	{
+		if (bDiscardUnmatched)
+		{
+			SetNumWeightLayers(0);
+		}
+		return;
+	}
+
+	for (int32 MatchLayerIdx = 0; MatchLayerIdx < ToMatch->WeightLayers.Num(); ++MatchLayerIdx)
+	{
+		const FDynamicMeshWeightAttribute& MatchLayer = ToMatch->WeightLayers[MatchLayerIdx];
+		// Forward-search for a matching layer, and if found swap it back as needed
+		bool bFound = false;
+		for (int32 LayerIdx = MatchLayerIdx; LayerIdx < WeightLayers.Num(); ++LayerIdx)
+		{
+			FDynamicMeshWeightAttribute& Layer = WeightLayers[LayerIdx];
+			if (MatchLayer.GetName() == Layer.GetName())
+			{
+				bFound = true;
+				if (LayerIdx > MatchLayerIdx)
+				{
+					Swap(WeightLayers[LayerIdx], WeightLayers[MatchLayerIdx]);
+				}
+				break;
+			}
+		}
+		if (!bFound)
+		{
+			if (MatchLayerIdx >= WeightLayers.Num())
+			{
+				// Add all remaining reference layers
+				SetNumWeightLayers(ToMatch->WeightLayers.Num());
+				for (int32 Idx = MatchLayerIdx; Idx < ToMatch->WeightLayers.Num(); ++Idx)
+				{
+					WeightLayers[Idx].SetName(ToMatch->WeightLayers[Idx].GetName());
+				}
+				break;
+			}
+			else
+			{
+				// Insert the missing layer
+				FDynamicMeshWeightAttribute* ToInsert = new FDynamicMeshWeightAttribute(ParentMesh);
+				ToInsert->SetName(MatchLayer.GetName());
+				WeightLayers.Insert(ToInsert, MatchLayerIdx);
+			}
+		}
+	}
+	if (bDiscardUnmatched)
+	{
+		SetNumWeightLayers(ToMatch->WeightLayers.Num());
+	}
+}
+
 
 void FDynamicMeshAttributeSet::EnableMaterialID()
 {
diff --git a/Engine/Source/Runtime/GeometryCore/Public/DynamicMesh/DynamicMeshAttributeSet.h b/Engine/Source/Runtime/GeometryCore/Public/DynamicMesh/DynamicMeshAttributeSet.h
--- a/Engine/Source/Runtime/GeometryCore/Public/DynamicMesh/DynamicMeshAttributeSet.h
+++ b/Engine/Source/Runtime/GeometryCore/Public/DynamicMesh/DynamicMeshAttributeSet.h
@@ -334,6 +334,14 @@ public:
 	/** @return the weight layer at the given Index */
 	GEOMETRYCORE_API const FDynamicMeshWeightAttribute* GetWeightLayer(int Index) const;
 
+	/**
+	 * Make the weight layers 1:1 with those of the ToMatch attribute set, 
+	 * adding any missing layers and re-ordering existing layers to match the reference name ordering 
+	 * @param ToMatch Reference attribute set whose weight layer names / ordering we must match
+	 * @param bDiscardUnmatched Whether to remove unmatched weight layers. Otherwise, they will be left at the end of the weight attribute array.
+	 */
+	GEOMETRYCORE_API void EnableMatchingWeightLayersByNames(const FDynamicMeshAttributeSet* ToMatch, bool bDiscardUnmatched);
+
 
 
 	//
