Various adaptive tessellation improvements

- Parallel adaptive refinement: fix an issue in refinement tagging causing non-determism and tessellation artifacts
- Feature adaptive refinement (curvature)
- Adaptive prefiltering based on sample rate to avoid aliasing
- Mesh post-processing to improve mesh regularity after tessellation

#rnx

diff --git a/Engine/Source/Runtime/GeometryCore/Public/Image/ImageBlur.h b/Engine/Source/Runtime/GeometryCore/Public/Image/ImageBlur.h
new file mode 100644
--- /dev/null
+++ b/Engine/Source/Runtime/GeometryCore/Public/Image/ImageBlur.h
@@ -0,0 +1,104 @@
+// Copyright Epic Games, Inc. All Rights Reserved.
+
+#pragma once
+
+#include "Solvers/Tridiagonal.h"
+#include "Spatial/SampledScalarField2.h"
+
+#include "Containers/Array.h"
+#include "Async/ParallelFor.h"
+
+namespace UE {
+namespace Geometry {
+
+/*
+ * Solve diffusion equation u_t = Laplacian(u), u(0) = u_0 until time Sigma^2/2, with Neumann boundary conditions.
+ * This corresponds to Gaussian blur with variance Sigma.
+ *
+ * Applies additive operator splitting and solves one implicit timestep. It converges independently of filtersize
+ * but becomes approximate for large Sigma.
+ */
+template <typename RealType>
+void HeatEquationImplicitAOS(TSampledScalarField2<RealType, RealType>& Field, const RealType SigmaX, const RealType SigmaY)
+{
+	const int32 Width  = Field.Width();
+	const int32 Height = Field.Height();
+
+	const RealType TauX = RealType(0.5) * SigmaX * SigmaX; // Timestep
+	const RealType TauY = RealType(0.5) * SigmaY * SigmaY; // Timestep
+
+	struct FTaskContext
+	{
+		TArray<RealType> Rhs;
+		TArray<RealType> Solution;
+	};
+	
+	TArray<FTaskContext> TaskContexts;
+
+	// columns
+	{
+		TArray<RealType> A, B, C;
+		A.Init(-TauY, Height);
+		B.Init(1. + 2. * TauY, Height);
+		C.Init(-TauY, Height);
+
+		C[0] *= 2.;
+		A[Height-1] *= 2.; 
+	
+		TTridiagonalSolver<RealType> Solver(A, B, C);
+		
+		ParallelForWithTaskContext(TEXT("HeatEquationImplicitAOS.Columns"), TaskContexts, Width, 32, [&](FTaskContext& TaskContext, int32 X)
+		{
+			TaskContext.Rhs.SetNum(Height);
+			TaskContext.Solution.SetNum(Height);
+
+			for (int32 Y=0; Y<Height; ++Y)
+			{
+				TaskContext.Rhs[Y] = Field.GridValues.At(X, Y);
+			}
+
+			Solver.Solve(TaskContext.Rhs, TaskContext.Solution);
+			
+			for (int32 Y=0; Y<Height; ++Y)
+			{
+				Field.GridValues.At(X, Y) = TaskContext.Solution[Y];
+			}
+		});
+	}
+
+	// blur along rows
+	{
+		TArray<RealType> A, B, C;
+		TArray<RealType> RowRhs, RowX;
+		
+		A.Init(-TauX, Width);
+		B.Init(1. + 2. * TauX, Width);
+		C.Init(-TauX, Width);
+
+		C[0] *= 2.;
+		A[Width-1] *= 2.; 
+	
+		TTridiagonalSolver<RealType> Solver(A, B, C);
+
+		ParallelForWithTaskContext(TEXT("HeatEquationImplicitAOS.Rows"), TaskContexts, Height, 32, [&](FTaskContext& TaskContext, int32 Y)
+		{
+			TaskContext.Rhs.SetNum(Width);
+			TaskContext.Solution.SetNum(Width);
+		
+			for (int32 X=0; X<Width; ++X)
+			{
+				TaskContext.Rhs[X] = Field.GridValues.At(X, Y);
+			}
+
+			Solver.Solve(TaskContext.Rhs, TaskContext.Solution);
+			
+			for (int32 X=0; X<Width; ++X)
+			{
+				Field.GridValues.At(X, Y) = TaskContext.Solution[X];
+			}
+		});
+	}
+}
+
+} // end namespace Geometry
+} // end namespace UE
\ No newline at end of file
diff --git a/Engine/Source/Runtime/GeometryCore/Public/Solvers/Tridiagonal.h b/Engine/Source/Runtime/GeometryCore/Public/Solvers/Tridiagonal.h
new file mode 100644
--- /dev/null
+++ b/Engine/Source/Runtime/GeometryCore/Public/Solvers/Tridiagonal.h
@@ -0,0 +1,104 @@
+// Copyright Epic Games, Inc. All Rights Reserved.
+
+#pragma once
+
+#include "Containers/Array.h"
+
+namespace UE {
+namespace Geometry {
+
+/*
+ * Solver for linear systems for tridiagonal matrices of the form
+ *
+ * [ B0  C0   0   0  ...   0    ]
+ * [ A1  B1  C1   0  ...   0    ]
+ * [  0  A2  B2  C2  ...   0    ]
+ * [  0   0  A3  B3  ...   0    ]
+ * [  .   .   .   .   .    Cn-2 ]
+ * [  0   0   0   0  ...   Bn-1 ]
+ *
+ * using Thomas algorithm.
+ */
+template <typename T>
+class TTridiagonalSolver
+{
+public:
+
+	TTridiagonalSolver(const TArray<T>& InA, const TArray<T>& InB, const TArray<T>& InC)
+	: A(InA)
+	, B(InB)
+	, C(InC)
+	, N(A.Num())
+	{
+		check(B.Num() == N);
+		check(C.Num() == N);
+		D.SetNum(N-1);
+		Prepare();
+	}
+
+	/*
+     * Solve linear system A * X = Rhs
+	 * 
+	 * Note that Rhs is modified. 
+	 */ 
+	bool Solve(TArray<T>& Rhs, TArray<T>& X)
+	{
+		if (!bInvertible)
+		{
+			return false;
+		}
+
+		Rhs[0] /= B[0];
+		for (size_t I=1; I<N; ++I)
+		{
+			const T Divisor = B[I] - A[I] * D[I-1];
+			if (FMath::Abs(Divisor) < T(UE_SMALL_NUMBER))
+			{
+				return false;
+			}
+			Rhs[I] = (Rhs[I] - A[I] * Rhs[I-1]) / Divisor;
+		}
+
+		X[N-1] = Rhs[N-1];
+		for (size_t I=N-2; ; --I)
+		{
+			X[I] = Rhs[I] - D[I] * X[I+1];
+			if (I==0)
+			{
+				break;
+			}
+		}
+		return true;
+	}
+
+private:
+
+	void Prepare()
+	{
+		bInvertible = false;
+		if (FMath::Abs(B[0]) < T(UE_SMALL_NUMBER))
+		{
+			return;
+		}
+
+		D[0] = C[0] / B[0];
+		for (size_t I=1; I<N-1; ++I)
+		{
+			const T Divisor = B[I] - A[I] * D[I-1];
+			if (FMath::Abs(Divisor) < T(UE_SMALL_NUMBER))
+			{
+				return;
+			}
+			D[I] = C[I] / Divisor;
+		}
+		bInvertible = true;
+	}
+
+	const TArray<T>& A, B, C; // matrix representation
+	TArray<T>        D;       // forward pass temporary
+	const size_t     N; 
+	bool             bInvertible;
+};
+
+} // end namespace Geometry
+} // end namespace UE
\ No newline at end of file
diff --git a/Engine/Source/Runtime/GeometryCore/Public/Tessellation/ParallelAdaptiveRefinement.h b/Engine/Source/Runtime/GeometryCore/Public/Tessellation/ParallelAdaptiveRefinement.h
--- a/Engine/Source/Runtime/GeometryCore/Public/Tessellation/ParallelAdaptiveRefinement.h
+++ b/Engine/Source/Runtime/GeometryCore/Public/Tessellation/ParallelAdaptiveRefinement.h
@@ -88,6 +88,11 @@ private:
 		void Resize(const int32 NumTris, const int32 NumEdges)
 		{
 			TRACE_CPUPROFILER_EVENT_SCOPE_TEXT("TParallelAdaptiveRefinement.FConcurrencyState.Resize");
+
+			for (FTriSplitRecord& Record : TriSplitRecords)
+			{
+				Record.Reset();
+			}
 			if (NumTris > TriSplitRecords.Num())
 			{
 				TriSplitRecords.AddDefaulted(NumTris - TriSplitRecords.Num());
@@ -110,19 +115,30 @@ public:
 	using RealType = typename MeshType::RealType; // the field related to the vector-space
 	using VecType  = typename MeshType::VecType;  // for positions and normals
 
-	// Initialize and run adaptive aptive tesselation
+	struct FOptions
+	{
+		int32 ParallelForBatchSize  { 1024 };
+
+		//< pre-reserve enough space for growing up to this many vertices
+		int32 ReserveVertices = 2'000'000;
+
+		bool bApplyNeighborhoodRefinement { true };
+	};
+
+	// Initialize and run adaptive refinement
 	TParallelAdaptiveRefinement(
 		MeshType& InMesh,                               //< mesh interface
 		RefinementPolicyType& InRefinementPolicy,       //< displacement and error interface
-		const int32 ParallelForBatchSize = 1024,        //<
-		const int32 ReserveVertices = 2'000'000 )       //< pre-reserve enough space for growing up to this many vertices
+		const FOptions& InOptions)  
 		: Mesh(InMesh)
 		, RefinementPolicy(InRefinementPolicy)
+		, Options(InOptions)
 	{
+		const int ParallelForBatchSize = Options.ParallelForBatchSize;
 		const EParallelForFlags ParallelForFlags = ParallelForBatchSize > 0 ? EParallelForFlags::None : EParallelForFlags::ForceSingleThread;
 
-		const int32 ReserveTriangles = 2 * ReserveVertices;
-		const int32 ReserveEdges     = 3 * ReserveVertices;
+		const int32 ReserveTriangles = 2 * Options.ReserveVertices;
+		const int32 ReserveEdges     = 3 * Options.ReserveVertices;
 
 		FConcurrencyState ConcurrencyState;
 		TArray<int32> RefinementCand;      //< list of triangles to inspect for refinement
@@ -136,7 +152,7 @@ public:
 		ConcurrencyState.TriSplitRecords.Reserve(ReserveTriangles);
 		ConcurrencyState.EdgeSplitRecords.Reserve(ReserveEdges);
 
-		Mesh.ReserveVertices(ReserveVertices);
+		Mesh.ReserveVertices(Options.ReserveVertices);
 		Mesh.ReserveTriangles(ReserveTriangles);
 
 		RefinementCand.Init(0, Mesh.MaxTriID());
@@ -175,11 +191,6 @@ public:
 					return;
 				}
 
-				if (ConcurrencyState.TriSplitRecords[TriIndex].bClaimed.test())
-				{
-					return; // early out
-				}
-
 				if (ShouldRefine(TriIndex, Level))
 				{
 					if (!ConcurrencyState.TriSplitRecords[TriIndex].bClaimed.test_and_set())
@@ -198,7 +209,10 @@ public:
 								EdgeSplitRequests[ NumEdgesToSplit.fetch_add(1, std::memory_order_relaxed) ] = EdgeIdx;
 							}
 						}
+					}
 
+					if (Options.bApplyNeighborhoodRefinement)
+					{
 						// now the same for all the neighors
 						for (int NbIdx=0; NbIdx<3; ++NbIdx)
 						{
@@ -293,7 +307,6 @@ public:
 				}
 			}, ParallelForFlags );
 
-			
 			// for each green refined triangles, store the edge mask for each neighboring triangle that is marked for refinement
 			ParallelFor( TEXT("TParallelAdaptiveRefinement.IdentifyGreenNeighbors.PF"), TriSplitQueueSize.load(), ParallelForBatchSize, [&](int32 Idx)
 			{
@@ -312,7 +325,7 @@ public:
 						}
 					}
 				}
-			});
+			}, ParallelForFlags);
 
 			// grow the arrays for edges and vertices
 			int32 PrevEdgeCount, PrevVtxCount;
@@ -686,6 +699,7 @@ private:
 
 	MeshType&               Mesh;
 	RefinementPolicyType&   RefinementPolicy;
+	const FOptions&         Options;
 };
 
 } // namespace Geometry
diff --git a/Engine/Source/Runtime/GeometryCore/Public/VectorUtil.h b/Engine/Source/Runtime/GeometryCore/Public/VectorUtil.h
--- a/Engine/Source/Runtime/GeometryCore/Public/VectorUtil.h
+++ b/Engine/Source/Runtime/GeometryCore/Public/VectorUtil.h
@@ -520,6 +520,17 @@ namespace VectorUtil
 		return AngleD(A, B);
 	}
 
+	/** 
+	 * @return oriented dihedral angle of two triangles sharing an oriented edge
+	 * 
+	 * @note Provided normals and Edge should be unit length, and Edge be perpendicular to both normals 
+	 * (see Keenan Crane, Discrete Differential Geometry: An Applied Introduction, Appendix 1.5) 
+	 */
+	template <typename RealType>
+	inline RealType OrientedDihedralAngle(const TVector<RealType>& Normal0, const TVector<RealType>& Normal1, const TVector<RealType>& Edge)
+	{
+		return FMath::Atan2( Edge | (Normal0 ^ Normal1), Normal0 | Normal1 );
+	}
 
 	/**
 	 * @return circumcenter of triangle ABC
@@ -547,6 +558,25 @@ namespace VectorUtil
 		return Center;
 	}
 
+
+	/*
+	 * @return incircle / circumcircle radii ratio of triangle ABC
+	 */
+	template <typename RealType>
+	RealType TriangleRegularity(const TVector<RealType>& A, const TVector<RealType>& B, const TVector<RealType>& C)
+	{
+		const RealType EAB = (A - B).Length();
+		const RealType EBC = (B - C).Length();
+		const RealType ECA = (C - A).Length();
+
+		const RealType SP = RealType(0.5) * (EAB + EBC + ECA); // Semiperimeter
+		const RealType Area = FMath::Sqrt(FMath::Max(RealType(0.), SP * (SP - EAB) * (SP - EBC) * (SP - ECA))); // Heron
+		const RealType IncircleRadius = Area / SP;
+		const RealType CircumcircleRadius = EAB * EBC * ECA / (RealType(4.) * Area);
+
+		return IncircleRadius / CircumcircleRadius;
+	}
+
 	/**
 	 * @return sign of Bitangent relative to Normal and Tangent
 	 */
@@ -616,7 +646,6 @@ namespace VectorUtil
 		return Normalized(Transform.GetScale3D() * Transform.InverseTransformVectorNoScale(Normal));
 	}
 
-
 }; // namespace VectorUtil
 
 
