Add a vertex -> single triangle method to dynamic mesh, as a faster/more-convenient path (vs e.g. using a vertex->triangles iterator to get a single element)

diff --git a/Engine/Source/Runtime/GeometryCore/Private/DynamicMesh/DynamicMesh3.cpp b/Engine/Source/Runtime/GeometryCore/Private/DynamicMesh/DynamicMesh3.cpp
--- a/Engine/Source/Runtime/GeometryCore/Private/DynamicMesh/DynamicMesh3.cpp
+++ b/Engine/Source/Runtime/GeometryCore/Private/DynamicMesh/DynamicMesh3.cpp
@@ -789,6 +789,16 @@ void FDynamicMesh3::EnumerateVertexTriangles(int32 VertexID, TFunctionRef<void(i
 	});
 }
 
+int32 FDynamicMesh3::GetSingleVertexTriangle(int32 VID) const
+{
+	checkSlow(VertexRefCounts.IsValid(VID));
+	if (!IsVertex(VID) || VertexEdgeLists.GetCount(VID) == 0)
+	{
+		return IndexConstants::InvalidID;
+	}
+	return VertexEdgeLists.First(VID);
+}
+
 
 void FDynamicMesh3::EnumerateEdgeTriangles(int32 EdgeID, TFunctionRef<void(int32)> ApplyFunc) const
 {
diff --git a/Engine/Source/Runtime/GeometryCore/Public/DynamicMesh/DynamicMesh3.h b/Engine/Source/Runtime/GeometryCore/Public/DynamicMesh/DynamicMesh3.h
--- a/Engine/Source/Runtime/GeometryCore/Public/DynamicMesh/DynamicMesh3.h
+++ b/Engine/Source/Runtime/GeometryCore/Public/DynamicMesh/DynamicMesh3.h
@@ -650,6 +650,9 @@ public:
 	/** Call ApplyFunc for each one-ring triangle of a vertex. Currently this is significantly more efficient than VtxTrianglesItr() in many use cases. */
 	GEOMETRYCORE_API void EnumerateVertexTriangles(int32 VertexID, TFunctionRef<void(int32)> ApplyFunc) const;
 
+	/** @return a single triangle connected to the given vertex, or INDEX_NONE if the vertex has no triangles */
+	GEOMETRYCORE_API int32 GetSingleVertexTriangle(int32 VID) const;
+
 	/** Call ApplyFunc for each triangle connected to an Edge (1 or 2 triangles) */
 	GEOMETRYCORE_API void EnumerateEdgeTriangles(int32 EdgeID, TFunctionRef<void(int32)> ApplyFunc) const;
 
