move per-access validation of AABB Tree queries to debug-only

diff --git a/Engine/Source/Runtime/GeometryCore/Public/Spatial/MeshAABBTree3.h b/Engine/Source/Runtime/GeometryCore/Public/Spatial/MeshAABBTree3.h
--- a/Engine/Source/Runtime/GeometryCore/Public/Spatial/MeshAABBTree3.h
+++ b/Engine/Source/Runtime/GeometryCore/Public/Spatial/MeshAABBTree3.h
@@ -61,6 +61,8 @@ class TMeshAABBTree3 : public IMeshSpatial
 {
 	friend class TFastWindingTree<TriangleMeshType>;
 
+	constexpr static bool bValidatePerAccess = (bool)DO_GUARD_SLOW;
+
 public:
 	using MeshType = TriangleMeshType;
 	using GetSplitAxisFunc = TUniqueFunction<int(int Depth, const FAxisAlignedBox3d& Box)>;
@@ -159,9 +161,12 @@ public:
 		const FQueryOptions& Options = FQueryOptions()
 	) const override
 	{
-		if ( ensure(IsValid(Options.bAllowUnsafeModifiedMeshQueries)) == false )
+		if constexpr (bValidatePerAccess)
 		{
-			return IndexConstants::InvalidID;
+			if (ensure(IsValid(Options.bAllowUnsafeModifiedMeshQueries)) == false)
+			{
+				return IndexConstants::InvalidID;
+			}
 		}
 
 		NearestDistSqr = (Options.MaxDistance < DOUBLE_MAX) ? Options.MaxDistance * Options.MaxDistance : DOUBLE_MAX;
@@ -213,9 +218,12 @@ public:
 		const FVector3d& Point, double ThresholdDistanceSqr, int& OutTriangleID, const FQueryOptions& Options = FQueryOptions()
 	) const
 	{
-		if (ensure(IsValid(Options.bAllowUnsafeModifiedMeshQueries)) == false)
+		if constexpr (bValidatePerAccess)
 		{
-			return false;
+			if (ensure(IsValid(Options.bAllowUnsafeModifiedMeshQueries)) == false)
+			{
+				return false;
+			}
 		}
 
 		OutTriangleID = IndexConstants::InvalidID;
@@ -328,9 +336,12 @@ public:
 	virtual int FindNearestVertex(const FVector3d& P, double& NearestDistSqr,
 		double MaxDist = TNumericLimits<double>::Max(), const FQueryOptions& Options = FQueryOptions())
 	{
-		if ( ensure(IsValid(Options.bAllowUnsafeModifiedMeshQueries)) == false )
+		if constexpr (bValidatePerAccess)
 		{
-			return IndexConstants::InvalidID;
+			if (ensure(IsValid(Options.bAllowUnsafeModifiedMeshQueries)) == false)
+			{
+				return IndexConstants::InvalidID;
+			}
 		}
 
 		NearestDistSqr = (MaxDist < DOUBLE_MAX) ? MaxDist * MaxDist : DOUBLE_MAX;
@@ -480,9 +491,12 @@ private:
 		//   nearestT to TNumericLimits<double>::Max(), then we will test all boxes (!)
 		NearestT = (Options.MaxDistance < TNumericLimits<float>::Max()) ? Options.MaxDistance : TNumericLimits<float>::Max();
 
-		if ( ensure(IsValid(Options.bAllowUnsafeModifiedMeshQueries)) == false )
+		if constexpr (bValidatePerAccess)
 		{
-			return false;
+			if (ensure(IsValid(Options.bAllowUnsafeModifiedMeshQueries)) == false)
+			{
+				return false;
+			}
 		}
 
 		FindHitTriangle(RootIndex, Ray, NearestT, TID, BaryCoords, Options);
@@ -494,9 +508,12 @@ private:
 		const RayType& Ray, TArray<MeshIntersection::FHitIntersectionResult>& OutHits,
 		const FQueryOptions& Options = FQueryOptions()) const
 	{
-		if ( ensure(IsValid(Options.bAllowUnsafeModifiedMeshQueries)) == false )
+		if constexpr (bValidatePerAccess)
 		{
-			return false;
+			if (ensure(IsValid(Options.bAllowUnsafeModifiedMeshQueries)) == false)
+			{
+				return false;
+			}
 		}
 
 		using HitResult = MeshIntersection::FHitIntersectionResult;
@@ -680,9 +697,12 @@ public:
 	 */
 	virtual bool TestAnyHitTriangle(const FRay3d& Ray, const FQueryOptions& Options = FQueryOptions()) const
 	{
-		if ( ensure(IsValid(Options.bAllowUnsafeModifiedMeshQueries)) == false )
+		if constexpr (bValidatePerAccess)
 		{
-			return false;
+			if (ensure(IsValid(Options.bAllowUnsafeModifiedMeshQueries)) == false)
+			{
+				return false;
+			}
 		}
 
 		double MaxDistance = (Options.MaxDistance < TNumericLimits<float>::Max()) ? Options.MaxDistance : TNumericLimits<float>::Max();
@@ -798,9 +818,12 @@ public:
 		double& Distance, const FQueryOptions& Options = FQueryOptions(), const FQueryOptions& OtherTreeOptions = FQueryOptions()
 	)
 	{
-		if ( ensure(IsValid(Options.bAllowUnsafeModifiedMeshQueries)) == false )
+		if constexpr (bValidatePerAccess)
 		{
-			return FIndex2i::Invalid();
+			if (ensure(IsValid(Options.bAllowUnsafeModifiedMeshQueries)) == false)
+			{
+				return FIndex2i::Invalid();
+			}
 		}
 
 		double NearestSqr = FMathd::MaxReal;
@@ -844,9 +867,12 @@ public:
 	 */
 	virtual void DoTraversal(FTreeTraversal& Traversal, const FQueryOptions& Options = FQueryOptions()) const
 	{
-		if ( ensure(IsValid(Options.bAllowUnsafeModifiedMeshQueries)) == false )
+		if constexpr (bValidatePerAccess)
 		{
-			return;
+			if (ensure(IsValid(Options.bAllowUnsafeModifiedMeshQueries)) == false)
+			{
+				return;
+			}
 		}
 
 		if (Traversal.NextBoxF(GetBox(RootIndex), 0))
@@ -924,9 +950,12 @@ public:
 		const FQueryOptions& Options = FQueryOptions()
 	) const
 	{
-		if ( ensure(IsValid(Options.bAllowUnsafeModifiedMeshQueries)) == false )
+		if constexpr (bValidatePerAccess)
 		{
-			return false;
+			if (ensure(IsValid(Options.bAllowUnsafeModifiedMeshQueries)) == false)
+			{
+				return false;
+			}
 		}
 
 		if (!TestMeshBounds.IsEmpty())
@@ -975,9 +1004,12 @@ public:
 		const FQueryOptions& Options = FQueryOptions(), const FQueryOptions& OtherTreeOptions = FQueryOptions()
 	) const
 	{
-		if ( ensure(IsValid(Options.bAllowUnsafeModifiedMeshQueries)) == false )
+		if constexpr (bValidatePerAccess)
 		{
-			return false;
+			if (ensure(IsValid(Options.bAllowUnsafeModifiedMeshQueries)) == false)
+			{
+				return false;
+			}
 		}
 
 		if (find_any_intersection(RootIndex, OtherTree, TransformF, OtherTree.RootIndex, 0,
@@ -1000,9 +1032,12 @@ public:
 		const FQueryOptions& Options = FQueryOptions()
 	) const
 	{
-		if ( ensure(IsValid(Options.bAllowUnsafeModifiedMeshQueries)) == false )
+		if constexpr (bValidatePerAccess)
 		{
-			return false;
+			if (ensure(IsValid(Options.bAllowUnsafeModifiedMeshQueries)) == false)
+			{
+				return false;
+			}
 		}
 
 		FAxisAlignedBox3d triBounds(Triangle.V[0], Triangle.V[1], Triangle.V[2]);
@@ -1072,9 +1107,12 @@ public:
 
 		MeshIntersection::FIntersectionsQueryResult Result;
 
-		if ( ensure(IsValid(Options.bAllowUnsafeModifiedMeshQueries)) == false )
+		if constexpr (bValidatePerAccess)
 		{
-			return Result;
+			if (ensure(IsValid(Options.bAllowUnsafeModifiedMeshQueries)) == false)
+			{
+				return Result;
+			}
 		}
 
 		find_self_intersections(&Result, bIgnoreTopoConnected, IntersectionFn, Options);
@@ -1084,9 +1122,12 @@ public:
 
 	virtual bool TestSelfIntersection(bool bIgnoreTopoConnected = true, const FQueryOptions& Options = FQueryOptions()) const
 	{
-		if ( ensure(IsValid(Options.bAllowUnsafeModifiedMeshQueries)) == false )
+		if constexpr (bValidatePerAccess)
 		{
-			return false;
+			if (ensure(IsValid(Options.bAllowUnsafeModifiedMeshQueries)) == false)
+			{
+				return false;
+			}
 		}
 
 		return find_self_intersections(nullptr, bIgnoreTopoConnected, 
