Move custom boolean custom inside test parameters to a struct, so we have the flexibility to add more info in the future (e.g., possibly a transform to/from the local space of the boolean)

diff --git a/Engine/Source/Runtime/GeometryCore/Private/Operations/MeshBoolean.cpp b/Engine/Source/Runtime/GeometryCore/Private/Operations/MeshBoolean.cpp
--- a/Engine/Source/Runtime/GeometryCore/Private/Operations/MeshBoolean.cpp
+++ b/Engine/Source/Runtime/GeometryCore/Private/Operations/MeshBoolean.cpp
@@ -208,9 +208,9 @@ using namespace UE::Geometry;
 bool FMeshBoolean::Compute()
 {
 	FCustomInsideMeshTest WindingTest(
-		[Threshold = WindingThreshold](FVector3d Pt, const FDynamicMeshAABBTree3& Spatial, const TFastWindingTree<FDynamicMesh3>* WindingTree) -> bool
+		[Threshold = WindingThreshold](FVector3d Pt, const FCustomInsideTestContext& Context) -> bool
 		{
-			return WindingTree->IsInside(Pt, Threshold);
+			return Context.OptionalWindingTree->IsInside(Pt, Threshold);
 		}, true);
 	return ComputeWithCustomInside(WindingTest, WindingTest);
 }
@@ -382,6 +382,11 @@ bool FMeshBoolean::ComputeWithCustomInside(FCustomInsideMeshTest InsideMeshA, FC
 			}
 			TFastWindingTree<FDynamicMesh3> LocalWinding(&OtherSpatial, UseInsideTest.bRequiresWinding);
 			TFastWindingTree<FDynamicMesh3>* UseWinding = UseInsideTest.bRequiresWinding ? &LocalWinding : nullptr;
+			FCustomInsideTestContext InsideTestContext
+			{
+				.Spatial = &OtherSpatial,
+				.OptionalWindingTree = UseWinding
+			};
 			
 			FDynamicMesh3& ProcessMesh = *CutMesh[MeshIdx];
 			int MaxTriID = ProcessMesh.MaxTriangleID();
@@ -467,7 +472,7 @@ bool FMeshBoolean::ComputeWithCustomInside(FCustomInsideMeshTest InsideMeshA, FC
 					}
 
 					// didn't already return a coplanar result; use the winding number
-					KeepTri[MeshIdx][TID] = (UseInsideTest.IsPointInsideFn(Centroid, OtherSpatial, UseWinding)) != bRemoveInside;
+					KeepTri[MeshIdx][TID] = (UseInsideTest.IsPointInsideFn(Centroid, InsideTestContext)) != bRemoveInside;
 				});
 		}
 
diff --git a/Engine/Source/Runtime/GeometryCore/Public/Operations/MeshBoolean.h b/Engine/Source/Runtime/GeometryCore/Public/Operations/MeshBoolean.h
--- a/Engine/Source/Runtime/GeometryCore/Public/Operations/MeshBoolean.h
+++ b/Engine/Source/Runtime/GeometryCore/Public/Operations/MeshBoolean.h
@@ -185,9 +185,18 @@ public:
 	 */
 	GEOMETRYCORE_API bool Compute();
 
+	// Data passed to the FCustomInsideMeshTest::IsPointInsideFn
+	struct FCustomInsideTestContext
+	{
+		// AABB Tree for the test mesh; always non-null
+		FDynamicMeshAABBTree3* Spatial;
+		// Note: OptionalWindingTree will be null when FCustomInsideMeshTest has bRequiresWinding == false
+		TFastWindingTree<FDynamicMesh3>* OptionalWindingTree;
+	};
+
 	// Function type for custom inside/outside query functions
-	// Note: If implementing a custom function that needs to access the underlying mesh, it is accessible via Spatial.GetMesh()
-	using TestInsideFunctionType = TFunction<bool(const FVector3d& TestPosition, const FDynamicMeshAABBTree3& Spatial, const TFastWindingTree<FDynamicMesh3>* OptionalWindingTree)>;
+	// Note: If implementing a custom function that needs to access the underlying mesh, it is accessible via Context.Spatial->GetMesh()
+	using TestInsideFunctionType = TFunction<bool(const FVector3d& TestPosition, const FCustomInsideTestContext& Context)>;
 	struct FCustomInsideMeshTest
 	{
 		TestInsideFunctionType IsPointInsideFn;
