add a version of Expand() to FAxisAlignedBox3 and FAxisAlignedBox2 that takes a vector

diff --git a/Engine/Source/Runtime/GeometryCore/Public/BoxTypes.h b/Engine/Source/Runtime/GeometryCore/Public/BoxTypes.h
--- a/Engine/Source/Runtime/GeometryCore/Public/BoxTypes.h
+++ b/Engine/Source/Runtime/GeometryCore/Public/BoxTypes.h
@@ -624,6 +624,12 @@ struct TAxisAlignedBox3
 		Min.Y -= Radius;
 		Min.Z -= Radius;
 	}
+
+	void Expand(const TVector<RealType>& ExpandByVec)
+	{
+		Max += ExpandByVec;
+		Min -= ExpandByVec;
+	}
 };
 
 template <typename RealType>
@@ -865,6 +871,12 @@ struct TAxisAlignedBox2
 		Min.X -= Radius;
 		Min.Y -= Radius;
 	}
+	
+	void Expand(const TVector2<RealType>& ExpandByVec)
+	{
+		Max += ExpandByVec;
+		Min -= ExpandByVec;
+	}
 };
 
 typedef TAxisAlignedBox2<float> FAxisAlignedBox2f;
