fast winding number optimizations: change set/map storage to array, remove redundant checks, tweak math to avoid sqrt
+ add test coverage for the fast winding values

diff --git a/Engine/Source/Runtime/GeometryCore/Public/Spatial/FastWinding.h b/Engine/Source/Runtime/GeometryCore/Public/Spatial/FastWinding.h
--- a/Engine/Source/Runtime/GeometryCore/Public/Spatial/FastWinding.h
+++ b/Engine/Source/Runtime/GeometryCore/Public/Spatial/FastWinding.h
@@ -82,7 +82,7 @@ namespace FastTriWinding
 	 */
 	template <class TriangleMeshType>
 	void ComputeCoeffs(const TriangleMeshType& Mesh,
-					   const TSet<int>& TriangleSet,
+					   const TArray<int>& TriangleArray,
 					   const FMeshTriInfoCache& TriCache,
 					   FVector3d& P,
 					   double& R,
@@ -92,14 +92,12 @@ namespace FastTriWinding
 	{
 		// If the data is small enough, don't bother trying to parallelize
 		constexpr int ParallelThreshold = 3000;			// Benchmarking shows parallel starts to outperform serial at around 2500 triangles or so
-		if (TriangleSet.Num() < ParallelThreshold)
+		if (TriangleArray.Num() < ParallelThreshold)
 		{
-			ComputeCoeffsSerial(Mesh, TriangleSet, TriCache, P, R, Order1, Order2);
+			ComputeCoeffsSerial(Mesh, TriangleArray, TriCache, P, R, Order1, Order2);
 			return;
 		}
 
-		TArray<int> TriangleArray = TriangleSet.Array();
-
 		// First compute the area-weighted centroid
 
 		struct PData
@@ -110,7 +108,6 @@ namespace FastTriWinding
 
 		auto CentroidTransform = [&](int64 TriangleSubsetIndex) -> PData
 		{
-			check(TriangleSubsetIndex < TriangleArray.Num());
 			int TID = TriangleArray[(int)TriangleSubsetIndex];
 
 			PData DataOut;
@@ -142,7 +139,6 @@ namespace FastTriWinding
 
 		auto ExpansionTransform = [&, P = CentroidData.P](int32 TriangleSubsetIndex) -> OrderData
 		{
-			check(TriangleSubsetIndex < TriangleArray.Num());
 			int tid = TriangleArray[(int)TriangleSubsetIndex];
 			FVector3d P0, P1, P2;
 			Mesh.GetTriVertices(tid, P0, P1, P2);
@@ -187,7 +183,37 @@ namespace FastTriWinding
 		Order2 = Orders.Order2;
 	}
 
-
+	/**
+	 *  precompute constant coefficients of triangle winding number approximation (evaluated in parallel for large sets of triangles)
+	 *  P: 'Center' of expansion for Triangles (area-weighted centroid avg)
+	 *  R: max distance from P to Triangles
+	 *  Order1: first-order vector coeff
+	 *  Order2: second-order matrix coeff
+	 *  TriCache: precomputed triangle centroid/normal/area
+	 */
+	 template <class TriangleMeshType>
+	 UE_DEPRECATED(5.7, "Use the above version of ComputeCoeffs, which takes a TArray of triangle indices instead")
+	 void ComputeCoeffs(const TriangleMeshType& Mesh,
+		 const TSet<int>& TriangleSet,
+		 const FMeshTriInfoCache& TriCache,
+		 FVector3d& P,
+		 double& R,
+		 FVector3d& Order1,
+		 FMatrix3d& Order2,
+		 int NumTasks = 16)
+	 {
+		 // If the data is small enough, don't bother trying to parallelize
+		 constexpr int ParallelThreshold = 3000;			// Benchmarking shows parallel starts to outperform serial at around 2500 triangles or so
+		 if (TriangleSet.Num() < ParallelThreshold)
+		 {
+			 ComputeCoeffsSerial(Mesh, TriangleSet, TriCache, P, R, Order1, Order2);
+			 return;
+		 }
+
+		 TArray<int> TriangleArray = TriangleSet.Array();
+		 ComputeCoeffs(Mesh, TriangleArray, TriCache, P, R, Order1, Order2, NumTasks);
+
+	 }
 
 
 	/**
@@ -296,9 +322,11 @@ class TFastWindingTree
 		double R;
 		FVector3d Order1Vec;
 		FMatrix3d Order2Mat;
+		bool bValid = false;
 	};
 
-	TMap<int, FWNInfo> FastWindingCache;
+	TArray<FWNInfo> FastWindingCache;
+	int32 FastWindingCacheOffset = 0;
 	uint64 FastWindingCacheMeshChangeStamp = 0;
 
 public:
@@ -397,8 +425,9 @@ private:
 				int ti = Tree->IndexList[idx + i];
 				Tree->Mesh->GetTriVertices(ti, a, b, c);
 				double angle = VectorUtil::TriSolidAngle(a, b, c, P);
-				branch_sum += angle / FMathd::FourPi;
+				branch_sum += angle;
 			}
+			branch_sum /= FMathd::FourPi;
 		}
 		else
 		{ // internal node, either 1 or 2 child boxes
@@ -453,18 +482,28 @@ private:
 	{
 		// set this to a larger number to ignore caches if number of Triangles is too small.
 		// (seems to be no benefit to doing this...is holdover from Tree-decomposition FWN code)
-		int WINDING_CACHE_THRESH = 1;
+		constexpr int WINDING_CACHE_THRESH = 1;
 
 		//FMeshTriInfoCache TriCache = null;
 		FMeshTriInfoCache TriCache = FMeshTriInfoCache::BuildTriInfoCache(*Tree->Mesh);
 
-		FastWindingCache.Empty(); // = TMap<int, FWNInfo>();
-		TOptional<TSet<int>> root_hash;
+		// Fast winding cache can only exist for non-leaf boxes, so find the offset where the cache-able boxes start
+		for (FastWindingCacheOffset = 0; FastWindingCacheOffset < Tree->BoxToIndex.Num(); ++FastWindingCacheOffset)
+		{
+			if (Tree->BoxToIndex[FastWindingCacheOffset] >= Tree->TrianglesEnd)
+			{
+				break;
+			}
+		}
+		int32 CacheSize = Tree->BoxToIndex.Num() - FastWindingCacheOffset;
+		FastWindingCache.Empty(CacheSize);
+		FastWindingCache.Init({ .bValid = false }, CacheSize);
+		TOptional<TArray<int>> root_hash;
 		build_fast_winding_cache(Tree->RootIndex, 0, WINDING_CACHE_THRESH, root_hash, TriCache);
 	}
-	int build_fast_winding_cache(int IBox, int Depth, int TriCountThresh, TOptional<TSet<int>>& TriHash, const FMeshTriInfoCache& TriCache)
+	int build_fast_winding_cache(int IBox, int Depth, int TriCountThresh, TOptional<TArray<int>>& TriArray, const FMeshTriInfoCache& TriCache)
 	{
-		TriHash.Reset();
+		TriArray.Reset();
 
 		int idx = Tree->BoxToIndex[IBox];
 		if (idx < Tree->TrianglesEnd)
@@ -478,7 +517,7 @@ private:
 			if (iChild1 < 0)
 			{ // 1 child, descend if nearer than cur min-dist
 				iChild1 = (-iChild1) - 1;
-				int num_child_tris = build_fast_winding_cache(iChild1, Depth + 1, TriCountThresh, TriHash, TriCache);
+				int num_child_tris = build_fast_winding_cache(iChild1, Depth + 1, TriCountThresh, TriArray, TriCache);
 
 				// if count in child is large enough, we already built a cache at lower node
 				return num_child_tris;
@@ -490,9 +529,9 @@ private:
 
 				// let each child build its own cache if it wants. If so, it will return the
 				// list of its child tris
-				TOptional<TSet<int>> child2_hash;
-				int num_tris_1 = build_fast_winding_cache(iChild1, Depth + 1, TriCountThresh, TriHash, TriCache);
-				int num_tris_2 = build_fast_winding_cache(iChild2, Depth + 1, TriCountThresh, child2_hash, TriCache);
+				TOptional<TArray<int>> child2_array;
+				int num_tris_1 = build_fast_winding_cache(iChild1, Depth + 1, TriCountThresh, TriArray, TriCache);
+				int num_tris_2 = build_fast_winding_cache(iChild2, Depth + 1, TriCountThresh, child2_array, TriCache);
 				bool build_cache = (num_tris_1 + num_tris_2 > TriCountThresh);
 
 				if (Depth == 0)
@@ -501,34 +540,34 @@ private:
 				}
 
 				// collect up the Triangles we need. there are various cases depending on what children already did
-				if (TriHash.IsSet() || child2_hash.IsSet() || build_cache)
+				if (TriArray.IsSet() || child2_array.IsSet() || build_cache)
 				{
-					if (!TriHash.IsSet() && child2_hash.IsSet())
+					if (!TriArray.IsSet() && child2_array.IsSet())
 					{
-						collect_triangles(iChild1, child2_hash.GetValue());
-						TriHash = child2_hash;
+						collect_triangles(iChild1, child2_array.GetValue());
+						TriArray = MoveTemp(child2_array); // Note child2_array is not used after this
 					}
 					else
 					{
-						if (!TriHash.IsSet())
+						if (!TriArray.IsSet())
 						{
-							TriHash.Emplace();
-							collect_triangles(iChild1, TriHash.GetValue());
+							TriArray.Emplace();
+							collect_triangles(iChild1, TriArray.GetValue());
 						}
-						if (!child2_hash.IsSet())
+						if (!child2_array.IsSet())
 						{
-							collect_triangles(iChild2, TriHash.GetValue());
+							collect_triangles(iChild2, TriArray.GetValue());
 						}
 						else
 						{
-							TriHash->Append(child2_hash.GetValue());
+							TriArray->Append(child2_array.GetValue());
 						}
 					}
 				}
 				if (build_cache)
 				{
-					check(TriHash.IsSet());
-					make_box_fast_winding_cache(IBox, TriHash.GetValue(), TriCache);
+					check(TriArray.IsSet());
+					make_box_fast_winding_cache(IBox, TriArray.GetValue(), TriCache);
 				}
 
 				return (num_tris_1 + num_tris_2);
@@ -539,37 +578,40 @@ private:
 	// check if value is in cache and far enough away from Q that we can use cached value
 	bool can_use_fast_winding_cache(int IBox, const FVector3d& Q) const
 	{
-		const FWNInfo* cacheInfo = FastWindingCache.Find(IBox);
-		if (!cacheInfo)
+		int32 CacheIdx = IBox - FastWindingCacheOffset;
+		if (CacheIdx < 0)
 		{
 			return false;
 		}
 
-		double dist_qp = Distance(cacheInfo->Center, Q);
-		if (dist_qp > FWNBeta * cacheInfo->R)
+		const FWNInfo& CacheInfo = FastWindingCache[CacheIdx];
+		if (!CacheInfo.bValid)
 		{
-			return true;
+			return false;
 		}
 
-		return false;
+		double DistQP_Sq = DistanceSquared(CacheInfo.Center, Q);
+		double DistThreshold = FWNBeta * CacheInfo.R;
+		return DistQP_Sq > DistThreshold * DistThreshold;
 	}
 
 	// compute FWN cache for all Triangles underneath this box
-	void make_box_fast_winding_cache(int IBox, const TSet<int>& Triangles, const FMeshTriInfoCache& TriCache)
+	void make_box_fast_winding_cache(int IBox, const TArray<int>& Triangles, const FMeshTriInfoCache& TriCache)
 	{
-		check(FastWindingCache.Find(IBox) == nullptr);
+		check(!FastWindingCache[IBox - FastWindingCacheOffset].bValid);
 
 		// construct cache
 		FWNInfo cacheInfo;
 		FastTriWinding::ComputeCoeffs(*Tree->Mesh, Triangles, TriCache, cacheInfo.Center, cacheInfo.R, cacheInfo.Order1Vec, cacheInfo.Order2Mat);
+		cacheInfo.bValid = true;
 
-		FastWindingCache.Add(IBox, cacheInfo);
+		FastWindingCache[IBox - FastWindingCacheOffset] = cacheInfo;
 	}
 
 	// evaluate the FWN cache for IBox
 	double evaluate_box_fast_winding_cache(int IBox, const FVector3d& Q) const
 	{
-		const FWNInfo& cacheInfo = FastWindingCache[IBox];
+		const FWNInfo& cacheInfo = FastWindingCache[IBox - FastWindingCacheOffset];
 
 		if (FWNApproxOrder == 2)
 		{
@@ -581,8 +623,8 @@ private:
 		}
 	}
 
-	// collect all the Triangles below IBox in a hash
-	void collect_triangles(int IBox, TSet<int>& Triangles)
+	// collect all the Triangles below IBox
+	void collect_triangles(int IBox, TArray<int>& Triangles)
 	{
 		int idx = Tree->BoxToIndex[IBox];
 		if (idx < Tree->TrianglesEnd)
