[GeometryCore] Delaunay 2D iterative update support

[RN]

diff --git a/Engine/Source/Runtime/GeometryCore/Private/CompGeom/Delaunay2.cpp b/Engine/Source/Runtime/GeometryCore/Private/CompGeom/Delaunay2.cpp
--- a/Engine/Source/Runtime/GeometryCore/Private/CompGeom/Delaunay2.cpp
+++ b/Engine/Source/Runtime/GeometryCore/Private/CompGeom/Delaunay2.cpp
@@ -1506,6 +1506,34 @@ namespace DelaunayInternal
 
 		return bSuccess;
 	}
+	
+	template<typename RealType>
+	void Update(FRandomStream& Random, FDelaunay2Connectivity& Connectivity,
+		TArrayView<const TVector2<RealType>> Vertices, TArray<int32> VertexOrder, FIndex3i StartSearchTri, bool bAssumeDelaunay)
+	{
+		FIndex3i SearchTri = StartSearchTri;
+		for (int32 OrderIdx = 0; OrderIdx < VertexOrder.Num(); OrderIdx++)
+		{
+			int32 Vertex = VertexOrder[OrderIdx];
+			int32 DuplicateOf = -1;
+			FIndex3i ContainingTri = WalkToContainingTri<RealType>(Random, Connectivity, Vertices, SearchTri, Vertex, bAssumeDelaunay, DuplicateOf);
+			if (DuplicateOf >= 0)
+			{
+				Connectivity.SetHasDuplicates();
+				if (Connectivity.HasDuplicateTracking())
+				{
+					Connectivity.MarkDuplicateVertex(DuplicateOf, Vertex);
+				}
+				continue;
+			}
+			if (ContainingTri[0] == FDelaunay2Connectivity::InvalidIndex)
+			{
+				continue;
+			}
+			SearchTri = Insert<RealType>(Connectivity, Vertices, ContainingTri, Vertex);
+			checkSlow(SearchTri.A != FDelaunay2Connectivity::InvalidIndex);
+		}
+	}
 
 	template<typename RealType>
 	bool Triangulate(FRandomStream& Random, FDelaunay2Connectivity& Connectivity,
@@ -1572,29 +1600,8 @@ namespace DelaunayInternal
 		Order.RemoveAt(BootstrapIndices[1]);
 		Order.RemoveAt(BootstrapIndices[2]);
 
-		FIndex3i SearchTri = FirstTri;
-		for (int32 OrderIdx = 0; OrderIdx < Order.Num(); OrderIdx++)
-		{
-			int32 Vertex = Order[OrderIdx];
-			int32 DuplicateOf = -1;
-			constexpr bool bAssumeDelaunay = true; // initial construction, before constraint edges, so safe to assume Delaunay
-			FIndex3i ContainingTri = WalkToContainingTri<RealType>(Random, Connectivity, Vertices, SearchTri, Vertex, bAssumeDelaunay, DuplicateOf);
-			if (DuplicateOf >= 0)
-			{
-				Connectivity.SetHasDuplicates();
-				if (Connectivity.HasDuplicateTracking())
-				{
-					Connectivity.MarkDuplicateVertex(DuplicateOf, Vertex);
-				}
-				continue;
-			}
-			if (ContainingTri[0] == FDelaunay2Connectivity::InvalidIndex)
-			{
-				continue;
-			}
-			SearchTri = Insert<RealType>(Connectivity, Vertices, ContainingTri, Vertex);
-			checkSlow(SearchTri.A != FDelaunay2Connectivity::InvalidIndex);
-		}
+		constexpr bool bAssumeDelaunay = true; // initial construction, before constraint edges, so safe to assume Delaunay
+		Update(Random, Connectivity, Vertices, Order, FirstTri, bAssumeDelaunay);
 
 		bool bSuccess = ConstrainEdges(Random, Connectivity, Vertices, Edges, bKeepFastEdgeAdjacencyData);
 		if (bSuccess)
@@ -1647,6 +1654,38 @@ bool FDelaunay2::ConstrainEdges(TArrayView<const FVector2f> Vertices, TArrayView
 	return bSuccess && ValidateEdgesResult(Edges);
 }
 
+bool FDelaunay2::Update(TArrayView<const TVector2<double>> Vertices, int32 FirstNewIdx)
+{
+	// Check if a triangulation already exists
+	if (Connectivity->NumTriangles() == 0)
+	{
+		return false;
+	}
+
+	// Construct BRIO for new vertices
+	FBRIOPoints InsertOrder;
+	InsertOrder.Compute(Vertices.RightChop(FirstNewIdx));
+	TArray<int32>& Order = InsertOrder.Order;
+	for (int32& VertexIndex : Order)
+	{
+		VertexIndex += FirstNewIdx;
+	}
+
+	// Get any existing triangle as a bootstrap one
+	FIndex3i SearchTri;
+	Connectivity->EnumerateTriangles([&SearchTri](const FIndex2i& Edge, int32 Vertex) -> bool
+	{
+		SearchTri.A = Vertex;
+		SearchTri.B = Edge.A;
+		SearchTri.C = Edge.B;
+		return false;
+	});
+
+	bool bAssumeDelaunay = !bIsConstrained; // this may have previously constrained edges in the original triangulation so we cannot assume Delaunay
+	DelaunayInternal::Update<double>(RandomStream, *Connectivity, Vertices, Order, SearchTri, bAssumeDelaunay);
+	return true;
+}
+
 TArray<FIndex3i> FDelaunay2::GetTriangles() const
 {
 	if (Connectivity.IsValid())
diff --git a/Engine/Source/Runtime/GeometryCore/Public/CompGeom/Delaunay2.h b/Engine/Source/Runtime/GeometryCore/Public/CompGeom/Delaunay2.h
--- a/Engine/Source/Runtime/GeometryCore/Public/CompGeom/Delaunay2.h
+++ b/Engine/Source/Runtime/GeometryCore/Public/CompGeom/Delaunay2.h
@@ -181,10 +181,9 @@ public:
 	GEOMETRYCORE_API bool ConstrainEdges(TArrayView<const TVector2<double>> Vertices, TArrayView<const FIndex2i> Edges);
 	GEOMETRYCORE_API bool ConstrainEdges(TArrayView<const TVector2<float>> Vertices, TArrayView<const FIndex2i> Edges);
 
-	// TODO: Support incremental vertex insertion
 	// Update the triangulation incrementally, assuming Vertices are unchanged before FirstNewIndex, and nothing after FirstNewIndex has been inserted yet
 	// Note that updating with new vertices *after* constraining edges may remove previously-constrained edges, unless we also add a way to tag constrained edges
-	// bool Update(TArrayView<const TVector2<double>> Vertices, int32 FirstNewIdx);
+	GEOMETRYCORE_API bool Update(TArrayView<const TVector2<double>> Vertices, int32 FirstNewIdx);
 
 	// Get the triangulation as an array of triangles
 	// Note: This creates a new array each call, because the internal data structure does not have a triangle array
