Fix incorrect comment on Dist3Triangle3Triangle3.h that refers to it as a triangle-segment distance

diff --git a/Engine/Source/Runtime/GeometryCore/Public/Distance/DistTriangle3Triangle3.h b/Engine/Source/Runtime/GeometryCore/Public/Distance/DistTriangle3Triangle3.h
--- a/Engine/Source/Runtime/GeometryCore/Public/Distance/DistTriangle3Triangle3.h
+++ b/Engine/Source/Runtime/GeometryCore/Public/Distance/DistTriangle3Triangle3.h
@@ -20,7 +20,7 @@ namespace Geometry
 using namespace UE::Math;
 
 /**
-* Compute unsigned distance between 3D segment and 3D triangle
+* Compute unsigned distance between two 3D triangles
 */
 template <typename Real>
 class TDistTriangle3Triangle3
