Fix core spline parameter issue

diff --git a/Engine/Source/Runtime/GeometryCore/Public/Splines/PolyBezierSpline.h b/Engine/Source/Runtime/GeometryCore/Public/Splines/PolyBezierSpline.h
--- a/Engine/Source/Runtime/GeometryCore/Public/Splines/PolyBezierSpline.h
+++ b/Engine/Source/Runtime/GeometryCore/Public/Splines/PolyBezierSpline.h
@@ -1308,10 +1308,21 @@ public:
 		{
 			float Start = Base::PairKnots[i].Value;
 			float End = Base::PairKnots[i + 1].Value;
-			if ((Parameter >= Start && Parameter < End) || (i == NumSegments - 1 && FMath::IsNearlyEqual(Parameter, End)))
+			// Exact match on last segment end point
+			if (i == NumSegments - 1 && FMath::IsNearlyEqual(Parameter, End))
+			{
+				OutLocalParam = 1.0f;
+				return i * 4;
+			}
+
+			if (Parameter >= Start && Parameter < End)
 			{
 				float SpanLength = End - Start;
-				OutLocalParam = (SpanLength > UE_KINDA_SMALL_NUMBER) ? (Parameter - Start) / SpanLength : 0.0f;
+
+				OutLocalParam = (SpanLength > FLT_EPSILON)
+					? (Parameter - Start) / SpanLength
+					: 0.0f;
+
 				return i * 4;
 			}
 		}
