Add a CreatePerVertex method to dynamic mesh overlays, as a faster alternative to the CreateFromPredicate() method in the per-vertex case

diff --git a/Engine/Source/Runtime/GeometryCore/Private/DynamicMesh/DynamicMeshOverlay.cpp b/Engine/Source/Runtime/GeometryCore/Private/DynamicMesh/DynamicMeshOverlay.cpp
--- a/Engine/Source/Runtime/GeometryCore/Private/DynamicMesh/DynamicMeshOverlay.cpp
+++ b/Engine/Source/Runtime/GeometryCore/Private/DynamicMesh/DynamicMeshOverlay.cpp
@@ -163,6 +163,44 @@ void TDynamicMeshOverlay<RealType, ElementSize>::CreateFromPredicate(TFunctionRe
 	}
 }
 
+template<typename RealType, int ElementSize>
+void TDynamicMeshOverlay<RealType, ElementSize>::CreatePerVertex(RealType InitElementValue)
+{
+	ClearElements(); // deletes all elements and initializes triangles to be 1:1 w/ parentmesh IDs
+	TArray<int32> VIDtoEID;
+	bool bIsCompactV = ParentMesh->IsCompactV();
+	if (!bIsCompactV)
+	{
+		VIDtoEID.SetNumUninitialized(ParentMesh->MaxVertexID());
+	}
+	for (int32 VertexID : ParentMesh->VertexIndicesItr())
+	{
+		int32 ElementID = AppendElement(InitElementValue);
+		ParentVertices.InsertAt(VertexID, ElementID, FDynamicMesh3::InvalidID);
+		if (!bIsCompactV)
+		{
+			VIDtoEID[VertexID] = ElementID;
+		}
+		else
+		{
+			checkSlow(VertexID == ElementID);
+		}
+	}
+
+	for (int32 TriangleID : ParentMesh->TriangleIndicesItr())
+	{
+		FIndex3i Tri = ParentMesh->GetTriangle(TriangleID);
+		int32 Start = 3 * TriangleID;
+		for (int32 SubIdx = 0; SubIdx < 3; ++SubIdx)
+		{
+			int32 VID = Tri[SubIdx];
+			int32 EID = bIsCompactV ? VID : VIDtoEID[VID];
+			ElementTriangles.InsertAt(EID, Start + SubIdx, FDynamicMesh3::InvalidID);
+			ElementsRefCounts.Increment(EID);
+		}
+	}
+}
+
 
 template<typename RealType, int ElementSize>
 void TDynamicMeshOverlay<RealType, ElementSize>::SplitVerticesWithPredicate(TFunctionRef<bool(int ElementID, int TriID)> ShouldSplitOutVertex, TFunctionRef<void(int ElementID, int TriID, RealType * FillVect)> GetNewElementValue)
diff --git a/Engine/Source/Runtime/GeometryCore/Public/DynamicMesh/DynamicMeshOverlay.h b/Engine/Source/Runtime/GeometryCore/Public/DynamicMesh/DynamicMeshOverlay.h
--- a/Engine/Source/Runtime/GeometryCore/Public/DynamicMesh/DynamicMeshOverlay.h
+++ b/Engine/Source/Runtime/GeometryCore/Public/DynamicMesh/DynamicMeshOverlay.h
@@ -379,6 +379,12 @@ public:
 	 */
 	void CreateFromPredicate(TFunctionRef<bool(int ParentVertexIdx, int TriIDA, int TriIDB)> TrisCanShareVertexPredicate, RealType InitElementValue);
 
+	/**
+	 * Build overlay topology with one element per vertex.
+	 * Note: Faster-but-equivalent-to calling CreateFromPredicate() with a CanShareVertex predicate that always returns true
+	 */
+	void CreatePerVertex(RealType InitElementValue);
+
 	/**
 	 * Refine an existing overlay topology.  For any element on a given triangle, if the predicate returns true, it gets topologically split out so it isn't shared by any other triangle.
 	 * Used for creating sharp vertices in the normals overlay.
