TangentBezierSpline: Fix several bugs manipulating the tangent modes array

diff --git a/Engine/Source/Runtime/GeometryCore/Public/Splines/TangentBezierSpline.h b/Engine/Source/Runtime/GeometryCore/Public/Splines/TangentBezierSpline.h
--- a/Engine/Source/Runtime/GeometryCore/Public/Splines/TangentBezierSpline.h
+++ b/Engine/Source/Runtime/GeometryCore/Public/Splines/TangentBezierSpline.h
@@ -772,7 +772,8 @@ public:
 			SetTangentIn(InsertedPointIndex, ControlPoint.TangentIn);
 			SetTangentOut(InsertedPointIndex, ControlPoint.TangentOut);
 		}
-        
+
+		TangentModes.InsertDefaulted(InsertedPointIndex);
 		SetPointTangentMode(InsertedPointIndex, ControlPoint.TangentMode);
 				
 		// Update adjacent points if they have automatic tangents
@@ -821,9 +822,10 @@ public:
 			SetTangentIn(InsertedPointIndex, ControlPoint.TangentIn);
 			SetTangentOut(InsertedPointIndex, ControlPoint.TangentOut);
 		}
-        
+		
+		TangentModes.InsertDefaulted(InsertedPointIndex);
 		SetPointTangentMode(InsertedPointIndex, ControlPoint.TangentMode);
-				
+		
 		// Update adjacent points if they have automatic tangents
 		if (InsertedPointIndex > 0 && 
 			(TangentModes[InsertedPointIndex-1] == ETangentMode::Auto || 
@@ -891,8 +893,9 @@ public:
 		// Let InternalSpline handle the prepending efficiently
 		InternalSpline.PrependBezierSegment(P0, P1, P2);
         
+		TangentModes.InsertDefaulted(0);
 		SetPointTangentMode(0, ControlPoint.TangentMode);
-                
+		
 		// Update tangents for the next point if it's auto
 		if (GetNumPoints() > 0 && 
 			(TangentModes[1] == ETangentMode::Auto || TangentModes[1] == ETangentMode::AutoClamped))
@@ -963,8 +966,8 @@ public:
     
 		// Add the segment
 		InternalSpline.AppendBezierSegment(P1, P2, P3);
-	    
 		
+		TangentModes.InsertDefaulted(NewIndex);
 		SetPointTangentMode(NewIndex, ControlPoint.TangentMode);
 
         // If this is auto mode, update tangents for adjacent points
@@ -1078,7 +1081,7 @@ public:
         bool bPositionChanged = (OldPosition != ControlPoint.Position);
         
         SetPointTangentMode(Index, ControlPoint.TangentMode);
-        
+
         // Update position
         if (Index == 0)
         {
@@ -1337,23 +1340,17 @@ public:
 	  */
 	void SetPointTangentMode(int32 Index, ETangentMode Mode)
 	{
-		const int32 NumPoints = GetNumPoints();
-		if (Index < 0 || Index > NumPoints)
+		if (Index < 0 || Index > GetNumPoints())
 		{
 			return;
 		}
         
-		// Ensure array is big enough
-		if (TangentModes.Num() <= NumPoints)
-		{
-			TangentModes.AddZeroed(NumPoints - TangentModes.Num() + 1);
-		}
-        
 		TangentModes[Index] = Mode;
 
 		// Update tangents based on new mode
         UpdatePointTangents(Index, Mode);
 	}
+	
 	virtual void Clear() override
 	{
 		InternalSpline.Clear();
@@ -1537,6 +1534,9 @@ public:
 			}
 			break;
 
+		case ETangentMode::User:
+			break;
+			
 		case ETangentMode::Unknown:
 		default:
 			// Default to Auto mode if unknown
@@ -1692,11 +1692,13 @@ private:
     			P2 = P3 - (ControlPoint.TangentIn / 3.0f);
     		}
     		
-	        // update tangent mode
-    		SetPointTangentMode(1, ControlPoint.TangentMode);
-    		
     		// Add the first segment
     		InternalSpline.SetControlPoints( {P0, P1, P2, P3}, EParameterizationPolicy::Centripetal);
+
+    		// update tangent mode
+    		TangentModes.AddDefaulted();
+    		SetPointTangentMode(1, ControlPoint.TangentMode);
+    		
     		OutPointIndex = 1;
     		return true;
     	}
