PolyBezierSpline: Fix bug where knots would fail to be inserted when appending or prepending segments if the segment length was too small relative to the magnitude of the knots at the beginning/end of the spline.

The bug is that we assumed X + e != X (where X is a float and e is epsilon, some small value). This does not hold for large values X due to the non-uniform distribution of floating point values over the real numbers.

diff --git a/Engine/Source/Runtime/GeometryCore/Public/Splines/PolyBezierSpline.h b/Engine/Source/Runtime/GeometryCore/Public/Splines/PolyBezierSpline.h
--- a/Engine/Source/Runtime/GeometryCore/Public/Splines/PolyBezierSpline.h
+++ b/Engine/Source/Runtime/GeometryCore/Public/Splines/PolyBezierSpline.h
@@ -1927,11 +1927,15 @@ private:
 		}
 		else if (bAppendToEnd)
 		{
-			this->InsertKnot( FKnot(Base::PairKnots.Last().Value + SegmentLength, this->bClampEnds ? this->GetDegree() : 1));
+			const float LastParameter = Base::PairKnots.Last().Value;
+			const float NextParameter = FMath::Max(LastParameter + SegmentLength, std::nextafter(LastParameter, std::numeric_limits<float>::infinity()));
+			this->InsertKnot( FKnot(NextParameter, this->bClampEnds ? this->GetDegree() : 1));
 		}
 		else if (bPrependToStart)
 		{
-			this->InsertKnot( FKnot(Base::PairKnots[SegmentIndex].Value - SegmentLength, this->bClampEnds ? this->GetDegree() : 1));
+			const float FirstParameter = Base::PairKnots[0].Value;
+			const float NextParameter = FMath::Min(FirstParameter - SegmentLength, std::nextafter(FirstParameter, -std::numeric_limits<float>::infinity()));
+			this->InsertKnot( FKnot(NextParameter, this->bClampEnds ? this->GetDegree() : 1));
 		}
 	
 		this->ApplyClampedKnotsMultiplicity();
