Splines: Harden serialization

1) Use FString for type names instead of FName (to have case preservation regardles of build configuration).
    - This is back compat safe because it corrects the registry at runtime only so that it matches saves splines, it does not invalidate currently serialized data

2) Fix multi spline serialization
    - This is not intrinsically back compat safe, but I have registered a legacy factory that matches the previous type ID so that splines saved before this change will continue to load properly.

diff --git a/Engine/Source/Runtime/GeometryCore/Private/GeometryCoreModule.cpp b/Engine/Source/Runtime/GeometryCore/Private/GeometryCoreModule.cpp
--- a/Engine/Source/Runtime/GeometryCore/Private/GeometryCoreModule.cpp
+++ b/Engine/Source/Runtime/GeometryCore/Private/GeometryCoreModule.cpp
@@ -3,6 +3,9 @@
 #include "GeometryCoreModule.h"
 
 #include "CompGeom/ExactPredicates.h"
+#include "Splines/MultiSpline.h"
+#include "Splines/SplineTypeRegistry.h"
+#include "Splines/TangentBezierSpline.h"
 
 #define LOCTEXT_NAMESPACE "FGeometryCoreModule"
 
@@ -10,6 +13,8 @@ DEFINE_LOG_CATEGORY(LogGeometry);
 
 void FGeometryCoreModule::StartupModule()
 {
+	RegisterSplineTypes();
+	
 	// This code will execute after your module is loaded into memory; the exact timing is specified in the .uplugin file per-module
 	UE::Geometry::ExactPredicates::GlobalInit();
 }
@@ -20,6 +25,25 @@ void FGeometryCoreModule::ShutdownModule()
 	// we call this function before unloading the module.
 }
 
+void FGeometryCoreModule::RegisterSplineTypes() const
+{
+	using namespace UE::Geometry::Spline;
+
+	/**
+	 * TMultiSpline<FTangentBezierSpline3d> once had an incorrect type ID which failed to take into account the wrapped spline type (FTangentBezierSpline3d).
+	 * Now, the implementation name of a multi-spline is in the form of "MultiSpline.[WrappedSplineName]".
+	 * This registered factory exists to allow for proper loading of multi-splines which were saved before this bug was fixed.
+	 */
+	const FString LegacyMultiSplineImplName = TEXT("MultiSpline");
+	const FString LegacyMultiSplineValueTypeName = TEXT("Vector");
+	const FSplineTypeId::IdType LegacyMultiSplineID = FSplineTypeId::GenerateTypeId(*LegacyMultiSplineImplName, *LegacyMultiSplineValueTypeName);
+	ensure (LegacyMultiSplineID == 0x54E3AD22);
+	FSplineTypeRegistry::RegisterType(LegacyMultiSplineID, LegacyMultiSplineImplName, LegacyMultiSplineValueTypeName, []() -> TUniquePtr<ISplineInterface>
+	{
+		return MakeUnique<TMultiSpline<FTangentBezierSpline3d>>();
+	});
+}
+
 #undef LOCTEXT_NAMESPACE
 	
 IMPLEMENT_MODULE(FGeometryCoreModule, GeometryCore)
\ No newline at end of file
diff --git a/Engine/Source/Runtime/GeometryCore/Public/GeometryCoreModule.h b/Engine/Source/Runtime/GeometryCore/Public/GeometryCoreModule.h
--- a/Engine/Source/Runtime/GeometryCore/Public/GeometryCoreModule.h
+++ b/Engine/Source/Runtime/GeometryCore/Public/GeometryCoreModule.h
@@ -13,4 +13,8 @@ public:
 	/** IModuleInterface implementation */
 	virtual void StartupModule() override;
 	virtual void ShutdownModule() override;
+
+private:
+
+	void RegisterSplineTypes() const;
 };
diff --git a/Engine/Source/Runtime/GeometryCore/Public/Splines/BSpline.h b/Engine/Source/Runtime/GeometryCore/Public/Splines/BSpline.h
--- a/Engine/Source/Runtime/GeometryCore/Public/Splines/BSpline.h
+++ b/Engine/Source/Runtime/GeometryCore/Public/Splines/BSpline.h
@@ -118,7 +118,7 @@ public:
 
 	DECLARE_SPLINE_TYPE_ID(
 		BSplineNameSelector<DEGREE>::Name,
-		*TSplineValueTypeTraits<VALUETYPE>::Name.ToString()
+		*TSplineValueTypeTraits<VALUETYPE>::Name
 	);
 	
 	TBSpline() = default;
diff --git a/Engine/Source/Runtime/GeometryCore/Public/Splines/BuiltInAttributeTypes.h b/Engine/Source/Runtime/GeometryCore/Public/Splines/BuiltInAttributeTypes.h
--- a/Engine/Source/Runtime/GeometryCore/Public/Splines/BuiltInAttributeTypes.h
+++ b/Engine/Source/Runtime/GeometryCore/Public/Splines/BuiltInAttributeTypes.h
@@ -21,85 +21,85 @@ namespace Spline
 // Basic types
 template<> struct TSplineValueTypeTraits<float> 
 { 
-	static inline const FName Name = TEXT("Float");
+	static inline const FString Name = TEXT("Float");
 };
 
 template<> struct TSplineValueTypeTraits<double> 
 { 
-	static inline const FName Name = TEXT("Double");
+	static inline const FString Name = TEXT("Double");
 };
 
 template<> struct TSplineValueTypeTraits<int32> 
 { 
-	static inline const FName Name = TEXT("Int32");
+	static inline const FString Name = TEXT("Int32");
 };
 
 template<> struct TSplineValueTypeTraits<bool> 
 { 
-	static inline const FName Name = TEXT("Bool");
+	static inline const FString Name = TEXT("Bool");
 };
 
 // Math types
 template<> struct TSplineValueTypeTraits<FVector3d> 
 { 
-	static inline const FName Name = TEXT("Vector");
+	static inline const FString Name = TEXT("Vector");
 };
 
 template<> struct TSplineValueTypeTraits<FVector2D> 
 { 
-	static inline const FName Name = TEXT("Vector2D");
+	static inline const FString Name = TEXT("Vector2D");
 };
 
 template<> struct TSplineValueTypeTraits<FVector3f> 
 { 
-	static inline const FName Name = TEXT("Vector3f");
+	static inline const FString Name = TEXT("Vector3f");
 };
 	
 template<> struct TSplineValueTypeTraits<FVector4> 
 { 
-	static inline const FName Name = TEXT("Vector4");
+	static inline const FString Name = TEXT("Vector4");
 };
 
 template<> struct TSplineValueTypeTraits<FQuat> 
 { 
-	static inline const FName Name = TEXT("Quat");
+	static inline const FString Name = TEXT("Quat");
 };
 
 template<> struct TSplineValueTypeTraits<FRotator> 
 { 
-	static inline const FName Name = TEXT("Rotator");
+	static inline const FString Name = TEXT("Rotator");
 };
 
 template<> struct TSplineValueTypeTraits<FTransform> 
 { 
-	static inline const FName Name = TEXT("Transform");
+	static inline const FString Name = TEXT("Transform");
 };
 
 // Color types
 template<> struct TSplineValueTypeTraits<FLinearColor> 
 { 
-	static inline const FName Name = TEXT("LinearColor");
+	static inline const FString Name = TEXT("LinearColor");
 };
 
 template<> struct TSplineValueTypeTraits<FColor> 
 { 
-	static inline const FName Name = TEXT("Color");
+	static inline const FString Name = TEXT("Color");
 };
 
 // String types
 template<> struct TSplineValueTypeTraits<FName> 
 { 
-	static inline const FName Name = TEXT("Name");
+	static inline const FString Name = TEXT("Name");
 };
 
 template<> struct TSplineValueTypeTraits<FString> 
 { 
-	static inline const FName Name = TEXT("String");
+	static inline const FString Name = TEXT("String");
 };
 
 template<> struct TSplineValueTypeTraits<FText> 
 { 
-	static inline const FName Name = TEXT("Text");
+	static inline const FString Name = TEXT("Text");
 };
 
 } // namespace Spline
diff --git a/Engine/Source/Runtime/GeometryCore/Public/Splines/LinearSpline.h b/Engine/Source/Runtime/GeometryCore/Public/Splines/LinearSpline.h
--- a/Engine/Source/Runtime/GeometryCore/Public/Splines/LinearSpline.h
+++ b/Engine/Source/Runtime/GeometryCore/Public/Splines/LinearSpline.h
@@ -16,7 +16,7 @@ class TLinearSpline : public TBSpline<T, 1>
 public:
 	using Base = TBSpline<T, 1>;
 
-	DECLARE_SPLINE_TYPE_ID(TEXT("BSpline1"), *TSplineValueTypeTraits<T>::Name.ToString());
+	DECLARE_SPLINE_TYPE_ID(TEXT("BSpline1"), *TSplineValueTypeTraits<T>::Name);
 
 	TLinearSpline()
 	{
diff --git a/Engine/Source/Runtime/GeometryCore/Public/Splines/MultiSpline.h b/Engine/Source/Runtime/GeometryCore/Public/Splines/MultiSpline.h
--- a/Engine/Source/Runtime/GeometryCore/Public/Splines/MultiSpline.h
+++ b/Engine/Source/Runtime/GeometryCore/Public/Splines/MultiSpline.h
@@ -204,14 +204,32 @@ private:
     };
 
 public:
+	
     using SplineType = typename TSplineWrapper<SPLINETYPE>::SplineType;
 	typedef typename TSplineWrapper<SPLINETYPE>::ValueType ValueType;
 	
-	// Generate compile-time type ID for MultiSpline
-	DECLARE_SPLINE_TYPE_ID(
-		TEXT("MultiSpline"),
-		*TSplineValueTypeTraits<ValueType>::Name.ToString()
-	);
+	static const FSplineTypeId::IdType& GetStaticTypeId()
+	{
+		static const FSplineTypeId::IdType CachedTypeId =
+		FSplineTypeId::GenerateTypeId(*GetSplineTypeName(),
+			*TSplineValueTypeTraits<ValueType>::Name);
+		return CachedTypeId;
+	}
+	
+	virtual FSplineTypeId::IdType GetTypeId() const override
+	{
+		return GetStaticTypeId();
+	}
+	
+	static FString GetSplineTypeName()
+	{
+		return FString::Printf(TEXT("MultiSpline.%s"), *SplineType::GetSplineTypeName());
+	}
+	
+	virtual FString GetImplementationName() const override
+	{
+		return GetSplineTypeName();
+	}
 	
     SplineType& GetSpline() { return TSplineWrapper<SPLINETYPE>::InternalSpline; }
     const SplineType& GetSpline() const { return TSplineWrapper<SPLINETYPE>::InternalSpline; }
@@ -298,8 +316,8 @@ public:
             if (ExistingChannel->Spline.IsValid())
             {
             	// Get implementation names for comparison
-            	const FName RequestedImplName = ImplType::GetSplineTypeName();
-            	const FName ExistingImplName = ExistingChannel->Spline->GetImplementationName();
+            	const FString RequestedImplName = ImplType::GetSplineTypeName();
+            	const FString ExistingImplName = ExistingChannel->Spline->GetImplementationName();
             
             	// Check implementation type compatibility
             	if (RequestedImplName == ExistingImplName)
@@ -310,7 +328,7 @@ public:
 
             	// Log warning about implementation type mismatch
             	UE_LOG(LogSpline, Warning, TEXT("Failed to fetch attribute channel '%s': Channel exists with implementation type '%s', but requested type is '%s'."),
-					   *Name.ToString(), *ExistingImplName.ToString(), *RequestedImplName.ToString());
+					   *Name.ToString(), *ExistingImplName, *RequestedImplName);
             }
         }
         return nullptr;
@@ -328,8 +346,8 @@ public:
         {
             if (ExistingChannel->Spline.IsValid())
 		    {
-			    const FName TypeName = TSplineValueTypeTraits<AttrType>::Name;
-                const FName ExistingTypeName = ExistingChannel->Spline->GetValueTypeName();
+			    const FString TypeName = TSplineValueTypeTraits<AttrType>::Name;
+                const FString ExistingTypeName = ExistingChannel->Spline->GetValueTypeName();
 			    
 			    if (TypeName == ExistingTypeName)
 			    {
@@ -338,7 +356,7 @@ public:
 
 			    // Channel exists but with a different type
 			    UE_LOG(LogSpline, Warning, TEXT("Failed to fetch attribute channel '%s': Channel already exists with type '%s', but requested type is '%s'."),
-					    *Name.ToString(), *ExistingTypeName.ToString(), *TypeName.ToString());
+					    *Name.ToString(), *ExistingTypeName, *TypeName);
 		    }
         }
 
@@ -420,7 +438,7 @@ public:
 	    }
 	    
 	    // Get implementation information from source
-	    const FName ImplName = SourceChannel->GetImplementationName();
+	    const FString ImplName = SourceChannel->GetImplementationName();
 	    
 	    // Create new destination channel with same implementation
 	    if (!CreateAttributeChannelInternal<AttrType>(DestName, ImplName))
@@ -524,8 +542,8 @@ public:
     	{
     		// Get type information for registration
     		TypeId = ImplType::GetStaticTypeId();
-    		FName ImplName = ImplType::GetSplineTypeName();
-    		FName ValueName = TSplineValueTypeTraits<AttrType>::Name;
+    		FString ImplName = ImplType::GetSplineTypeName();
+    		FString ValueName = TSplineValueTypeTraits<AttrType>::Name;
         
     		// Register with a factory function that creates this specific type
     		bool bSuccess = FSplineTypeRegistry::RegisterType(
@@ -538,7 +556,7 @@ public:
     		if (!bSuccess)
     		{
     			UE_LOG(LogSpline, Error, TEXT("Failed to register spline type '%s' with value type '%s'"),
-					  *ImplName.ToString(), *ValueName.ToString());
+					  *ImplName, *ValueName);
     			return nullptr;
     		}
     	}
@@ -658,7 +676,7 @@ private:
      * @return True if channel was created successfully
      */
     template <typename AttrType>
-    bool CreateAttributeChannelInternal(const FName& Name, const FName& ImplName)
+    bool CreateAttributeChannelInternal(const FName& Name, const FString& ImplName)
     {
         // Check if channel already exists
         if (Children.Contains(Name))
@@ -672,7 +690,7 @@ private:
         if (TypeId == 0)
         {
             UE_LOG(LogSpline, Error, TEXT("Failed to find type ID for implementation '%s' with value type '%s'."),
-                   *ImplName.ToString(), *TSplineValueTypeTraits<AttrType>::Name.ToString());
+                   *ImplName, *TSplineValueTypeTraits<AttrType>::Name);
             return false;
         }
 
@@ -681,7 +699,7 @@ private:
         if (!NewSpline)
         {
                 UE_LOG(LogSpline, Error, TEXT("Failed to create spline of type '%s'/'%s' (ID: 0x%08X) for attribute channel '%s'."),
-               *ImplName.ToString(), *TSplineValueTypeTraits<AttrType>::Name.ToString(), TypeId, *Name.ToString());
+               *ImplName, *TSplineValueTypeTraits<AttrType>::Name, TypeId, *Name.ToString());
             return false;
         }
 
@@ -689,7 +707,7 @@ private:
         if (NewSpline->GetValueTypeName() != TSplineValueTypeTraits<AttrType>::Name)
         {
             UE_LOG(LogSpline, Error, TEXT("Created spline has incorrect value type. Expected '%s', got '%s'."),
-                   *TSplineValueTypeTraits<AttrType>::Name.ToString(), *NewSpline->GetValueTypeName().ToString());
+                   *TSplineValueTypeTraits<AttrType>::Name, *NewSpline->GetValueTypeName());
             return false;
         }
 
diff --git a/Engine/Source/Runtime/GeometryCore/Public/Splines/PolyBezierSpline.h b/Engine/Source/Runtime/GeometryCore/Public/Splines/PolyBezierSpline.h
--- a/Engine/Source/Runtime/GeometryCore/Public/Splines/PolyBezierSpline.h
+++ b/Engine/Source/Runtime/GeometryCore/Public/Splines/PolyBezierSpline.h
@@ -25,7 +25,7 @@ public:
 	// Generate compile-time type ID for PolyBezier
 	DECLARE_SPLINE_TYPE_ID(
 		TEXT("PolyBezier"),
-		*TSplineValueTypeTraits<VALUETYPE>::Name.ToString()
+		*TSplineValueTypeTraits<VALUETYPE>::Name
 	);
 	
     TPolyBezierSpline() = default;
diff --git a/Engine/Source/Runtime/GeometryCore/Public/Splines/Spline.h b/Engine/Source/Runtime/GeometryCore/Public/Splines/Spline.h
--- a/Engine/Source/Runtime/GeometryCore/Public/Splines/Spline.h
+++ b/Engine/Source/Runtime/GeometryCore/Public/Splines/Spline.h
@@ -41,8 +41,8 @@ public:
        	// Serialize wrapped spline
 	    return InternalSpline.Serialize(Ar);
     }
-    virtual FName GetValueTypeName() const override { return InternalSpline.GetValueTypeName(); }
-    virtual FName GetImplementationName() const override { return InternalSpline.GetImplementationName(); }
+    virtual FString GetValueTypeName() const override { return InternalSpline.GetValueTypeName(); }
+    virtual FString GetImplementationName() const override { return InternalSpline.GetImplementationName(); }
     virtual FInterval1f GetParameterSpace() const override { return InternalSpline.GetParameterSpace(); }
     virtual void SetClosedLoop(bool bClosed) override { InternalSpline.SetClosedLoop(bClosed); }
     virtual bool IsClosedLoop() const override { return InternalSpline.IsClosedLoop(); }
diff --git a/Engine/Source/Runtime/GeometryCore/Public/Splines/SplineInterfaces.h b/Engine/Source/Runtime/GeometryCore/Public/Splines/SplineInterfaces.h
--- a/Engine/Source/Runtime/GeometryCore/Public/Splines/SplineInterfaces.h
+++ b/Engine/Source/Runtime/GeometryCore/Public/Splines/SplineInterfaces.h
@@ -40,8 +40,8 @@ public:
     virtual TUniquePtr<ISplineInterface> Clone() const = 0;
 	virtual bool IsEqual(const ISplineInterface* OtherSpline) const = 0;
     
-    virtual FName GetValueTypeName() const = 0;
-    virtual FName GetImplementationName() const = 0;
+    virtual FString GetValueTypeName() const = 0;
+    virtual FString GetImplementationName() const = 0;
     virtual FInterval1f GetParameterSpace() const = 0;
     virtual void SetClosedLoop(bool bClosed) = 0;
     virtual bool IsClosedLoop() const = 0;
@@ -90,7 +90,7 @@ public:
 template <typename VALUETYPE>
 class TSplineInterface : public ISplineInterface
 {
-    static FName TypeName;
+    static FString TypeName;
 
 public:
     typedef VALUETYPE ValueType;
@@ -101,7 +101,7 @@ public:
     TSplineInterface()
     {
         // Ensure the value type name is initialized if not already
-        if (TypeName.IsNone())
+        if (TypeName.IsEmpty())
         {
             TypeName = TSplineValueTypeTraits<VALUETYPE>::Name;
         }
@@ -170,8 +170,8 @@ public:
     /**
      * Set and get the type name for the value type
      */
-    static void SetTypeName(const FName& InTypeName) { TypeName = InTypeName; } 
-    virtual FName GetValueTypeName() const override { return TypeName; }
+    static void SetTypeName(const FString& InTypeName) { TypeName = InTypeName; } 
+    virtual FString GetValueTypeName() const override { return TypeName; }
 
     /** Set Pre-Infinity handling mode when out of bounds */ 
     virtual void SetPreInfinityMode(EOutOfBoundsHandlingMode InMode){ PreInfinityMode = InMode; }
@@ -235,7 +235,7 @@ protected:
 
 // Initialize static member
 template <typename VALUETYPE>
-inline FName TSplineInterface<VALUETYPE>::TypeName;
+inline FString TSplineInterface<VALUETYPE>::TypeName;
 
 } // namespace Spline
 } // namespace Geometry
diff --git a/Engine/Source/Runtime/GeometryCore/Public/Splines/SplineTypeId.h b/Engine/Source/Runtime/GeometryCore/Public/Splines/SplineTypeId.h
--- a/Engine/Source/Runtime/GeometryCore/Public/Splines/SplineTypeId.h
+++ b/Engine/Source/Runtime/GeometryCore/Public/Splines/SplineTypeId.h
@@ -93,8 +93,8 @@ public:
      * Checks if two type IDs might be in conflict
      * Used during registration to detect potential hash collisions
      */
-    static bool CheckTypeIdConflict(IdType TypeId, const FName& ImplName, const FName& ValueTypeName, 
-                                  const FName& ExistingImplName, const FName& ExistingValueTypeName)
+    static bool CheckTypeIdConflict(IdType TypeId, const FString& ImplName, const FString& ValueTypeName, 
+                                  const FString& ExistingImplName, const FString& ExistingValueTypeName)
     {
         // If the names don't match but the TypeId does, we have a hash collision
         bool bImplNameMismatch = (ImplName != ExistingImplName);
@@ -105,9 +105,9 @@ public:
             UE_LOG(LogSpline, Error, TEXT(" HASH COLLISION in SplineTypeId!! ️"));
             UE_LOG(LogSpline, Error, TEXT("  TypeId: 0x%08X"), TypeId);
             UE_LOG(LogSpline, Error, TEXT("  Attempted registration: Impl=%s, Value=%s"), 
-                   *ImplName.ToString(), *ValueTypeName.ToString());
+                   *ImplName, *ValueTypeName);
             UE_LOG(LogSpline, Error, TEXT("  Existing registration: Impl=%s, Value=%s"),
-                   *ExistingImplName.ToString(), *ExistingValueTypeName.ToString());
+                   *ExistingImplName, *ExistingValueTypeName);
             return true;
         }
         
@@ -127,13 +127,13 @@ public:
     { \
         return GetStaticTypeId(); \
     } \
-    static FName GetSplineTypeName() \
+    static FString GetSplineTypeName() \
     { \
-        return FName(ImplName); \
+        return FString(ImplName); \
     } \
-    virtual FName GetImplementationName() const override \
+    virtual FString GetImplementationName() const override \
     { \
-        return FName(ImplName); \
+        return FString(ImplName); \
     }
 
 } // namespace Spline
diff --git a/Engine/Source/Runtime/GeometryCore/Public/Splines/SplineTypeRegistry.h b/Engine/Source/Runtime/GeometryCore/Public/Splines/SplineTypeRegistry.h
--- a/Engine/Source/Runtime/GeometryCore/Public/Splines/SplineTypeRegistry.h
+++ b/Engine/Source/Runtime/GeometryCore/Public/Splines/SplineTypeRegistry.h
@@ -41,16 +41,16 @@ public:
      */
     static inline bool RegisterType(
         TypeId TypeId,
-        const FName& ImplName, 
-        const FName& ValueTypeName,
+        const FString& ImplName, 
+        const FString& ValueTypeName,
         FactoryFunction Factory)
     {
         // Check if this type is already registered
-        TPair<FName, FName> NamePair(ImplName, ValueTypeName);
+        TPair<FString, FString> NamePair(ImplName, ValueTypeName);
         FScopeLock Lock(&GetRegistryLock());
         
         // Check for existing registration with same ID but different names (hash collision)
-        if (const TPair<FName, FName>* ExistingNames = GetIdToNames().Find(TypeId))
+        if (const TPair<FString, FString>* ExistingNames = GetIdToNames().Find(TypeId))
         {
             if (FSplineTypeId::CheckTypeIdConflict(TypeId, ImplName, ValueTypeName, 
                                                   ExistingNames->Key, ExistingNames->Value))
@@ -75,8 +75,8 @@ public:
         GetNamesToId().Add(NamePair, TypeId);
         GetFactories().Add(TypeId, MoveTemp(Factory));
         
-        UE_LOG(LogSpline, Verbose, TEXT("Registered spline type: TypeId=0x%08X, Impl=%s, Value=%s"),
-            TypeId, *ImplName.ToString(), *ValueTypeName.ToString());
+        UE_LOG(LogSpline, Log, TEXT("Registered spline type: TypeId=0x%08X, Impl=%s, Value=%s"),
+            TypeId, *ImplName, *ValueTypeName);
         
         return true;
     }
@@ -98,11 +98,11 @@ public:
         }
         
         // Log helpful error message
-        const TPair<FName, FName>* Names = GetIdToNames().Find(TypeId);
+        const TPair<FString, FString>* Names = GetIdToNames().Find(TypeId);
         if (Names)
         {
             UE_LOG(LogSpline, Warning, TEXT("No factory for type ID 0x%08X (Impl=%s, Value=%s)"),
-                TypeId, *Names->Key.ToString(), *Names->Value.ToString());
+                TypeId, *Names->Key, *Names->Value);
         }
         else
         {
@@ -120,10 +120,10 @@ public:
      * @return A new spline instance or nullptr if not found
      */
     static inline TUniquePtr<ISplineInterface> CreateSpline(
-        const FName& ImplName,
-        const FName& ValueTypeName)
+        const FString& ImplName,
+        const FString& ValueTypeName)
     {
-        TPair<FName, FName> NamePair(ImplName, ValueTypeName);
+        TPair<FString, FString> NamePair(ImplName, ValueTypeName);
         
         FScopeLock Lock(&GetRegistryLock());
         
@@ -134,7 +134,7 @@ public:
         }
         
         UE_LOG(LogSpline, Warning, TEXT("No registered type for Impl=%s, Value=%s"),
-            *ImplName.ToString(), *ValueTypeName.ToString());
+            *ImplName, *ValueTypeName);
         return nullptr;
     }
     
@@ -148,12 +148,12 @@ public:
      */
     static inline bool GetTypeNames(
         TypeId TypeId,
-        FName& OutImplName,
-        FName& OutValueTypeName)
+        FString& OutImplName,
+        FString& OutValueTypeName)
     {
         FScopeLock Lock(&GetRegistryLock());
         
-        const TPair<FName, FName>* Names = GetIdToNames().Find(TypeId);
+        const TPair<FString, FString>* Names = GetIdToNames().Find(TypeId);
         if (Names)
         {
             OutImplName = Names->Key;
@@ -171,10 +171,10 @@ public:
      * @return The type ID or 0 if not found
      */
     static inline TypeId GetTypeId(
-        const FName& ImplName,
-        const FName& ValueTypeName)
+        const FString& ImplName,
+        const FString& ValueTypeName)
     {
-        TPair<FName, FName> NamePair(ImplName, ValueTypeName);
+        TPair<FString, FString> NamePair(ImplName, ValueTypeName);
         
         FScopeLock Lock(&GetRegistryLock());
         
@@ -199,18 +199,18 @@ private:
     /**
      * Gets the map from type ID to name pair (thread-safe singleton access)
      */
-    static inline TMap<TypeId, TPair<FName, FName>>& GetIdToNames()
+    static inline TMap<TypeId, TPair<FString, FString>>& GetIdToNames()
     {
-        static TMap<TypeId, TPair<FName, FName>> IdToNames;
+        static TMap<TypeId, TPair<FString, FString>> IdToNames;
         return IdToNames;
     }
     
     /**
      * Gets the map from name pair to type ID (thread-safe singleton access)
      */
-    static inline TMap<TPair<FName, FName>, TypeId>& GetNamesToId()
+    static inline TMap<TPair<FString, FString>, TypeId>& GetNamesToId()
     {
-        static TMap<TPair<FName, FName>, TypeId> NamesToId;
+        static TMap<TPair<FString, FString>, TypeId> NamesToId;
         return NamesToId;
     }
     
@@ -255,8 +255,8 @@ private:
     {
         // Get type information from the derived class
         FSplineTypeRegistry::TypeId TypeId = DerivedType::GetStaticTypeId();
-        FName ImplName = DerivedType::GetSplineTypeName();
-        FName ValueName = TSplineValueTypeTraits<ValueType>::Name;
+        FString ImplName = DerivedType::GetSplineTypeName();
+        FString ValueName = TSplineValueTypeTraits<ValueType>::Name;
         
         // Register factory function
         FSplineTypeRegistry::RegisterType(
diff --git a/Engine/Source/Runtime/GeometryCore/Public/Splines/TangentBezierSpline.h b/Engine/Source/Runtime/GeometryCore/Public/Splines/TangentBezierSpline.h
--- a/Engine/Source/Runtime/GeometryCore/Public/Splines/TangentBezierSpline.h
+++ b/Engine/Source/Runtime/GeometryCore/Public/Splines/TangentBezierSpline.h
@@ -135,7 +135,7 @@ public:
 	// Generate compile-time type ID for TangentBezier
 	DECLARE_SPLINE_TYPE_ID(
 		TEXT("TangentBezier"),
-		*TSplineValueTypeTraits<VALUETYPE>::Name.ToString()
+		*TSplineValueTypeTraits<VALUETYPE>::Name
 	);
 	
 	TTangentBezierSpline() = default;
