Splines: Fix missed type trait fixups from recent change

diff --git a/Engine/Source/Runtime/GeometryCore/Public/Splines/CustomAttributeExample.h b/Engine/Source/Runtime/GeometryCore/Public/Splines/CustomAttributeExample.h
--- a/Engine/Source/Runtime/GeometryCore/Public/Splines/CustomAttributeExample.h
+++ b/Engine/Source/Runtime/GeometryCore/Public/Splines/CustomAttributeExample.h
@@ -104,7 +104,7 @@ struct FFoo
 // of TSplineValueTypeTraits with a Name field
 template<> struct TSplineValueTypeTraits<FFoo> 
 { 
-    static inline const FName Name = TEXT("Foo");
+    static inline const FString Name = TEXT("Foo");
 };
 
 // Define a specialization of TSplineInterpolationPolicy to handle
