minor additions/cleanup for geometrycore's oriented bounding box:
 1. add conversions to/from geometrycore's OBB and the engine math OBB.
 2. move simple shape approx's 'reparameterize oriented box' method over to the oriented box types struct so it can be used more generally.
 3. move large 'obb merge' method implementation to a cpp file.

diff --git a/Engine/Plugins/Runtime/GeometryProcessing/Source/DynamicMesh/Private/ShapeApproximation/MeshSimpleShapeApproximation.cpp b/Engine/Plugins/Runtime/GeometryProcessing/Source/DynamicMesh/Private/ShapeApproximation/MeshSimpleShapeApproximation.cpp
--- a/Engine/Plugins/Runtime/GeometryProcessing/Source/DynamicMesh/Private/ShapeApproximation/MeshSimpleShapeApproximation.cpp
+++ b/Engine/Plugins/Runtime/GeometryProcessing/Source/DynamicMesh/Private/ShapeApproximation/MeshSimpleShapeApproximation.cpp
@@ -28,48 +28,7 @@ namespace FMeshSimpleShapeApproximationLocals
 	// to the side if the original box was actually on its "side").
 	FOrientedBox3d ReparameterizeBoxCloserToWorldFrame(const FOrientedBox3d& Box)
 	{
-		FVector3d Axes[3]{ Box.AxisX(), Box.AxisY(), Box.AxisZ() };	
-
-		int32 BestZ = FMath::Max3Index(
-			FMath::Abs(Axes[0].Z),
-			FMath::Abs(Axes[1].Z),
-			FMath::Abs(Axes[2].Z));
-		FVector3d NewZ = Axes[BestZ] * (Axes[BestZ].Z > 0 ? 1 : -1);
-
-		// Pick the best Y out of the two remaining
-		double DotY[3]{
-			Axes[0].Y,
-			Axes[1].Y,
-			Axes[2].Y
-		};
-		DotY[BestZ] = 0; // don't pick this one
-
-		int32 BestY = FMath::Max3Index(
-			FMath::Abs(DotY[0]),
-			FMath::Abs(DotY[1]),
-			FMath::Abs(DotY[2]));
-		FVector3d NewY = Axes[BestY] * (DotY[BestY] > 0 ? 1 : -1);
-
-		// Sum of BestY and BestZ will be either 0+1=1, 0+2=2, or 1+2=3, and the 
-		// corresponding leftover index will be 2, 1, or 0.
-		int32 BestX = 3 - (BestY + BestZ);
-		// Static analyzer probably won't like that, so here's a clamp for safety
-		BestX = FMath::Clamp(BestX, 0, 2);
-
-		// Sanity check to make sure we picked unique axes
-		if (!ensure(BestX != BestY 
-			&& BestY != BestZ 
-			&& BestX != BestZ))
-		{
-			return Box;
-		}
-
-		return FOrientedBox3d (
-			FFrame3d(Box.Frame.Origin,
-				NewY.Cross(NewZ), // better to do this than risk picking the wrong direction for BestX
-				NewY,
-				NewZ),
-			FVector3d(Box.Extents[BestX], Box.Extents[BestY], Box.Extents[BestZ]));
+		return Box.ReparameterizedCloserToWorldFrame();
 	}
 }
 
diff --git a/Engine/Source/Runtime/GeometryCore/Private/OrientedBoxTypes.cpp b/Engine/Source/Runtime/GeometryCore/Private/OrientedBoxTypes.cpp
new file mode 100644
--- /dev/null
+++ b/Engine/Source/Runtime/GeometryCore/Private/OrientedBoxTypes.cpp
@@ -0,0 +1,231 @@
+// Copyright Epic Games, Inc. All Rights Reserved.
+
+#include "OrientedBoxTypes.h"
+
+#include "Math/UnrealMath.h"
+
+
+
+namespace UE::Geometry
+{
+
+template<typename RealType>
+TOrientedBox3<RealType> TOrientedBox3<RealType>::Merge(const TOrientedBox3<RealType>& Other, bool bOnlyConsiderExistingAxes) const
+{
+	TQuaternion<RealType> RotInv = Frame.Rotation.Inverse();
+	TQuaternion<RealType> RotInvOtherRot = RotInv * Other.Frame.Rotation;
+	TVector<RealType> O2mO1 = Other.Frame.Origin - Frame.Origin;
+	TVector<RealType> TranslateIntoFrame1 = RotInv * O2mO1;
+	TVector<RealType> TranslateIntoFrame2 = Other.Frame.Rotation.InverseMultiply(-O2mO1);
+	TAxisAlignedBox3<RealType> LocalBox1(-Extents, Extents);
+	TAxisAlignedBox3<RealType> LocalBox2(-Other.Extents, Other.Extents);
+	// Corners of this box, in its local coordinate frame
+	TVector<RealType> Corners1[8]
+	{
+		LocalBox1.Min,
+		LocalBox1.Max,
+		TVector<RealType>(LocalBox1.Min.X, LocalBox1.Min.Y, LocalBox1.Max.Z),
+		TVector<RealType>(LocalBox1.Min.X, LocalBox1.Max.Y, LocalBox1.Min.Z),
+		TVector<RealType>(LocalBox1.Max.X, LocalBox1.Min.Y, LocalBox1.Min.Z),
+		TVector<RealType>(LocalBox1.Min.X, LocalBox1.Max.Y, LocalBox1.Max.Z),
+		TVector<RealType>(LocalBox1.Max.X, LocalBox1.Min.Y, LocalBox1.Max.Z),
+		TVector<RealType>(LocalBox1.Max.X, LocalBox1.Max.Y, LocalBox1.Min.Z),
+	};
+	// Cordinates of Other, in its local coordinate frame
+	TVector<RealType> Corners2[8]
+	{
+		LocalBox2.Min,
+		LocalBox2.Max,
+		TVector<RealType>(LocalBox2.Min.X, LocalBox2.Min.Y, LocalBox2.Max.Z),
+		TVector<RealType>(LocalBox2.Min.X, LocalBox2.Max.Y, LocalBox2.Min.Z),
+		TVector<RealType>(LocalBox2.Max.X, LocalBox2.Min.Y, LocalBox2.Min.Z),
+		TVector<RealType>(LocalBox2.Min.X, LocalBox2.Max.Y, LocalBox2.Max.Z),
+		TVector<RealType>(LocalBox2.Max.X, LocalBox2.Min.Y, LocalBox2.Max.Z),
+		TVector<RealType>(LocalBox2.Max.X, LocalBox2.Max.Y, LocalBox2.Min.Z),
+	};
+	// Expand each local AABB to contain the other box's points
+	for (int32 Idx = 0; Idx < 8; ++Idx)
+	{
+		LocalBox1.Contain(TranslateIntoFrame1 + RotInvOtherRot * Corners2[Idx]);
+		LocalBox2.Contain(TranslateIntoFrame2 + RotInvOtherRot.InverseMultiply(Corners1[Idx]));
+	}
+	RealType CandidateVolumes[3]{ LocalBox1.Volume(), LocalBox2.Volume(), TMathUtil<RealType>::MaxReal };
+	TVector<RealType> Axes1[3], Axes2[3], FoundAxes[3];
+	TAxisAlignedBox3<RealType> FoundBounds;
+	if (!bOnlyConsiderExistingAxes)
+	{
+		// Compute the two local frames
+		Frame.Rotation.GetAxes(Axes1[0], Axes1[1], Axes1[2]);
+		Other.Frame.Rotation.GetAxes(Axes2[0], Axes2[1], Axes2[2]);
+		// Compute metrics re the size of the boxes and how far apart they are
+		RealType MaxExtentsSizeSq = FMath::Max(Extents.SizeSquared(), Other.Extents.SizeSquared());
+		RealType CenterDiffSizeSq = O2mO1.SizeSquared();
+		// For boxes that are far apart relative to their size, fit a box along the line connecting their centers
+		if (CenterDiffSizeSq > MaxExtentsSizeSq)
+		{
+			// find the most-perpendicular original box axis as our second axis
+			// (Note we could alternatively try interpolating axes, or picking the axis with the largest projected extent ...)
+			FoundAxes[0] = O2mO1.GetSafeNormal(0.0);
+			TVector<RealType> MostPerpendicular;
+			RealType SmallestDot = TMathUtil<RealType>::MaxReal;
+			for (int32 AxisIdx = 0; AxisIdx < 3; ++AxisIdx)
+			{
+				RealType Dot1 = FMath::Abs(Axes1[AxisIdx].Dot(FoundAxes[0]));
+				RealType Dot2 = FMath::Abs(Axes2[AxisIdx].Dot(FoundAxes[0]));
+				if (Dot1 < Dot2)
+				{
+					if (Dot1 < SmallestDot)
+					{
+						SmallestDot = Dot1;
+						MostPerpendicular = Axes1[AxisIdx];
+					}
+				}
+				else if (Dot2 < SmallestDot)
+				{
+					SmallestDot = Dot2;
+					MostPerpendicular = Axes2[AxisIdx];
+				}
+			}
+			FoundAxes[2] = FoundAxes[0].Cross(MostPerpendicular).GetSafeNormal(0);
+			FoundAxes[1] = FoundAxes[2].Cross(FoundAxes[0]);
+		}
+		else // otherwise, for boxes that are closer together, interpolate their orientations
+		{
+			// Find which axes on the Other.Frame.Rotation are closest to the Frame.Rotation's X and Y axes
+			int32 Axis1X_IdxInAxis2 = 0;
+			int32 Axis1Y_IdxInAxis2 = 0;
+			RealType BestXAlignment = FMath::Abs(Axes1[0].Dot(Axes2[0]));
+			RealType BestYAlignment = FMath::Abs(Axes1[1].Dot(Axes2[0]));
+			for (int32 Axis2Idx = 1; Axis2Idx < 3; ++Axis2Idx)
+			{
+				RealType XAlignment = FMath::Abs(Axes1[0].Dot(Axes2[Axis2Idx]));
+				RealType YAlignment = FMath::Abs(Axes1[1].Dot(Axes2[Axis2Idx]));
+				if (XAlignment > BestXAlignment)
+				{
+					Axis1X_IdxInAxis2 = Axis2Idx;
+					BestXAlignment = XAlignment;
+				}
+				if (YAlignment > BestYAlignment)
+				{
+					Axis1Y_IdxInAxis2 = Axis2Idx;
+					BestYAlignment = YAlignment;
+				}
+			}
+			if (Axis1X_IdxInAxis2 == Axis1Y_IdxInAxis2)
+			{
+				// don't let the closest axis to Y be the same as the closest to X
+				// (this could happen in a 45-degree case where two axes are tied; should be very rare)
+				Axis1Y_IdxInAxis2 = (1 + Axis1Y_IdxInAxis2) % 3;
+			}
+			RealType SignX = RealType(Axes1[0].Dot(Axes2[Axis1X_IdxInAxis2]) < 0 ? -1 : 1);
+			RealType SignY = RealType(Axes1[1].Dot(Axes2[Axis1Y_IdxInAxis2]) < 0 ? -1 : 1);
+			FoundAxes[0] = (Axes1[0] + SignX * Axes2[Axis1X_IdxInAxis2]).GetSafeNormal();
+			FoundAxes[1] = (Axes1[1] + SignY * Axes2[Axis1Y_IdxInAxis2]);
+			// Make sure the axes form an orthonormal frame
+			FoundAxes[2] = (FoundAxes[0].Cross(FoundAxes[1])).GetSafeNormal();
+			FoundAxes[1] = FoundAxes[2].Cross(FoundAxes[0]).GetSafeNormal();
+		}
+
+		// Fit a local AABB in the Found space
+		for (int32 CornerIdx = 0; CornerIdx < 8; ++CornerIdx)
+		{
+			TVector<RealType> WorldCorner1 = Frame.FromFramePoint(Corners1[CornerIdx]);
+			TVector<RealType> AxesCorner1(WorldCorner1.Dot(FoundAxes[0]), WorldCorner1.Dot(FoundAxes[1]), WorldCorner1.Dot(FoundAxes[2]));
+			FoundBounds.Contain(AxesCorner1);
+			TVector<RealType> WorldCorner2 = Other.Frame.FromFramePoint(Corners2[CornerIdx]);
+			TVector<RealType> AxesCorner2(WorldCorner2.Dot(FoundAxes[0]), WorldCorner2.Dot(FoundAxes[1]), WorldCorner2.Dot(FoundAxes[2]));
+			FoundBounds.Contain(AxesCorner2);
+		}
+		CandidateVolumes[2] = FoundBounds.Volume();
+	}
+
+	// Choose the coordinate space with the smallest bounding box, by volume
+	int BestIdx = 0;
+	RealType BestVolume = CandidateVolumes[0];
+	int AxesToConsider = bOnlyConsiderExistingAxes ? 2 : 3;
+	for (int Idx = 1; Idx < AxesToConsider; ++Idx)
+	{
+		if (CandidateVolumes[Idx] < BestVolume)
+		{
+			BestIdx = Idx;
+			BestVolume = CandidateVolumes[Idx];
+		}
+	}
+	TOrientedBox3<RealType> Result;
+	if (BestIdx == 2)
+	{
+		// Convert the Found Axes into a new oriented box
+		Result.Frame.Rotation.SetFromRotationMatrix(TMatrix3<RealType>(FoundAxes[0], FoundAxes[1], FoundAxes[2], false));
+		TVector<RealType> FoundCenter = FoundBounds.Center();
+		Result.Frame.Origin = FoundAxes[0] * FoundCenter.X + FoundAxes[1] * FoundCenter.Y + FoundAxes[2] * FoundCenter.Z;
+		Result.Extents = FoundBounds.Extents();
+	}
+	else if (BestIdx == 0)
+	{
+		Result.Frame.Rotation = Frame.Rotation;
+		Result.Frame.Origin = Frame.FromFramePoint(LocalBox1.Center());
+		Result.Extents = LocalBox1.Extents();
+	}
+	else // BestIdx == 1
+	{
+		Result.Frame.Rotation = Other.Frame.Rotation;
+		Result.Frame.Origin = Other.Frame.FromFramePoint(LocalBox2.Center());
+		Result.Extents = LocalBox2.Extents();
+	}
+	// Add a small tolerance so that the input boxes are more likely to be properly contained in the merged result, after floating point error
+	Result.Extents += TVector<RealType>(TMathUtil<RealType>::ZeroTolerance);
+	return Result;
+}
+
+template<typename RealType>
+TOrientedBox3<RealType> TOrientedBox3<RealType>::ReparameterizedCloserToWorldFrame() const
+{
+	TVector<RealType> Axes[3]{ AxisX(), AxisY(), AxisZ() };
+
+	int32 BestZ = FMath::Max3Index(
+		FMath::Abs(Axes[0].Z),
+		FMath::Abs(Axes[1].Z),
+		FMath::Abs(Axes[2].Z));
+	TVector<RealType> NewZ = Axes[BestZ] * (Axes[BestZ].Z > 0 ? 1 : -1);
+
+	// Pick the best Y out of the two remaining
+	double DotY[3]{
+		Axes[0].Y,
+		Axes[1].Y,
+		Axes[2].Y
+	};
+	DotY[BestZ] = 0; // don't pick this one
+
+	int32 BestY = FMath::Max3Index(
+		FMath::Abs(DotY[0]),
+		FMath::Abs(DotY[1]),
+		FMath::Abs(DotY[2]));
+	TVector<RealType> NewY = Axes[BestY] * (DotY[BestY] > 0 ? 1 : -1);
+
+	// Sum of BestY and BestZ will be either 0+1=1, 0+2=2, or 1+2=3, and the 
+	// corresponding leftover index will be 2, 1, or 0.
+	int32 BestX = 3 - (BestY + BestZ);
+	// Static analyzer probably won't like that, so here's a clamp for safety
+	BestX = FMath::Clamp(BestX, 0, 2);
+
+	// Sanity check to make sure we picked unique axes
+	if (!ensure(BestX != BestY
+		&& BestY != BestZ
+		&& BestX != BestZ))
+	{
+		return *this;
+	}
+
+	return TOrientedBox3<RealType>(
+		TFrame3<RealType>(Frame.Origin,
+			NewY.Cross(NewZ), // better to do this than risk picking the wrong direction for BestX
+			NewY,
+			NewZ),
+		TVector<RealType>(Extents[BestX], Extents[BestY], Extents[BestZ]));
+}
+
+template struct TOrientedBox3<float>;
+template struct TOrientedBox3<double>;
+
+} // namespace UE::Geometry
+
diff --git a/Engine/Source/Runtime/GeometryCore/Public/OrientedBoxTypes.h b/Engine/Source/Runtime/GeometryCore/Public/OrientedBoxTypes.h
--- a/Engine/Source/Runtime/GeometryCore/Public/OrientedBoxTypes.h
+++ b/Engine/Source/Runtime/GeometryCore/Public/OrientedBoxTypes.h
@@ -4,6 +4,7 @@
 
 #pragma once
 
+#include "Math/OrientedBox.h"
 #include "VectorTypes.h"
 #include "BoxTypes.h"
 #include "FrameTypes.h"
@@ -56,6 +57,29 @@ struct TOrientedBox3
 	{
 	}
 
+	// Convert to Core Math's FOrientedBox
+	explicit operator FOrientedBox() const
+	{
+		FOrientedBox ToRet;
+		ToRet.Center = (FVector)Frame.Origin;
+		ToRet.AxisX = (FVector)Frame.X();
+		ToRet.AxisY = (FVector)Frame.Y();
+		ToRet.AxisZ = (FVector)Frame.Z();
+		ToRet.ExtentX = (double)Extents.X;
+		ToRet.ExtentY = (double)Extents.Y;
+		ToRet.ExtentZ = (double)Extents.Z;
+		return ToRet;
+	}
+
+	// Convert from Core Math's FOrientedBox
+	TOrientedBox3(const FOrientedBox& Box) :
+		Frame(TVector<RealType>(Box.Center),
+			TVector<RealType>(Box.AxisX),
+			TVector<RealType>(Box.AxisY),
+			TVector<RealType>(Box.AxisZ)),
+		Extents((RealType)Box.ExtentX, (RealType)Box.ExtentY, (RealType)Box.ExtentZ)
+	{
+	}
 
 	/** @return box with unit dimensions centered at origin */
 	static TOrientedBox3<RealType> UnitZeroCentered() { return TOrientedBox3<RealType>(TVector<RealType>::Zero(), (RealType)0.5*TVector<RealType>::One()); }
@@ -258,172 +282,13 @@ struct TOrientedBox3
 	 // Uses heuristics to test a few possible orientations and pick the smallest-volume result. Not an optimal fit.
 	 // @param Other						The box to merge with
 	 // @param bOnlyConsiderExistingAxes	If true, the merged box will always use the Rotation of one of the input boxes
-	 TOrientedBox3<RealType> Merge(const TOrientedBox3<RealType>& Other, bool bOnlyConsiderExistingAxes = false) const
-	 {
-		 TQuaternion<RealType> RotInv = Frame.Rotation.Inverse();
-		 TQuaternion<RealType> RotInvOtherRot = RotInv * Other.Frame.Rotation;
-		 TVector<RealType> O2mO1 = Other.Frame.Origin - Frame.Origin;
-		 TVector<RealType> TranslateIntoFrame1 = RotInv * O2mO1;
-		 TVector<RealType> TranslateIntoFrame2 = Other.Frame.Rotation.InverseMultiply(-O2mO1);
-		 TAxisAlignedBox3<RealType> LocalBox1(-Extents, Extents);
-		 TAxisAlignedBox3<RealType> LocalBox2(-Other.Extents, Other.Extents);
-		 // Corners of this box, in its local coordinate frame
-		 TVector<RealType> Corners1[8]
-		 {
-			 LocalBox1.Min,
-			 LocalBox1.Max,
-			 TVector<RealType>(LocalBox1.Min.X, LocalBox1.Min.Y, LocalBox1.Max.Z),
-			 TVector<RealType>(LocalBox1.Min.X, LocalBox1.Max.Y, LocalBox1.Min.Z),
-			 TVector<RealType>(LocalBox1.Max.X, LocalBox1.Min.Y, LocalBox1.Min.Z),
-			 TVector<RealType>(LocalBox1.Min.X, LocalBox1.Max.Y, LocalBox1.Max.Z),
-			 TVector<RealType>(LocalBox1.Max.X, LocalBox1.Min.Y, LocalBox1.Max.Z),
-			 TVector<RealType>(LocalBox1.Max.X, LocalBox1.Max.Y, LocalBox1.Min.Z),
-		 };
-		 // Cordinates of Other, in its local coordinate frame
-		 TVector<RealType> Corners2[8]
-		 {
-			 LocalBox2.Min,
-			 LocalBox2.Max,
-			 TVector<RealType>(LocalBox2.Min.X, LocalBox2.Min.Y, LocalBox2.Max.Z),
-			 TVector<RealType>(LocalBox2.Min.X, LocalBox2.Max.Y, LocalBox2.Min.Z),
-			 TVector<RealType>(LocalBox2.Max.X, LocalBox2.Min.Y, LocalBox2.Min.Z),
-			 TVector<RealType>(LocalBox2.Min.X, LocalBox2.Max.Y, LocalBox2.Max.Z),
-			 TVector<RealType>(LocalBox2.Max.X, LocalBox2.Min.Y, LocalBox2.Max.Z),
-			 TVector<RealType>(LocalBox2.Max.X, LocalBox2.Max.Y, LocalBox2.Min.Z),
-		 };
-		 // Expand each local AABB to contain the other box's points
-		 for (int32 Idx = 0; Idx < 8; ++Idx)
-		 {
-			 LocalBox1.Contain(TranslateIntoFrame1 + RotInvOtherRot * Corners2[Idx]);
-			 LocalBox2.Contain(TranslateIntoFrame2 + RotInvOtherRot.InverseMultiply(Corners1[Idx]));
-		 }
-		 RealType CandidateVolumes[3]{ LocalBox1.Volume(), LocalBox2.Volume(), TMathUtil<RealType>::MaxReal };
-		 TVector<RealType> Axes1[3], Axes2[3], FoundAxes[3];
-		 TAxisAlignedBox3<RealType> FoundBounds;
-		 if (!bOnlyConsiderExistingAxes)
-		 {
-			 // Compute the two local frames
-			 Frame.Rotation.GetAxes(Axes1[0], Axes1[1], Axes1[2]);
-			 Other.Frame.Rotation.GetAxes(Axes2[0], Axes2[1], Axes2[2]);
-			 // Compute metrics re the size of the boxes and how far apart they are
-			 RealType MaxExtentsSizeSq = FMath::Max(Extents.SizeSquared(), Other.Extents.SizeSquared());
-			 RealType CenterDiffSizeSq = O2mO1.SizeSquared();
-			 // For boxes that are far apart relative to their size, fit a box along the line connecting their centers
-			 if (CenterDiffSizeSq > MaxExtentsSizeSq)
-			 {
-				 // find the most-perpendicular original box axis as our second axis
-				 // (Note we could alternatively try interpolating axes, or picking the axis with the largest projected extent ...)
-				 FoundAxes[0] = O2mO1.GetSafeNormal(0.0);
-				 TVector<RealType> MostPerpendicular;
-				 RealType SmallestDot = TMathUtil<RealType>::MaxReal;
-				 for (int32 AxisIdx = 0; AxisIdx < 3; ++AxisIdx)
-				 {
-					 RealType Dot1 = FMath::Abs(Axes1[AxisIdx].Dot(FoundAxes[0]));
-					 RealType Dot2 = FMath::Abs(Axes2[AxisIdx].Dot(FoundAxes[0]));
-					 if (Dot1 < Dot2)
-					 {
-						 if (Dot1 < SmallestDot)
-						 {
-							 SmallestDot = Dot1;
-							 MostPerpendicular = Axes1[AxisIdx];
-						 }
-					 }
-					 else if (Dot2 < SmallestDot)
-					 {
-						 SmallestDot = Dot2;
-						 MostPerpendicular = Axes2[AxisIdx];
-					 }
-				 }
-				 FoundAxes[2] = FoundAxes[0].Cross(MostPerpendicular).GetSafeNormal(0);
-				 FoundAxes[1] = FoundAxes[2].Cross(FoundAxes[0]);
-			 }
-			 else // otherwise, for boxes that are closer together, interpolate their orientations
-			 {
-				 // Find which axes on the Other.Frame.Rotation are closest to the Frame.Rotation's X and Y axes
-				 int32 Axis1X_IdxInAxis2 = 0;
-				 int32 Axis1Y_IdxInAxis2 = 0;
-				 RealType BestXAlignment = FMath::Abs(Axes1[0].Dot(Axes2[0]));
-				 RealType BestYAlignment = FMath::Abs(Axes1[1].Dot(Axes2[0]));
-				 for (int32 Axis2Idx = 1; Axis2Idx < 3; ++Axis2Idx)
-				 {
-					 RealType XAlignment = FMath::Abs(Axes1[0].Dot(Axes2[Axis2Idx]));
-					 RealType YAlignment = FMath::Abs(Axes1[1].Dot(Axes2[Axis2Idx]));
-					 if (XAlignment > BestXAlignment)
-					 {
-						 Axis1X_IdxInAxis2 = Axis2Idx;
-						 BestXAlignment = XAlignment;
-					 }
-					 if (YAlignment > BestYAlignment)
-					 {
-						 Axis1Y_IdxInAxis2 = Axis2Idx;
-						 BestYAlignment = YAlignment;
-					 }
-				 }
-				 if (Axis1X_IdxInAxis2 == Axis1Y_IdxInAxis2)
-				 {
-					 // don't let the closest axis to Y be the same as the closest to X
-					 // (this could happen in a 45-degree case where two axes are tied; should be very rare)
-					 Axis1Y_IdxInAxis2 = (1 + Axis1Y_IdxInAxis2) % 3;
-				 }
-				 RealType SignX = Axes1[0].Dot(Axes2[Axis1X_IdxInAxis2]) < 0 ? -1 : 1;
-				 RealType SignY = Axes1[1].Dot(Axes2[Axis1Y_IdxInAxis2]) < 0 ? -1 : 1;
-				 FoundAxes[0] = (Axes1[0] + SignX * Axes2[Axis1X_IdxInAxis2]).GetSafeNormal();
-				 FoundAxes[1] = (Axes1[1] + SignY * Axes2[Axis1Y_IdxInAxis2]);
-				 // Make sure the axes form an orthonormal frame
-				 FoundAxes[2] = (FoundAxes[0].Cross(FoundAxes[1])).GetSafeNormal();
-				 FoundAxes[1] = FoundAxes[2].Cross(FoundAxes[0]).GetSafeNormal();
-			 }
-
-			 // Fit a local AABB in the Found space
-			 for (int32 CornerIdx = 0; CornerIdx < 8; ++CornerIdx)
-			 {
-				 TVector<RealType> WorldCorner1 = Frame.FromFramePoint(Corners1[CornerIdx]);
-				 TVector<RealType> AxesCorner1(WorldCorner1.Dot(FoundAxes[0]), WorldCorner1.Dot(FoundAxes[1]), WorldCorner1.Dot(FoundAxes[2]));
-				 FoundBounds.Contain(AxesCorner1);
-				 TVector<RealType> WorldCorner2 = Other.Frame.FromFramePoint(Corners2[CornerIdx]);
-				 TVector<RealType> AxesCorner2(WorldCorner2.Dot(FoundAxes[0]), WorldCorner2.Dot(FoundAxes[1]), WorldCorner2.Dot(FoundAxes[2]));
-				 FoundBounds.Contain(AxesCorner2);
-			 }
-			 CandidateVolumes[2] = FoundBounds.Volume();
-		 }
-
-		 // Choose the coordinate space with the smallest bounding box, by volume
-		 int BestIdx = 0;
-		 RealType BestVolume = CandidateVolumes[0];
-		 int AxesToConsider = bOnlyConsiderExistingAxes ? 2 : 3;
-		 for (int Idx = 1; Idx < AxesToConsider; ++Idx)
-		 {
-			 if (CandidateVolumes[Idx] < BestVolume)
-			 {
-				 BestIdx = Idx;
-				 BestVolume = CandidateVolumes[Idx];
-			 }
-		 }
-		 TOrientedBox3<RealType> Result;
-		 if (BestIdx == 2)
-		 {
-			 // Convert the Found Axes into a new oriented box
-			 Result.Frame.Rotation.SetFromRotationMatrix(TMatrix3<RealType>(FoundAxes[0], FoundAxes[1], FoundAxes[2], false));
-			 TVector<RealType> FoundCenter = FoundBounds.Center();
-			 Result.Frame.Origin = FoundAxes[0] * FoundCenter.X + FoundAxes[1] * FoundCenter.Y + FoundAxes[2] * FoundCenter.Z;
-			 Result.Extents = FoundBounds.Extents();
-		 }
-		 else if (BestIdx == 0)
-		 {
-			 Result.Frame.Rotation = Frame.Rotation;
-			 Result.Frame.Origin = Frame.FromFramePoint(LocalBox1.Center());
-			 Result.Extents = LocalBox1.Extents();
-		 }
-		 else // BestIdx == 1
-		 {
-			 Result.Frame.Rotation = Other.Frame.Rotation;
-			 Result.Frame.Origin = Other.Frame.FromFramePoint(LocalBox2.Center());
-			 Result.Extents = LocalBox2.Extents();
-		 }
-		 // Add a small tolerance so that the input boxes are more likely to be properly contained in the merged result, after floating point error
-		 Result.Extents += FVector3d(FMathd::ZeroTolerance, FMathd::ZeroTolerance, FMathd::ZeroTolerance);
-		 return Result;
-	 }
+	 GEOMETRYCORE_API TOrientedBox3<RealType> Merge(const TOrientedBox3<RealType>& Other, bool bOnlyConsiderExistingAxes = false) const;
+
+	 /**
+	  * Transform the frame of the oriented box so its local axes are closer to the corresponding world axes
+	  * (Does not change the shape of the box, just re-orders the axes if needed)
+	  */
+	 [[nodiscard]] GEOMETRYCORE_API TOrientedBox3<RealType> ReparameterizedCloserToWorldFrame() const;
 
 };
 
