fix mesh adapter BuildEdgeConnectivity using the wrong index when adding to spill arrays

diff --git a/Engine/Source/Runtime/GeometryCore/Public/MeshAdapter.h b/Engine/Source/Runtime/GeometryCore/Public/MeshAdapter.h
--- a/Engine/Source/Runtime/GeometryCore/Public/MeshAdapter.h
+++ b/Engine/Source/Runtime/GeometryCore/Public/MeshAdapter.h
@@ -93,7 +93,7 @@ public:
 						{
 							SpillID = SpillIDtoTIDs.Emplace();
 						}
-						SpillIDtoTIDs[FoundEID].Add(TID);
+						SpillIDtoTIDs[SpillID].Add(TID);
 					}
 				}
 				// edge didn't exist yet; create it
