fix incorrect checkslow in dynamic mesh overlay appenddefaulted method

diff --git a/Engine/Source/Runtime/GeometryCore/Public/DynamicMesh/DynamicMeshOverlay.h b/Engine/Source/Runtime/GeometryCore/Public/DynamicMesh/DynamicMeshOverlay.h
--- a/Engine/Source/Runtime/GeometryCore/Public/DynamicMesh/DynamicMeshOverlay.h
+++ b/Engine/Source/Runtime/GeometryCore/Public/DynamicMesh/DynamicMeshOverlay.h
@@ -284,7 +284,7 @@ public:
 	void AppendDefaulted(const FDynamicMesh3::FAppendInfo& AppendInfo)
 	{
 		// Note: ElementRefCounts, Elements and ParentVertices remain unchanged, since the new triangles are unset
-		checkSlow(ElementTriangles.Num() * 3 == AppendInfo.TriangleOffset);
+		checkSlow(ElementTriangles.Num() == AppendInfo.TriangleOffset * 3);
 		ElementTriangles.Resize(3 * (AppendInfo.TriangleOffset + AppendInfo.NumTriangle), IndexConstants::InvalidID);
 	}
 
