PolyBezier: Accelerate knot vector search

Original search was O(n). New search is O(log(n)) w/ caching. Cache hit rate for sampling is measured at 96% (this is the optimal case where the spline is sampled repeatedly with a monotonically increasing parameter value with many samples per segment, random sampling or a low number of samples per segment will have a lower cache hit rate).

diff --git a/Engine/Source/Runtime/GeometryCore/Public/Splines/PolyBezierSpline.h b/Engine/Source/Runtime/GeometryCore/Public/Splines/PolyBezierSpline.h
--- a/Engine/Source/Runtime/GeometryCore/Public/Splines/PolyBezierSpline.h
+++ b/Engine/Source/Runtime/GeometryCore/Public/Splines/PolyBezierSpline.h
@@ -4,6 +4,7 @@
 
 #include "BSpline.h"
 #include "Misc/ScopeExit.h"
+#include "HAL/IConsoleManager.h"
 
 namespace UE
 {
@@ -12,6 +13,12 @@ namespace Geometry
 namespace Spline
 {
 
+inline FAutoConsoleVariable CVarEnablePolyBezierSegmentCache(
+	TEXT("Spline.EnablePolyBezierSegmentCache"),
+	true,
+	TEXT("When true, enables internal caching of the most recently found segment from a global parameter This can greatly accelerate evaluation.")
+);
+	
 /* Augments B-Spline interface */
 template<typename VALUETYPE>
 class TPolyBezierSpline :
@@ -1426,32 +1433,61 @@ public:
 			return true;
 		}
         
-		// Find segment containing parameter
-		for (int32 i = 0; i < Base::PairKnots.Num() - 1; ++i)
+		int32 SegmentIndex = INDEX_NONE;
+
+		if (CVarEnablePolyBezierSegmentCache->GetBool() && LastEvaluatedSegment != INDEX_NONE && LastEvaluatedSegment < Base::PairKnots.Num() - 1)
 		{
-			if (GlobalParam >= Base::PairKnots[i].Value && GlobalParam < Base::PairKnots[i + 1].Value)
+			if (GlobalParam < Base::PairKnots[LastEvaluatedSegment + 1].Value && GlobalParam >= Base::PairKnots[LastEvaluatedSegment].Value)
 			{
-				OutSegmentIndex = i;
-                
-				// Calculate local parameter
-				float SegmentLength = Base::PairKnots[i + 1].Value - Base::PairKnots[i].Value;
-				if (SegmentLength > UE_SMALL_NUMBER)
+				// cache hit!
+				SegmentIndex = LastEvaluatedSegment;
+			}
+		}
+		
+		if (SegmentIndex == INDEX_NONE)
+		{
+			// cache miss!
+			int32 SegmentLow = 0;
+			int32 SegmentHigh = Base::PairKnots.Num() - 2;
+			SegmentIndex = SegmentLow + (SegmentHigh - SegmentLow) / 2;
+
+			while (SegmentLow <= SegmentHigh)
+			{
+				int32 Mid = SegmentLow + (SegmentHigh - SegmentLow) / 2;
+				const float& StartParam = Base::PairKnots[Mid].Value;
+				const float& EndParam   = Base::PairKnots[Mid + 1].Value;
+
+				if (GlobalParam >= StartParam && GlobalParam < EndParam)
 				{
-					OutLocalParam = (GlobalParam - Base::PairKnots[i].Value) / SegmentLength;
+					SegmentIndex = Mid;
+					break;
+				}
+				else if (GlobalParam < StartParam)
+				{
+					SegmentHigh = Mid - 1;
 				}
 				else
 				{
-					OutLocalParam = 0.0f;
+					SegmentLow = Mid + 1;
 				}
-                
-				return true;
 			}
 		}
+		
+		OutSegmentIndex = SegmentIndex;
+		LastEvaluatedSegment = SegmentIndex;
         
-		// Should never reach here
-		OutSegmentIndex = 0;
-		OutLocalParam = 0.0f;
-		return false;
+		// Calculate local parameter
+		float SegmentLength = Base::PairKnots[SegmentIndex + 1].Value - Base::PairKnots[SegmentIndex].Value;
+		if (SegmentLength > UE_SMALL_NUMBER)
+		{
+			OutLocalParam = (GlobalParam - Base::PairKnots[SegmentIndex].Value) / SegmentLength;
+		}
+		else
+		{
+			OutLocalParam = 0.0f;
+		}
+        
+		return true;
 	}
 
 	/**
@@ -1946,6 +1982,10 @@ private:
 		Base::Dump();
 		return SegmentIndex;
 	}
+
+private:
+
+	mutable int32 LastEvaluatedSegment = INDEX_NONE;
 };
 
 
