PolyBezier: Remove problematic inline CVar

this is being registered multiple times due to being in a header.

diff --git a/Engine/Source/Runtime/GeometryCore/Public/Splines/PolyBezierSpline.h b/Engine/Source/Runtime/GeometryCore/Public/Splines/PolyBezierSpline.h
--- a/Engine/Source/Runtime/GeometryCore/Public/Splines/PolyBezierSpline.h
+++ b/Engine/Source/Runtime/GeometryCore/Public/Splines/PolyBezierSpline.h
@@ -13,12 +13,6 @@ namespace Geometry
 namespace Spline
 {
 
-inline FAutoConsoleVariable CVarEnablePolyBezierSegmentCache(
-	TEXT("Spline.EnablePolyBezierSegmentCache"),
-	true,
-	TEXT("When true, enables internal caching of the most recently found segment from a global parameter This can greatly accelerate evaluation.")
-);
-	
 /* Augments B-Spline interface */
 template<typename VALUETYPE>
 class TPolyBezierSpline :
@@ -1435,7 +1429,9 @@ public:
         
 		int32 SegmentIndex = INDEX_NONE;
 
-		if (CVarEnablePolyBezierSegmentCache->GetBool() && LastEvaluatedSegment != INDEX_NONE && LastEvaluatedSegment < Base::PairKnots.Num() - 1)
+		// todo: make CVar
+		constexpr bool bEnableSegmentCaching = true;
+		if (bEnableSegmentCaching && LastEvaluatedSegment != INDEX_NONE && LastEvaluatedSegment < Base::PairKnots.Num() - 1)
 		{
 			if (GlobalParam < Base::PairKnots[LastEvaluatedSegment + 1].Value && GlobalParam >= Base::PairKnots[LastEvaluatedSegment].Value)
 			{
