Splines: Knot vector fixes & optimization

1) All knot insertion goes through InsertKnot

2) When computing knot values, we can safely assume that Knot + Epsilon > Knot using the SafeAdd function. This was causing knot vector corruption before for small segments

3) No more min segment length, knot spacing is a function of the magnitude of the nearby values. We can hit the theoretical limit for knot vector density (this was not possible with min segment length).

4) Replaced several TArrays with TStaticArrays.

5) Computing valid knot values is now configurable. We can search only left, only right, or in both directions. This allows us to make better assumptions when reparameterizing (we can trust that the knot vector order matches our point order even if we have collisions by disallowing leftward searching)

diff --git a/Engine/Source/Runtime/GeometryCore/Public/Splines/BSpline.h b/Engine/Source/Runtime/GeometryCore/Public/Splines/BSpline.h
--- a/Engine/Source/Runtime/GeometryCore/Public/Splines/BSpline.h
+++ b/Engine/Source/Runtime/GeometryCore/Public/Splines/BSpline.h
@@ -116,6 +116,9 @@ public:
 	static constexpr int32 Degree = DEGREE;
 	using ValueType = typename TSplineInterface<VALUETYPE>::ValueType;
 
+	static constexpr int32 WindowSize = Degree + 1;
+	using FWindow = TStaticArray<const ValueType*, WindowSize>;
+	
 	DECLARE_SPLINE_TYPE_ID(
 		BSplineNameSelector<DEGREE>::Name,
 		*TSplineValueTypeTraits<VALUETYPE>::Name
@@ -248,7 +251,7 @@ public:
 		if (NumValues == 1) return GetValue(0);
 
 		// Get window using override or evaluator
-		TArray<const ValueType*> Window = FindWindow(Parameter);
+		FWindow Window = FindWindow(Parameter);
     	
 		if (Window.Num() == 0) return ValueType();
         
@@ -574,6 +577,7 @@ protected:
 		{
 			if (PairKnots[i].Value == InKnot.Value)  // exact match
 			{
+				ensureAlwaysMsgf(false, TEXT("Knot insertion failed, knot already exists!"));
 				return;
 			}
 			else if (PairKnots[i].Value > InKnot.Value)
@@ -588,27 +592,44 @@ protected:
 		MarkFlatKnotsCacheDirty();
 	}
 
+	struct FValidKnotSearchParams
+	{
+		FValidKnotSearchParams() = default;
+		
+		FValidKnotSearchParams(float InDesiredParameter)
+			: DesiredParameter(InDesiredParameter)
+		{}
+		
+		float DesiredParameter = 0.f;
+		bool bSearchRight = true;
+		bool bSearchLeft = true;
+	};
+	
 	/**
 	 * Finds the nearest available knot value that does not conflict with existing knots.
-	 * @param DesiredParameter - The desired parameter value to find a valid knot for.
+	 * @param InSearchParams - Describes value to search for and how to search.
 	 * @return The nearest valid knot parameter that does not conflict with existing knots.
 	 */
-	float GetNearestAvailableKnotValue(float DesiredParameter) const
+	float GetNearestAvailableKnotValue(const FValidKnotSearchParams& InSearchParams) const
     {
-        if (PairKnots.Num() == 0)
+		const float& DesiredParameter = InSearchParams.DesiredParameter;
+		
+        if (!ensureAlwaysMsgf(!std::isnan(DesiredParameter), TEXT("DesiredParameter is NaN!")) ||
+        	!ensureAlwaysMsgf(InSearchParams.bSearchLeft || InSearchParams.bSearchRight, TEXT("Invalid search params!")) ||
+        	PairKnots.Num() == 0)
         {
             return DesiredParameter;
         }
 
-        // Get valid parameter space
-        FInterval1f ParamSpace = GetParameterSpace();
-
 		// Convert existing knots into a set for O(1) lookups
 		TSet<float> KnotValuesSet;
-		for (const FKnot& Knot : PairKnots)
 		{
-			KnotValuesSet.Add(Knot.Value);
+			for (const FKnot& Knot : PairKnots)
+			{
+				KnotValuesSet.Add(Knot.Value);
+			}
 		}
+
 		
 		// If no conflict, return immediately
 		if (!KnotValuesSet.Contains(DesiredParameter))
@@ -626,22 +647,29 @@ protected:
 		while (IterationCount < MaxIterations)
 		{
 			IterationCount++;
-			
-        	Left = std::nextafter(Left, -std::numeric_limits<float>::infinity());
-        	if (!KnotValuesSet.Contains(Left))
-        	{
-        		return Left;
-        	}
-        	Right = std::nextafter(Right, std::numeric_limits<float>::infinity());
-        	if (!KnotValuesSet.Contains(Right))
-        	{
-        		return Right;
-        	}
+
+			if (InSearchParams.bSearchLeft)
+			{
+				Left = std::nextafter(Left, -std::numeric_limits<float>::infinity());
+				if (!KnotValuesSet.Contains(Left))
+				{
+					return Left;
+				}
+			}
+
+			if (InSearchParams.bSearchRight)
+			{
+				Right = std::nextafter(Right, std::numeric_limits<float>::infinity());
+				if (!KnotValuesSet.Contains(Right))
+				{
+					return Right;
+				}
+			}
         }
 		
 		// If no valid parameter was found within limits, return the original (or handle error)
         // This should rarely be hit unless space is densely packed
-        ensureMsgf(false, TEXT("Max iterations reached in GetNearestAvailableKnotValue"));
+        ensureAlwaysMsgf(false, TEXT("Max iterations reached in GetNearestAvailableKnotValue"));
         return DesiredParameter;
     }
 
@@ -1034,9 +1062,9 @@ private:
 		return Degree;
 	}
 
-    virtual TArray<const ValueType*> FindWindow(float Parameter) const
+    virtual FWindow FindWindow(float Parameter) const
     {
-    	TArray<const ValueType*> Window;
+    	FWindow Window;
     	if (NumKeys() < Degree + 1)
     	{
     		return Window;
@@ -1054,7 +1082,7 @@ private:
     		{
     			int32 Index = (Span - Degree + i) % NumPoints;
     			if (Index < 0) Index += NumPoints;
-    			Window.Add(&GetValue(Index));
+    			Window[i] = &GetValue(Index);
     		}
     	}
     	else
@@ -1065,7 +1093,7 @@ private:
     			const int32 Index = Span - Degree + i;
     			if (Index >= 0 && Index < NumKeys())
     			{
-    				Window.Add(&GetValue(Index));
+    				Window[i] = &GetValue(Index);
     			}
     		}
     	}
diff --git a/Engine/Source/Runtime/GeometryCore/Public/Splines/PolyBezierSpline.h b/Engine/Source/Runtime/GeometryCore/Public/Splines/PolyBezierSpline.h
--- a/Engine/Source/Runtime/GeometryCore/Public/Splines/PolyBezierSpline.h
+++ b/Engine/Source/Runtime/GeometryCore/Public/Splines/PolyBezierSpline.h
@@ -20,8 +20,9 @@ class TPolyBezierSpline :
 	private TSelfRegisteringSpline<TPolyBezierSpline<VALUETYPE>, VALUETYPE>
 {
 public:
-	typedef typename TBSpline<VALUETYPE, 3>::ValueType ValueType;
 	using Base = TBSpline<VALUETYPE, 3>;
+	typedef typename Base::ValueType ValueType;
+	using FWindow = typename Base::FWindow;
 
 	// Generate compile-time type ID for PolyBezier
 	DECLARE_SPLINE_TYPE_ID(
@@ -260,7 +261,7 @@ public:
 	virtual void Clear() override { Base::Clear(); }
 	virtual TUniquePtr<ISplineInterface> Clone() const override
     {
-        TUniquePtr<TPolyBezierSpline<VALUETYPE>> Clone = MakeUnique<TPolyBezierSpline<VALUETYPE>>();
+        TUniquePtr<TPolyBezierSpline> Clone = MakeUnique<TPolyBezierSpline>();
         
         // Copy base class members
         Clone->Values = this->Values;
@@ -284,8 +285,7 @@ public:
     	}
 
     	// Get Bezier control points for this segment
-    	TArray<ValueType> Coeffs;
-    	Coeffs.SetNum(4);
+    	TStaticArray<ValueType, 4> Coeffs;
     	const int32 BaseIndex = SegmentIndex * 4;
         
     	// Get the control points
@@ -300,13 +300,10 @@ public:
     	Coeffs[2] = (P2 - P1*2 + P0) * 3;          // quadratic coefficient  
     	Coeffs[3] = P3 - P2*3 + P1*3 - P0;         // cubic coefficient
 
-    	float LocalT = Math::FindNearestPoint_Cubic(Coeffs, 0.0f, 1.0f, OutSquaredDistance);
-    	// Get the actual parameter bounds for this segment from the knot vector
-		
-    	float SplineParam = MapLocalSegmentParameterToGlobal(SegmentIndex, LocalT);
-    
-    	return SplineParam;
-	}
+    	const float LocalT = Math::FindNearestPoint_Cubic(MakeArrayView(Coeffs), 0.0f, 1.0f, OutSquaredDistance);
+    	return MapLocalSegmentParameterToGlobal(SegmentIndex, LocalT);
+    }
+	
 	virtual float FindNearest(const ValueType& Point, float& OutSquaredDistance) const override
 	{
     	float BestDistSq = TNumericLimits<float>::Max();
@@ -351,18 +348,14 @@ public:
     	// Convert global parameter to segment information
     	int32 SegmentIndex;
     	float LocalT;
-    	bool bValidSegment = MapGlobalParameterToLocalSegment(Parameter, SegmentIndex, LocalT);
-    
+    	MapGlobalParameterToLocalSegment(Parameter, SegmentIndex, LocalT);
 		
 		// Get window for the current parameter
-	    TArray<const ValueType*> Window = FindWindow(Parameter);
+	    FWindow Window = FindWindow(Parameter);
 	    if (Window.Num() < 4)  // Need all 4 control points for a cubic Bezier
 	    {
 	        return ValueType();
 	    }
-
-    	int32 NumSegments = GetNumberOfSegments();
-    	// Map parameter to parent space
     	
 	    // For PolyBezier, each segment is normalized to unit parameter space
         constexpr float SegmentScale = 1.0f;
@@ -771,15 +764,15 @@ public:
         const ValueType& Position,
         EParameterizationPolicy ParameterizationPolicy = EParameterizationPolicy::Centripetal)
     {
-		int32 NewPointIndex;
-		if (InitializeFirstSegment(Parameter, Position, NewPointIndex)) return NewPointIndex;
- 
-        // Find segment to split
-        const int32 NumSegments = Base::NumKeys() / 4;
-		// Cast to proper parameterization type
+		// todo: stop returning bool, instead interpret INDEX_NONE as failure
+	    if (int32 NewPointIndex;
+	    	InitializeFirstSegment(Parameter, Position, NewPointIndex))
+		{
+			return NewPointIndex;
+		}
 
 		// Make sure we have a valid parameter
-		Parameter =  Base::GetNearestAvailableKnotValue(Parameter);
+		Parameter = Base::GetNearestAvailableKnotValue(Parameter);
 		
         FInterval1f ParameterRange = GetParameterSpace();
 		int32 SegmentIndex;
@@ -821,9 +814,17 @@ public:
         ValueType R3 = P3;
         
         FInterval1f SegmentRange = GetSegmentParameterRange(SegmentIndex);
-		float OriginalStartParam = SegmentRange.Min;
-		float OriginalEndParam = SegmentRange.Max;
 		float TargetSplitParam = SegmentRange.Interpolate(LocalT);
+
+		if (TargetSplitParam == SegmentRange.Min)
+		{
+			TargetSplitParam = std::nextafter(TargetSplitParam, std::numeric_limits<float>::infinity());
+		}
+
+		if (TargetSplitParam == SegmentRange.Max)
+		{
+			TargetSplitParam = std::nextafter(TargetSplitParam, -std::numeric_limits<float>::infinity());
+		}
 		
 		for (int32 i = 0; i < 4; ++i)
 		{
@@ -834,11 +835,9 @@ public:
         AddBezierSegmentInternal(SegmentIndex, Q0, Q1, Q2, Q3, ParameterizationPolicy);
 		Base::InsertKnot( FKnot(TargetSplitParam, this->bClampEnds ? this->GetDegree() : 1));
 		int32 IndexToReturn = AddBezierSegmentInternal(SegmentIndex + 1, R0, R1, R2, R3, ParameterizationPolicy);
-
 		
 		return IndexToReturn;
     }
-
 	
     /**
      * Inserts a Bezier segment at the specified parameter
@@ -1080,30 +1079,17 @@ public:
 	 */
 	virtual void Reparameterize(EParameterizationPolicy Mode = EParameterizationPolicy::Centripetal) override
 	{
-		TArray<ValueType> Points = this->Values;
+		const TArray<ValueType>& Points = this->Values;
 		const int32 NumValues = Points.Num();
 
-		// const int32 NumKnots = this->GetExpectedNumKnots();
-                
-		switch (Mode)
-		{
-		case EParameterizationPolicy::ChordLength:
-		case EParameterizationPolicy::Centripetal:
-			GenerateDistanceKnotsForBezier(Mode);
-			break;
-            
-		case EParameterizationPolicy::Uniform:
-		default:
-			GenerateBezierKnots();
-			break;
-                
-		}
-    	
+		GenerateDistanceKnotsForBezier(Mode);
     	
 		UE_LOG(LogSpline, Verbose, TEXT("Set knot vector - Points: %d, Knots: %d, Mode: %d"),
 					   NumValues,  this->PairKnots.Num(), static_cast<int32>(Mode));
+		
 		Base::PrintKnotVector();
 	}
+	
 	virtual FInterval1f GetParameterSpace() const override
 	{
 		return Base::PairKnots.Num() > 0
@@ -1241,7 +1227,6 @@ public:
 	
 	void FlipSegment(int32 Segment)
 	{
-		static constexpr int32 Degree = 3; // Cubic
 		static constexpr int32 ControlPointsPerSegment = 4;
 		
 		const int32 NumSegments = GetNumberOfSegments();
@@ -1305,6 +1290,8 @@ public:
 	{
 		int32 NumSegments = Base::PairKnots.Num() - 1;
 
+		// todo: either delete this function in favor of MapGlobalParameterToLocalSegment or fix the loop to binary search with caching
+		
 		for (int32 i = 0; i < NumSegments; ++i)
 		{
 			float Start = Base::PairKnots[i].Value;
@@ -1330,6 +1317,7 @@ public:
 		OutLocalParam = 0.0f;
 		return 0; // Fallback to first segment
 	}
+	
 	/**
 	 * Sets the closed loop flag without modifying segments
 	 * @param bClosed Whether the spline should be a closed loop
@@ -1513,6 +1501,7 @@ public:
 		}
 		return 4 + (NumSegments - 1) * 3 + 4;
 	}
+	
 	void SetKnotVector(const TArray<FKnot>& NewKnots)
 	{
 		Base::SetCustomKnots(NewKnots);
@@ -1532,11 +1521,9 @@ private:
 		// For a Bezier spline, the control points don't all lie on the curve
 		// So we need to compute chord lengths between segment endpoints only
 		int NumSegments = Points.Num() / 4;
-		constexpr float MinSegmentLength = UE_KINDA_SMALL_NUMBER;
 		
 		TArray<float> SegmentLengths;
 		SegmentLengths.Reserve(NumSegments);
-		float TotalLength = 0.0f;
         
 		for (int32 i = 0; i < NumSegments; ++i)
 		{
@@ -1545,22 +1532,7 @@ private:
             
 			if (EndIndex < Points.Num())
 			{
-				double Length;
-                
-				// Use appropriate distance calculation based on knot mode
-				if (Mode == EParameterizationPolicy::Centripetal)
-				{
-					// Square root of distance for centripetal mode
-					Length = Math::CentripetalDistance(Points[StartIndex], Points[EndIndex]);
-				}
-				else // Chord length mode
-				{
-					// Actual distance for chord length mode
-					Length = Math::Distance(Points[EndIndex], Points[StartIndex]);
-				}
-				float ClampedLength = FMath::Max(static_cast<float>(Length), MinSegmentLength);
-				SegmentLengths.Add(static_cast<float>(ClampedLength));
-				TotalLength += static_cast<float>(ClampedLength);
+				SegmentLengths.Add(ComputeSegmentLength(Points[StartIndex], Points[EndIndex], Mode));
 			}
 		}
         
@@ -1568,13 +1540,30 @@ private:
 		// 1. Bezier segment structure (knot multiplicity)
 		// 2. Proportional chord/centripetal lengths
         
+		// This will only return Val if Offset is exactly 0.f. Otherwise, it is guaranteed to change in the direction of Offset.
+		// We sacrifice a bit of accuracy for the assumption that Val will actually change.
+		auto SafeAdd = [](float Val, float Offset) -> float
+		{
+			if (Offset == 0.f) return Val;
+				
+			const float Result = Val + Offset;
+				
+			return Result == Val
+				? std::nextafter(Val, Offset > 0.f
+					? std::numeric_limits<float>::infinity()
+					: -std::numeric_limits<float>::infinity())
+				: Result;
+		};
+		
 		// Reset pair knots
 		this->PairKnots.Reset();
         
-		int32 CurrentKnot = 0;
+		typename Base::FValidKnotSearchParams SearchParams;
+		SearchParams.bSearchLeft = false;	// we are appending repeatedly, we want insertion order to be predictable (always go after conflicting knots).
 
 		// Add start knot with multiplicity 4
-		this->PairKnots.Add(FKnot(0.0f, 4));
+		SearchParams.DesiredParameter = 0.f;
+		this->InsertKnot(FKnot(this->GetNearestAvailableKnotValue(SearchParams), 4));
         
 		// Calculate accumulated distances for internal knots
 		float AccumLength = 0.0f;
@@ -1582,66 +1571,29 @@ private:
 		{
 			if (i < SegmentLengths.Num())
 			{
-				AccumLength += SegmentLengths[i];
+				AccumLength = SafeAdd(AccumLength, SegmentLengths[i]);
                 
-				// Use normalized accumulated length as the knot value
-				const float t = AccumLength;
+				// Use normalized accumulated length as the desired knot value
+				SearchParams.DesiredParameter = AccumLength;
 				
 				// Interior knots at segment boundaries with multiplicity 3
-				this->PairKnots.Add(FKnot(t, 3));
+				this->InsertKnot(FKnot(this->GetNearestAvailableKnotValue(SearchParams), 3));
 			}
 		}
 
-		// End knot with multiplicity 4
-		const float EndKnot = AccumLength + (SegmentLengths.Num() > 0 ? SegmentLengths.Last() : 0.0f);
-		this->PairKnots.Add(FKnot(EndKnot, 4));
-
-		// Mark flat knots cache as dirty
-		this->MarkFlatKnotsCacheDirty();
-	}
-
-	/**
-	 * Generates a standard Bezier-compatible knot vector
-	 */
-	void GenerateBezierKnots()
-	{
-		int32 NumSegments = GetNumberOfSegments();
-		// For Bezier curves:
-		// - Each segment has degree 3
-		// - For closed loop, we need n + 2*p + 1 knots where:
-		//   n = number of control points
-		//   p = degree (3 for cubic Bezier)
-		const int32 NumKnots = this->GetExpectedNumKnots();             // normal case: n + p + 1 where p = 3
-        
-		// Set up knot vector with multiple knots at segment boundaries
-		this->PairKnots.Reset();
-        
-		int32 CurrentKnot = 0;
-
-		// Start knots with multiplicity 4
-		this->PairKnots.Add( FKnot(0.0f, 4));
-
-		// For closed loop, add an extra segment to connect back to start
-		const int32 NumSegmentsToProcess = this->IsClosedLoop() ? (NumSegments + 1) : NumSegments;
-
-		// Interior knots at segment boundaries with multiplicity 3
-		for (int32 i = 1; i < NumSegmentsToProcess; ++i)
-		{
-			const float t = static_cast<float>(i);
-			this->PairKnots.Add(FKnot(t, 3));
-		}
+		SearchParams.DesiredParameter = SafeAdd(AccumLength, SegmentLengths.Num() > 0
+			? SegmentLengths.Last()
+			: 0.0f);
 
-		// End knot with multiplicity 4
-		const float EndKnot = static_cast<float>(NumSegmentsToProcess);
-		this->PairKnots.Add(FKnot(EndKnot, 4));
+		this->InsertKnot(FKnot(this->GetNearestAvailableKnotValue(SearchParams), 4));
 
 		// Mark flat knots cache as dirty
 		this->MarkFlatKnotsCacheDirty();
 	}
 
-    virtual TArray<const ValueType*> FindWindow(float Parameter) const override
+    virtual FWindow FindWindow(float Parameter) const override
     {
-        TArray<const ValueType*> Window;
+        FWindow Window;
         if (Base::NumKeys() < 4)  // Need at least one complete Bezier segment
         {
             return Window;
@@ -1653,9 +1605,6 @@ private:
        	int32 SegmentIndex;
     	float LocalT;
     	MapGlobalParameterToLocalSegment(Parameter, SegmentIndex, LocalT);
-    	// Handle end of spline case
-    	//int32 NumSegmentsDistinct = GetNumDistinctSegments();
-    	//check (NumSegmentsDistinct == NumSegments)
 
         if (Base::IsClosedLoop())
         {
@@ -1666,7 +1615,6 @@ private:
         
         // Each Bezier segment has exactly 4 control points
         const int32 BaseIndex = SegmentIndex * 4;
-        Window.Reserve(4);
 
         if (Base::IsClosedLoop())
         {
@@ -1677,10 +1625,10 @@ private:
                 return ((Index % NumPoints) + NumPoints) % NumPoints;
             };
 
-            Window.Add(&Entries[WrapIndex(BaseIndex)]);     // P0 
-            Window.Add(&Entries[WrapIndex(BaseIndex + 1)]); // P1
-            Window.Add(&Entries[WrapIndex(BaseIndex + 2)]); // P2
-            Window.Add(&Entries[WrapIndex(BaseIndex + 3)]); // P3
+			Window[0] = &Entries[WrapIndex(BaseIndex + 0)];
+        	Window[1] = &Entries[WrapIndex(BaseIndex + 1)];
+        	Window[2] = &Entries[WrapIndex(BaseIndex + 2)];
+        	Window[3] = &Entries[WrapIndex(BaseIndex + 3)];
         }
         else
         {
@@ -1689,17 +1637,18 @@ private:
             {
                 return Window;
             }
-            Window.Add(&Entries[BaseIndex]);     // P0 
-            Window.Add(&Entries[BaseIndex + 1]); // P1
-            Window.Add(&Entries[BaseIndex + 2]); // P2
-            Window.Add(&Entries[BaseIndex + 3]); // P3
+
+        	Window[0] = &Entries[BaseIndex + 0];
+        	Window[1] = &Entries[BaseIndex + 1];
+        	Window[2] = &Entries[BaseIndex + 2];
+        	Window[3] = &Entries[BaseIndex + 3];
         }
 
         return Window;
     }
 
     virtual ValueType InterpolateWindow(TArrayView<const ValueType* const> Window, float Parameter) const override 
-    {
+    {		
         if (Window.Num() < 4)
         {
             return ValueType();
@@ -1735,7 +1684,6 @@ private:
 	/**
 	 * Finds the local parameter on a segment closest to the given position
 	 * 
-	 * @param Storage - Storage to search in
 	 * @param SegmentIndex Index of the segment
 	 * @param Position Position to find closest point to
 	 * @param OutSquaredDistance Output parameter for squared distance
@@ -1761,14 +1709,14 @@ private:
     	FInterval1f SegmentRange = GetSegmentParameterRange(SegmentIndex);
     	float SegmentLength = SegmentRange.Max - SegmentRange.Min;
     
-    	if (SegmentLength < UE_KINDA_SMALL_NUMBER)
+    	if (SegmentLength == 0.f)
     		return 0.0f;
         
     	return (GlobalParam - SegmentRange.Min) / SegmentLength;
     }
+	
 	/**
 	 * Checks if a segment index is valid
-	 * @param Storage - Storage to check against
 	 * @param Index - Index to check
 	 * @return true if the index is valid
 	 */
@@ -1833,9 +1781,6 @@ private:
 				SetClosedLoop(true);
 			}
 		};
-		
-		// Get current parameter range start
-		FInterval1f ParameterRange = GetParameterSpace();
     
 		// Create a control point at the beginning
 		ValueType P3 = Base::GetValue(0); // First point of the spline
@@ -1875,7 +1820,6 @@ private:
     		
 		// Append - create a point at the end
 		int32 LastIndex = Base::NumKeys() - 1;
-		float StartParam = Base::GetParameter(LastIndex);
 		ValueType P0 = Base::GetValue(LastIndex);
         
 		// Calculate reasonable auto-tangents
@@ -1898,14 +1842,12 @@ private:
 		const ValueType& P3,
 		EParameterizationPolicy ParameterizationPolicy)
     {
-		constexpr float MinSegmentLength = UE_KINDA_SMALL_NUMBER;
-		
     	switch (ParameterizationPolicy)
     	{
     	case EParameterizationPolicy::ChordLength:
-    		return static_cast<float>(FMath::Max(Math::Distance(P0, P3), MinSegmentLength));
+    		return static_cast<float>(Math::Distance(P0, P3));
     	case EParameterizationPolicy::Centripetal:
-    		return static_cast<float>(FMath::Max(Math::CentripetalDistance(P0, P3), MinSegmentLength));
+    		return static_cast<float>(Math::CentripetalDistance(P0, P3));
     	case EParameterizationPolicy::Uniform:
     	default:
 			return 1.0f;
@@ -1955,7 +1897,7 @@ private:
 		{
 			this->ResetKnotVector();
 			this->InsertKnot( FKnot(0.0f, this->bClampEnds ? this->GetDegree() + 1 : 1));
-			this->InsertKnot( FKnot(SegmentLength, this->bClampEnds ? this->GetDegree() + 1 : 1));
+			this->InsertKnot( FKnot(1.f, this->bClampEnds ? this->GetDegree() + 1 : 1));
 		}
 		else if (bAppendToEnd)
 		{
@@ -1981,6 +1923,7 @@ private:
 
 private:
 
+	/** This is the last segment fetched by parameter. Used to accelerate predictable sampling patterns. */
 	mutable int32 LastEvaluatedSegment = INDEX_NONE;
 };
 
diff --git a/Engine/Source/Runtime/GeometryCore/Public/Splines/SplineMath.h b/Engine/Source/Runtime/GeometryCore/Public/Splines/SplineMath.h
--- a/Engine/Source/Runtime/GeometryCore/Public/Splines/SplineMath.h
+++ b/Engine/Source/Runtime/GeometryCore/Public/Splines/SplineMath.h
@@ -618,7 +618,7 @@ struct TBezierDerivativeCalculator
 
 template<typename ValueType>
 float FindNearestPoint_Cubic(
-	const TArray<ValueType>& Coeffs,
+	const TArrayView<ValueType> Coeffs,
 	float StartParam,
 	float EndParam, 
 	float& OutSquaredDistance)
