Add oriented box intersection tests to TOrientedBox3 (covering OBB-OBB and OBB-AABB cases)

diff --git a/Engine/Source/Runtime/GeometryCore/Private/OrientedBoxTypes.cpp b/Engine/Source/Runtime/GeometryCore/Private/OrientedBoxTypes.cpp
--- a/Engine/Source/Runtime/GeometryCore/Private/OrientedBoxTypes.cpp
+++ b/Engine/Source/Runtime/GeometryCore/Private/OrientedBoxTypes.cpp
@@ -5,10 +5,153 @@
 #include "Math/UnrealMath.h"
 
 
+namespace UE::OrientedBoxTypeLocals
+{
+	using namespace UE::Math;
+	using namespace UE::Geometry;
+
+	/**
+	 * Helper to test box intersection in the local space of Box A -- i.e., where Box A is centered at the origin, and axis-aligned.
+	 * 
+	 * Assumes extents are not negative.
+	 * Applies 15 separating axis tests: 3 per axis of the two boxes, and 9 per cross product of the box axes.
+	 * If the boxes are aligned on any axis, the 9 cross-edge tests are not needed (and would have robustness issues for the aligned-axis cross product)
+	 */
+	template<typename RealType>
+	static bool BoxIntersectsLocalSpaceHelper(
+		const TVector<RealType>& ExtentsA, 
+		const TVector<RealType>& CenterB, const TVector<RealType> (&AxesB)[3], const TVector<RealType>& ExtentsB, 
+		RealType ParallelAxisTol = TMathUtil<RealType>::ZeroTolerance)
+	{
+		// Convert zero tolerance to a dot-product tolerance -- i.e., take 1-Tolerance
+		RealType DotParallelAxisTol = (RealType)1 - FMath::Max((RealType)0, ParallelAxisTol);
+
+		// Absolute values of all B axes, computed in the loop below.
+		// Used to project the farthest distance from box center to corner, along the separating axis test direction
+		TVector<RealType> AbsAxesB[3]{};
+
+		bool bHasParallelAxis = false;
+
+		// Test if axes of box A (i.e. the major axes) are separating
+		for (int32 AxisA = 0; AxisA < 3; ++AxisA)
+		{
+			// Projected BoxA radius along its own axis -- just the extent
+			RealType ProjRadiusA = ExtentsA[AxisA];
+
+			// BoxB radius along A axis (Dot product of AxisA and largest box corner vector)
+			// Compute absolute values of AxesB as well (for reuse in later tests), and test for alignment w/ a major axis
+			RealType ProjRadiusB = 0;
+			int32 NumZeroes = 0;
+			for (int32 AxisB = 0; AxisB < 3; ++AxisB)
+			{
+				AbsAxesB[AxisB][AxisA] = FMath::Abs(AxesB[AxisB][AxisA]);
+				bHasParallelAxis |= (AbsAxesB[AxisB][AxisA] >= DotParallelAxisTol);
+				ProjRadiusB += ExtentsB[AxisB] * AbsAxesB[AxisB][AxisA];
+			}
+			
+			// Projected distance between centers
+			RealType ProjCenterDist = FMath::Abs(CenterB[AxisA]);
+
+			if (ProjCenterDist > ProjRadiusA + ProjRadiusB)
+			{
+				return false;
+			}
+		}
+
+		// Test if axes of box B are separating
+		for (int32 AxisB = 0; AxisB < 3; ++AxisB)
+		{
+			// Projected BoxB radius along its own axis -- just the extent
+			RealType ProjRadiusB = ExtentsB[AxisB];
+
+			// BoxA radius along B axis
+			RealType ProjRadiusA = ExtentsA.Dot(AbsAxesB[AxisB]);
+
+			// Projected distance between centers
+			RealType ProjCenterDist = FMath::Abs(CenterB.Dot(AxesB[AxisB]));
+
+			if (ProjCenterDist > ProjRadiusB + ProjRadiusA)
+			{
+				return false;
+			}
+		}
+
+		// if boxes have a parallel axis, we don't need to test cross axes -- since we haven't found a separating axis above, the boxes intersect
+		// (this is equivalent to projecting along the shared axis, and doing 2D OBB intersection, where just major axes are sufficient)
+		if (bHasParallelAxis)
+		{
+			return true;
+		}
+
+		// Test the 9 cross-product axes
+		for (int32 AxisA = 0, OffA1 = 1, OffA2 = 2; AxisA < 3; OffA1 = OffA2, OffA2 = AxisA++)
+		{
+			for (int32 AxisB = 0, OffB1 = 1, OffB2 = 2; AxisB < 3; OffB1 = OffB2, OffB2 = AxisB++)
+			{
+				// AxisA x AxisB projections, simplified due to AxisA being a major axis
+
+				// Projected BoxA radius along AxisA x AxisB
+				RealType ProjRadiusA = ExtentsA[OffA1] * AbsAxesB[AxisB][OffA2] + ExtentsA[OffA2] * AbsAxesB[AxisB][OffA1];
+
+				// Projected BoxB radius along AxisA x AxisB
+				RealType ProjRadiusB = ExtentsB[OffB1] * AbsAxesB[OffB2][AxisA] + ExtentsB[OffB2] * AbsAxesB[OffB1][AxisA];
+
+				// Abs(Dot(Center, Cross(AxisA, AxisB)))
+				RealType ProjCenterDist = FMath::Abs(
+					CenterB[OffA1] * AxesB[AxisB][OffA2] - CenterB[OffA2] * AxesB[AxisB][OffA1]
+				);
+
+				if (ProjCenterDist > ProjRadiusA + ProjRadiusB)
+				{
+					return false;
+				}
+			}
+		}
+
+		return true;
+	}
+}
+
 
 namespace UE::Geometry
 {
 
+
+template<typename RealType>
+bool TOrientedBox3<RealType>::Intersects(const TOrientedBox3& OtherBox, RealType ParallelAxesTolerance) const
+{
+	using namespace UE::OrientedBoxTypeLocals;
+
+	TQuaternion<RealType> ToFrameRot = Frame.Rotation.Inverse();
+	TQuaternion<RealType> LocalRot = ToFrameRot * OtherBox.Frame.Rotation;
+
+	TVector<RealType> LocalAxesB[3];
+	LocalRot.GetAxes(LocalAxesB[0], LocalAxesB[1], LocalAxesB[2]);
+	
+	TVector<RealType> LocalCenterB = ToFrameRot * (OtherBox.Frame.Origin - Frame.Origin);
+
+	return BoxIntersectsLocalSpaceHelper(
+		Extents,
+		LocalCenterB, LocalAxesB, OtherBox.Extents, ParallelAxesTolerance
+	);
+}
+
+template<typename RealType>
+bool TOrientedBox3<RealType>::Intersects(const TAxisAlignedBox3<RealType>& OtherBox, RealType ParallelAxesTolerance) const
+{
+	using namespace UE::OrientedBoxTypeLocals;
+
+	TVector<RealType> LocalAxesB[3];
+	Frame.Rotation.GetAxes(LocalAxesB[0], LocalAxesB[1], LocalAxesB[2]);
+
+	TVector<RealType> LocalCenterB = Frame.Origin - OtherBox.Center();
+
+	return BoxIntersectsLocalSpaceHelper(
+		OtherBox.Extents(),
+		LocalCenterB, LocalAxesB, Extents, ParallelAxesTolerance
+	);
+}
+
 template<typename RealType>
 TOrientedBox3<RealType> TOrientedBox3<RealType>::Merge(const TOrientedBox3<RealType>& Other, bool bOnlyConsiderExistingAxes) const
 {
diff --git a/Engine/Source/Runtime/GeometryCore/Public/OrientedBoxTypes.h b/Engine/Source/Runtime/GeometryCore/Public/OrientedBoxTypes.h
--- a/Engine/Source/Runtime/GeometryCore/Public/OrientedBoxTypes.h
+++ b/Engine/Source/Runtime/GeometryCore/Public/OrientedBoxTypes.h
@@ -290,6 +290,19 @@ struct TOrientedBox3
 	  */
 	 [[nodiscard]] GEOMETRYCORE_API TOrientedBox3<RealType> ReparameterizedCloserToWorldFrame() const;
 
+	 /**
+	  * @param OtherBox An oriented box to test for intersection against this box
+	  * @param ParallelAxesTolerance Tolerance for handling the special case where the box axes are nearly parallel
+	  * @return Whether the boxes overlap
+	  */
+	 [[nodiscard]] GEOMETRYCORE_API bool Intersects(const TOrientedBox3& OtherBox, RealType ParallelAxesTolerance = TMathUtil<RealType>::ZeroTolerance) const;
+	 
+	 /**
+	  * @param OtherBox An axis-aligned box to test for intersection against this box
+	  * @param ParallelAxesTolerance Tolerance for handling the special case where the box axes are nearly parallel
+	  * @return Whether the boxes overlap
+	  */
+	 [[nodiscard]] GEOMETRYCORE_API bool Intersects(const TAxisAlignedBox3<RealType>& OtherBox, RealType ParallelAxesTolerance = TMathUtil<RealType>::ZeroTolerance) const;
 };
 
 template<typename RealType>
