[GeometryCore] Adding morph target attribute to dynamic mesh

diff --git a/Engine/Plugins/Runtime/GeometryProcessing/Source/DynamicMesh/Private/Operations/SelectiveTessellate.cpp b/Engine/Plugins/Runtime/GeometryProcessing/Source/DynamicMesh/Private/Operations/SelectiveTessellate.cpp
--- a/Engine/Plugins/Runtime/GeometryProcessing/Source/DynamicMesh/Private/Operations/SelectiveTessellate.cpp
+++ b/Engine/Plugins/Runtime/GeometryProcessing/Source/DynamicMesh/Private/Operations/SelectiveTessellate.cpp
@@ -1867,6 +1867,14 @@ namespace SelectiveTessellateLocals
 				OutAttributes->AttachSkinWeightsAttribute(AttributeInfo.Key, SkinAttribute);
 			}
 
+			for (const TTuple<FName, TUniquePtr<FDynamicMeshMorphTargetAttribute>>& AttributeInfo : InAttributes->GetMorphTargetAttributes())
+			{
+				FDynamicMeshMorphTargetAttribute* MorphTargetAttribute = new FDynamicMeshMorphTargetAttribute(OutMesh);
+				ConstructDynamicMeshVertexAttribute<float, 3>(InMesh, AttributeInfo.Value.Get(), TessData, bUseParallel, CompactInfo, MorphTargetAttribute);
+				MorphTargetAttribute->SetName(AttributeInfo.Value.Get()->GetName());
+				OutAttributes->AttachMorphTargetAttribute(AttributeInfo.Key, MorphTargetAttribute);
+			}
+
 			if (Progress && Progress->Cancelled()) 
 			{
 				return false;
diff --git a/Engine/Plugins/Runtime/GeometryProcessing/Source/DynamicMesh/Private/Operations/UniformTessellate.cpp b/Engine/Plugins/Runtime/GeometryProcessing/Source/DynamicMesh/Private/Operations/UniformTessellate.cpp
--- a/Engine/Plugins/Runtime/GeometryProcessing/Source/DynamicMesh/Private/Operations/UniformTessellate.cpp
+++ b/Engine/Plugins/Runtime/GeometryProcessing/Source/DynamicMesh/Private/Operations/UniformTessellate.cpp
@@ -1877,6 +1877,22 @@ namespace UniformTessellateLocals
 				OutAttributes->AttachSkinWeightsAttribute(AttributeInfo.Key, SkinAttribute);
 			}
 
+			for (const TTuple<FName, TUniquePtr<FDynamicMeshMorphTargetAttribute>>& AttributeInfo : InAttributes->GetMorphTargetAttributes())
+			{
+				TUniquePtr<FVertexAttributeGenerator<float, 3>> Generator;
+				Generator = FVertexAttributeGenerator<float, 3>::CreateFromVertexAttribute(InMesh, TessellationNum, Progress, bUseParallel, AttributeInfo.Value.Get());
+				if (Generator->Generate() == false)
+				{
+					return false;
+				}
+
+				FDynamicMeshMorphTargetAttribute* MorphTargetAttribute = new FDynamicMeshMorphTargetAttribute(OutMesh);			
+				Generator->CopyToAttribute(MorphTargetAttribute);
+				MorphTargetAttribute->SetName(AttributeInfo.Value.Get()->GetName());
+				
+				OutAttributes->AttachMorphTargetAttribute(AttributeInfo.Key, MorphTargetAttribute);
+			}
+
 			if (InAttributes->NumWeightLayers() > 0)
 			{
 				OutAttributes->SetNumWeightLayers(InAttributes->NumWeightLayers());
diff --git a/Engine/Source/Runtime/Core/Public/UObject/FortniteMainBranchObjectVersions.inl b/Engine/Source/Runtime/Core/Public/UObject/FortniteMainBranchObjectVersions.inl
--- a/Engine/Source/Runtime/Core/Public/UObject/FortniteMainBranchObjectVersions.inl
+++ b/Engine/Source/Runtime/Core/Public/UObject/FortniteMainBranchObjectVersions.inl
@@ -682,5 +682,7 @@ FORTNITE_MAIN_VERSION(WorldPartitionActorComponentDesc, 222)
 // Migrate Non-Edit layer landscapes to use the edit layer (ULandscapeEditLayer) system
 FORTNITE_MAIN_VERSION(MigrateLandscapeNonEditLayerToEditLayer, 223)
 
+// FDynamicMeshAttributeSet has Morph Targets.
+FORTNITE_MAIN_VERSION(DynamicMeshAttributesMorphTargets, 224)
 // -----<new versions can be added above this line>-------------------------------------------------
 #undef FORTNITE_MAIN_VERSION
\ No newline at end of file
diff --git a/Engine/Source/Runtime/GeometryCore/Private/DynamicMesh/DynamicMeshAttributeSet.cpp b/Engine/Source/Runtime/GeometryCore/Private/DynamicMesh/DynamicMeshAttributeSet.cpp
--- a/Engine/Source/Runtime/GeometryCore/Private/DynamicMesh/DynamicMeshAttributeSet.cpp
+++ b/Engine/Source/Runtime/GeometryCore/Private/DynamicMesh/DynamicMeshAttributeSet.cpp
@@ -84,6 +84,13 @@ void FDynamicMeshAttributeSet::Copy(const FDynamicMeshAttributeSet& Copy)
 			static_cast<FDynamicMeshVertexSkinWeightsAttribute *>(AttribPair.Value->MakeCopy(ParentMesh)));
 	}
 
+	MorphTargetAttributes.Reset();
+	for (const TPair<FName, TUniquePtr<FDynamicMeshMorphTargetAttribute>>& AttribPair : Copy.MorphTargetAttributes)
+	{
+		AttachMorphTargetAttribute(AttribPair.Key,
+			static_cast<FDynamicMeshMorphTargetAttribute*>(AttribPair.Value->MakeCopy(ParentMesh)));
+	}
+
 	CopyBoneAttributes(Copy);
 	
 	GenericAttributes.Reset();
@@ -181,6 +188,13 @@ void FDynamicMeshAttributeSet::CompactCopy(const FCompactMaps& CompactMaps, cons
 			static_cast<FDynamicMeshVertexSkinWeightsAttribute *>(AttribPair.Value->MakeCompactCopy(CompactMaps, ParentMesh)));
 	}
 
+	MorphTargetAttributes.Reset();
+	for (const TPair<FName, TUniquePtr<FDynamicMeshMorphTargetAttribute>>& AttribPair : Copy.MorphTargetAttributes)
+	{
+		AttachMorphTargetAttribute(AttribPair.Key,
+			static_cast<FDynamicMeshMorphTargetAttribute*>(AttribPair.Value->MakeCompactCopy(CompactMaps, ParentMesh)));
+	}
+	
 	CopyBoneAttributes(Copy);
 	
 	GenericAttributes.Reset();
@@ -271,6 +285,11 @@ void FDynamicMeshAttributeSet::Append(const FDynamicMeshAttributeSet& ToAppend,
 		}
 	}
 
+	for (TPair<FName, TUniquePtr<FDynamicMeshMorphTargetAttribute>>& AttribPair : MorphTargetAttributes)
+	{
+		AppendHelper(*AttribPair.Value, ToAppend.GetMorphTargetAttribute(AttribPair.Key));
+	}
+
 	for (TPair<FName, TUniquePtr<FDynamicMeshAttributeBase>>& AttribPair : GenericAttributes)
 	{
 		const TUniquePtr<FDynamicMeshAttributeBase>* AppendAttr = ToAppend.GenericAttributes.Find(AttribPair.Key);
@@ -322,6 +341,11 @@ void FDynamicMeshAttributeSet::AppendDefaulted(const FDynamicMesh3::FAppendInfo&
 		AppendHelper(*AttribPair.Value);
 	}
 
+	for (const TPair<FName, TUniquePtr<FDynamicMeshMorphTargetAttribute>>& AttribPair : MorphTargetAttributes)
+	{
+		AppendHelper(*AttribPair.Value);
+	}
+
 	// Note: In the 'defaulted' case, there is nothing to append re the Bone attributes (since they are not tied to the base mesh elements)
 
 	for (const TPair<FName, TUniquePtr<FDynamicMeshAttributeBase>>& AttribPair : GenericAttributes)
@@ -487,6 +511,7 @@ void FDynamicMeshAttributeSet::EnableMatchingAttributes(const FDynamicMeshAttrib
 	{
 		// discard existing
 		SkinWeightAttributes.Reset();
+		MorphTargetAttributes.Reset();
 		GenericAttributes.Reset();
 		ResetRegisteredAttributes();
 
@@ -496,6 +521,11 @@ void FDynamicMeshAttributeSet::EnableMatchingAttributes(const FDynamicMeshAttrib
 			AttachSkinWeightsAttribute(AttribPair.Key,
 				static_cast<FDynamicMeshVertexSkinWeightsAttribute *>(AttribPair.Value->MakeNew(ParentMesh)));	
 		}
+		for (const TPair<FName, TUniquePtr<FDynamicMeshMorphTargetAttribute>>& AttribPair : ToMatch.MorphTargetAttributes)
+		{
+			AttachMorphTargetAttribute(AttribPair.Key,
+				static_cast<FDynamicMeshMorphTargetAttribute *>(AttribPair.Value->MakeNew(ParentMesh)));	
+		}
 		for (const TPair<FName, TUniquePtr<FDynamicMeshAttributeBase>>& AttribPair : ToMatch.GenericAttributes)
 		{
 			AttachAttribute(AttribPair.Key, AttribPair.Value->MakeNew(ParentMesh));
@@ -593,6 +623,11 @@ void FDynamicMeshAttributeSet::Reparent(FDynamicMesh3* NewParent)
 		AttribPair.Value.Get()->Reparent(NewParent);
 	}
 
+	for (const TPair<FName, TUniquePtr<FDynamicMeshMorphTargetAttribute>>& AttribPair : MorphTargetAttributes)
+	{
+		AttribPair.Value.Get()->Reparent(NewParent);
+	}
+
 	if (BoneNameAttrib)
 	{
 		BoneNameAttrib->Reparent(NewParent);
@@ -860,6 +895,46 @@ void FDynamicMeshAttributeSet::RemoveSkinWeightsAttribute(FName InProfileName)
 	}
 }
 
+void FDynamicMeshAttributeSet::RemoveAllSkinWeightsAttributes()
+{
+	for (const TPair<FName, TUniquePtr<FDynamicMeshVertexSkinWeightsAttribute>>& SkinWeightAttribute : SkinWeightAttributes)
+	{
+		UnregisterExternalAttribute(SkinWeightAttribute.Value.Get());
+	}
+
+	SkinWeightAttributes.Reset();	
+}
+
+void FDynamicMeshAttributeSet::AttachMorphTargetAttribute(FName InMorphTargetName, FDynamicMeshMorphTargetAttribute* InAttribute)
+{
+	RemoveMorphTargetAttribute(InMorphTargetName);
+
+	// Ensure proper ownership.
+	static_cast<FDynamicMeshAttributeBase *>(InAttribute)->Reparent(ParentMesh);
+	MorphTargetAttributes.Add(InMorphTargetName, TUniquePtr<FDynamicMeshMorphTargetAttribute>(InAttribute));
+
+	RegisterExternalAttribute(InAttribute);
+}
+
+void FDynamicMeshAttributeSet::RemoveMorphTargetAttribute(FName InMorphTargetName)
+{
+	if (MorphTargetAttributes.Contains(InMorphTargetName))
+	{
+		UnregisterExternalAttribute(MorphTargetAttributes[InMorphTargetName].Get());
+		MorphTargetAttributes.Remove(InMorphTargetName);
+	}
+}
+
+void FDynamicMeshAttributeSet::RemoveAllMorphTargetAttributes()
+{
+	for (const TPair<FName, TUniquePtr<FDynamicMeshMorphTargetAttribute>>& MorphTargetAttribute : MorphTargetAttributes)
+	{
+		UnregisterExternalAttribute(MorphTargetAttribute.Value.Get());
+	}
+
+	MorphTargetAttributes.Reset();
+}
+
 
 bool FDynamicMeshAttributeSet::IsSeamEdge(int eid) const
 {
@@ -1501,6 +1576,35 @@ bool FDynamicMeshAttributeSet::IsSameAs(const FDynamicMeshAttributeSet& Other, b
 		}
 	}
 
+	if (MorphTargetAttributes.Num() != Other.MorphTargetAttributes.Num())
+	{
+		return false;
+	}
+	
+	if (!MorphTargetAttributes.IsEmpty())
+	{
+		MorphTargetAttributesMap::TConstIterator It(MorphTargetAttributes);
+		MorphTargetAttributesMap::TConstIterator ItOther(Other.MorphTargetAttributes);
+		while (It && ItOther)
+		{
+			if (It.Key() != ItOther.Key())
+			{
+				return false;
+			}
+			const FDynamicMeshMorphTargetAttribute* MorphTarget = It->Value.Get();
+			const FDynamicMeshMorphTargetAttribute* MorphTargetOther = ItOther->Value.Get();
+
+			if (!MorphTarget->IsSameAs(*MorphTargetOther, bIgnoreDataLayout))
+			{
+				return false;
+			}
+			
+			++It;
+			++ItOther;
+		}	
+	}
+	
+	
 	if (!IsSameBoneAttributesAs(Other))
 	{
 		return false;
@@ -1539,6 +1643,7 @@ void FDynamicMeshAttributeSet::Serialize(FArchive& Ar, const FCompactMaps* Compa
 	
 	Ar.UsingCustomVersion(FUE5MainStreamObjectVersion::GUID);
 	Ar.UsingCustomVersion(FFortniteReleaseBranchCustomObjectVersion::GUID);
+	Ar.UsingCustomVersion(FFortniteMainBranchObjectVersion::GUID);
 
 	const bool bUseLegacySerialization = Ar.IsLoading() && Ar.CustomVer(FUE5MainStreamObjectVersion::GUID) < FUE5MainStreamObjectVersion::DynamicMeshCompactedSerialization;
 
@@ -1654,6 +1759,52 @@ void FDynamicMeshAttributeSet::Serialize(FArchive& Ar, const FCompactMaps* Compa
 		}
 	}
 
+	const bool bSerializeMorphTargets = !Ar.IsLoading() || Ar.CustomVer(FFortniteMainBranchObjectVersion::GUID) >= FFortniteMainBranchObjectVersion::DynamicMeshAttributesMorphTargets;
+	if (bSerializeMorphTargets)
+	{
+		int32 NumMorphTargetAttributes = MorphTargetAttributes.Num();
+		Ar << NumMorphTargetAttributes;
+
+		if (Ar.IsLoading())
+		{
+			MorphTargetAttributes.Reset();
+
+			for (int32 i = 0; i < NumMorphTargetAttributes; ++i)
+			{
+				FName Key;
+				FNameAsStringProxyArchive ProxyArchive(Ar);
+				ProxyArchive << Key;
+
+				TUniquePtr<FDynamicMeshMorphTargetAttribute>& Value = MorphTargetAttributes.Emplace(Key, nullptr);
+				bool IsValid;
+				Ar << IsValid;
+				if (IsValid)
+				{
+					Value.Reset(new FDynamicMeshMorphTargetAttribute(ParentMesh, false));
+					Value.Get()->Serialize(Ar, CompactMaps, bUseCompression);
+					RegisterExternalAttribute(Value.Get());
+				}
+			}
+		}
+		else
+		{
+			for (MorphTargetAttributesMap::TIterator It(MorphTargetAttributes); It; ++It)
+			{
+				FNameAsStringProxyArchive ProxyArchive(Ar);
+				ProxyArchive << It.Key();
+
+				const TUniquePtr<FDynamicMeshMorphTargetAttribute>& Value = It.Value();
+				bool bIsValid = Value.IsValid();
+				Ar << bIsValid;
+				if (bIsValid)
+				{
+					Value.Get()->Serialize(Ar, CompactMaps, bUseCompression);
+				}
+			}
+		}	
+	}
+	
+	
 	const bool bSerializeBones = !Ar.IsLoading() || Ar.CustomVer(FFortniteReleaseBranchCustomObjectVersion::GUID) >= FFortniteReleaseBranchCustomObjectVersion::DynamicMeshAttributesSerializeBones;
 	if (bSerializeBones)
 	{
@@ -1735,6 +1886,16 @@ SIZE_T FDynamicMeshAttributeSet::GetByteCount() const
 		ByteCount += WeightLayer.GetByteCount();
 	}
 
+	for (const TPair<FName, TUniquePtr<FDynamicMeshVertexSkinWeightsAttribute>>& SkinWeightAttribute : SkinWeightAttributes)
+	{
+		ByteCount += SkinWeightAttribute.Value->GetByteCount();
+	}
+
+	for (const TPair<FName, TUniquePtr<FDynamicMeshMorphTargetAttribute>>& MorphTargetAttribute : MorphTargetAttributes)
+	{
+		ByteCount += MorphTargetAttribute.Value->GetByteCount();
+	}
+
 	if (BoneNameAttrib)
 	{
 		ByteCount += BoneNameAttrib->GetByteCount();
@@ -1813,6 +1974,7 @@ bool FDynamicMeshAttributeSet::CheckValidity(bool bAllowNonmanifold, EValidityCh
 	{
 		bValid = GetWeightLayer(WeightLayerIndex)->CheckValidity(bAllowNonmanifold, FailMode) && bValid;
 	}
+
 	for (const TPair<FName, TUniquePtr<FDynamicMeshVertexSkinWeightsAttribute>>& Pair : SkinWeightAttributes)
 	{
 		if (Pair.Value)
@@ -1826,6 +1988,19 @@ bool FDynamicMeshAttributeSet::CheckValidity(bool bAllowNonmanifold, EValidityCh
 		}
 	}
 
+	for (const TPair<FName, TUniquePtr<FDynamicMeshMorphTargetAttribute>>& Pair : MorphTargetAttributes)
+	{
+		if (Pair.Value)
+		{
+			bValid = Pair.Value->CheckValidity(bAllowNonmanifold, FailMode) && bValid;
+
+			// If non-null, Morph Target Attributes are in the RegisteredAttributes array
+			int32 FoundIndex;
+			const bool bFound = RegisteredAttributes.Find(Pair.Value.Get(), FoundIndex);
+			bValid = bFound && bValid;
+		}
+	}
+
 	bValid = CheckBoneValidity(FailMode) && bValid;
 
 	bValid = SculptLayers.CheckValidity(this, bAllowNonmanifold, FailMode) && bValid;
diff --git a/Engine/Source/Runtime/GeometryCore/Private/DynamicMeshEditor.cpp b/Engine/Source/Runtime/GeometryCore/Private/DynamicMeshEditor.cpp
--- a/Engine/Source/Runtime/GeometryCore/Private/DynamicMeshEditor.cpp
+++ b/Engine/Source/Runtime/GeometryCore/Private/DynamicMeshEditor.cpp
@@ -2105,6 +2105,15 @@ void FDynamicMeshEditor::AppendMesh(const FDynamicMesh3* AppendMesh,
 				}
 			}
 		}
+
+		for (const TPair<FName, TUniquePtr<FDynamicMeshMorphTargetAttribute>>& AttribPair : AppendMesh->Attributes()->GetMorphTargetAttributes())
+		{
+			FDynamicMeshMorphTargetAttribute* ToAttrib = Mesh->Attributes()->GetMorphTargetAttribute(AttribPair.Key);
+			if (ToAttrib)
+			{
+				ToAttrib->CopyThroughMapping(AttribPair.Value.Get(), IndexMapsOut);	
+			}
+		}
 		
 		for (const TPair<FName, TUniquePtr<FDynamicMeshAttributeBase>>& AttribPair : AppendMesh->Attributes()->GetAttachedAttributes())
 		{
@@ -2462,6 +2471,15 @@ static void AppendVertexAttributes(const FDynamicMesh3* FromMesh, FDynamicMesh3*
 			ToAttrib->CopyThroughMapping(AttribPair.Value.Get(), IndexMaps);
 		}
 	}
+
+	for (const TPair<FName, TUniquePtr<FDynamicMeshMorphTargetAttribute>>& AttribPair : FromMesh->Attributes()->GetMorphTargetAttributes())
+	{
+		FDynamicMeshMorphTargetAttribute* ToAttrib = ToMesh->Attributes()->GetMorphTargetAttribute(AttribPair.Key);
+		if (ToAttrib)
+		{
+			ToAttrib->CopyThroughMapping(AttribPair.Value.Get(), IndexMaps);
+		}
+	}
 	
 	for (const TPair<FName, TUniquePtr<FDynamicMeshAttributeBase>>& AttribPair : FromMesh->Attributes()->GetAttachedAttributes())
 	{
diff --git a/Engine/Source/Runtime/GeometryCore/Public/DynamicMesh/DynamicMeshAttributeSet.h b/Engine/Source/Runtime/GeometryCore/Public/DynamicMesh/DynamicMeshAttributeSet.h
--- a/Engine/Source/Runtime/GeometryCore/Public/DynamicMesh/DynamicMeshAttributeSet.h
+++ b/Engine/Source/Runtime/GeometryCore/Public/DynamicMesh/DynamicMeshAttributeSet.h
@@ -52,6 +52,9 @@ typedef TDynamicMeshScalarTriangleAttribute<int32> FDynamicMeshPolygroupAttribut
 /** Per-vertex scalar float weight */
 typedef TDynamicMeshVertexAttribute<float, 1> FDynamicMeshWeightAttribute;
 
+/** Per-vertex 3-element float morph target*/
+typedef TDynamicMeshVertexAttribute<float, 3> FDynamicMeshMorphTargetAttribute;
+
 /** Forward declarations */
 template<typename ParentType>
 class TDynamicVertexSkinWeightsAttribute;
@@ -377,6 +380,9 @@ public:
 	/// Remove a skin weights attribute matching the given profile name.
 	GEOMETRYCORE_API void RemoveSkinWeightsAttribute(FName InProfileName);
 
+	/// Remove all skin weights attributes
+	GEOMETRYCORE_API void RemoveAllSkinWeightsAttributes();
+
 	/// Returns true if the list of skin weight attributes includes the given profile name.
 	bool HasSkinWeightsAttribute(FName InProfileName) const
 	{
@@ -403,7 +409,43 @@ public:
 		return SkinWeightAttributes;
 	}
 
+	/** Morph Targets */
+	
+	/// Create a new morph target attribute with a given name. If an attribute already exists with
+	/// that name, that existing attribute will be deleted.
+	GEOMETRYCORE_API void AttachMorphTargetAttribute(FName InMorphTargetName, FDynamicMeshMorphTargetAttribute* InAttribute);
+
+	/// Remove a morph target attribute matching the given name.
+	GEOMETRYCORE_API void RemoveMorphTargetAttribute(FName InMorphTargetName);
+	
+	/// Remove all morph target attributes
+	GEOMETRYCORE_API void RemoveAllMorphTargetAttributes();
 
+	/// Returns true if the morph target attribute exists.
+	bool HasMorphTargetAttribute(FName InMorphTargetName) const
+	{
+		return MorphTargetAttributes.Contains(InMorphTargetName);
+	}
+
+	/// Returns a pointer to a morph target attribute of the given name. If the attribute
+	/// does not exist, a nullptr is returned.
+	FDynamicMeshMorphTargetAttribute* GetMorphTargetAttribute(FName InMorphTargetName) const
+	{
+		if (const TUniquePtr<FDynamicMeshMorphTargetAttribute>* AttributePtr = MorphTargetAttributes.Find(InMorphTargetName))
+		{
+			return AttributePtr->Get();
+		}
+		else
+		{
+			return nullptr;
+		}
+	}
+	
+	/// Returns a map of all morph target attributes.
+	const TMap<FName, TUniquePtr<FDynamicMeshMorphTargetAttribute>>& GetMorphTargetAttributes() const
+	{
+		return MorphTargetAttributes;
+	}
 
 	//
 	// Bone Attributes
@@ -613,6 +655,9 @@ protected:
 	using SkinWeightAttributesMap = TMap<FName, TUniquePtr<FDynamicMeshVertexSkinWeightsAttribute>>;
 	SkinWeightAttributesMap SkinWeightAttributes;
 
+	using MorphTargetAttributesMap = TMap<FName, TUniquePtr<FDynamicMeshMorphTargetAttribute>>;
+	MorphTargetAttributesMap MorphTargetAttributes;
+
 	// Bone attributes
 	TUniquePtr<FDynamicMeshBoneNameAttribute> BoneNameAttrib;
 	TUniquePtr<FDynamicMeshBoneParentIndexAttribute> BoneParentIndexAttrib;
diff --git a/Engine/Source/Runtime/MeshConversion/Private/DynamicMeshToMeshDescription.cpp b/Engine/Source/Runtime/MeshConversion/Private/DynamicMeshToMeshDescription.cpp
--- a/Engine/Source/Runtime/MeshConversion/Private/DynamicMeshToMeshDescription.cpp
+++ b/Engine/Source/Runtime/MeshConversion/Private/DynamicMeshToMeshDescription.cpp
@@ -783,6 +783,7 @@ void FDynamicMeshToMeshDescription::Convert_NoSharedInstances(const FDynamicMesh
 	const FDynamicMeshBoneColorAttribute* BoneColors = bHasBones ? MeshIn->Attributes()->GetBoneColors() : nullptr;
 
 	const bool bHasSkinWeightsAttributes = bHasAttributes && !MeshIn->Attributes()->GetSkinWeightsAttributes().IsEmpty();
+	const bool bHasMorphTargetAttributes = bHasAttributes && !MeshIn->Attributes()->GetMorphTargetAttributes().IsEmpty();
 
 	// cache the UV layers
 	TArray<const FDynamicMeshUVOverlay*> UVLayers;
@@ -851,6 +852,19 @@ void FDynamicMeshToMeshDescription::Convert_NoSharedInstances(const FDynamicMesh
 		}
 	}
 
+	TMap<FName, TVertexAttributesRef<FVector3f>> MorphTargetsMap;
+	if (bHasMorphTargetAttributes)
+	{
+		FSkeletalMeshAttributes MeshOutAttributes(MeshOut);
+		for (const TTuple<FName, TUniquePtr<FDynamicMeshMorphTargetAttribute>>& AttributeInfo: MeshIn->Attributes()->GetMorphTargetAttributes())
+		{
+			FName MorphTargetName = AttributeInfo.Key;
+
+			MeshOutAttributes.RegisterMorphTargetAttribute(MorphTargetName, false);
+			MorphTargetsMap.Add(MorphTargetName, MeshOutAttributes.GetVertexMorphPositionDelta(MorphTargetName));
+		}
+		
+	}
 
 	// always copy when we are baking new mesh? should this be a config option?
 	bool bCopyInstanceColors = (ColorOverlay != nullptr); 
@@ -1106,6 +1120,28 @@ void FDynamicMeshToMeshDescription::Convert_NoSharedInstances(const FDynamicMesh
 			}
 		}
 	}
+
+	if (bHasMorphTargetAttributes)
+	{
+		for (const TTuple<FName, TUniquePtr<FDynamicMeshMorphTargetAttribute>>& AttributeInfo: MeshIn->Attributes()->GetMorphTargetAttributes())
+		{
+			FName ProfileName = AttributeInfo.Key;
+		
+			TVertexAttributesRef<FVector3f>& MeshDescMorphTarget = MorphTargetsMap[ProfileName];
+
+			const FDynamicMeshMorphTargetAttribute* DynaMeshMorphTarget = AttributeInfo.Value.Get();
+			
+			for (int32 VertexIndex = 0; VertexIndex < MapV.Num(); VertexIndex++)
+			{
+				if (const FVertexID VertexID = MapV[VertexIndex]; VertexID != INDEX_NONE)
+				{
+					FVector3f Offset;
+					DynaMeshMorphTarget->GetValue(VertexIndex, Offset);
+					MeshDescMorphTarget.Set(VertexID, Offset);	
+				}
+			}
+		}
+	}
 	
 	Builder.ResumeMeshDescriptionIndexing();
 }
diff --git a/Engine/Source/Runtime/MeshConversion/Private/MeshDescriptionToDynamicMesh.cpp b/Engine/Source/Runtime/MeshConversion/Private/MeshDescriptionToDynamicMesh.cpp
--- a/Engine/Source/Runtime/MeshConversion/Private/MeshDescriptionToDynamicMesh.cpp
+++ b/Engine/Source/Runtime/MeshConversion/Private/MeshDescriptionToDynamicMesh.cpp
@@ -171,6 +171,12 @@ struct FSkinWeightsAttribCopyInfo
 	FDynamicMeshVertexSkinWeightsAttribute* DynaMeshAttribRef = nullptr;
 };
 
+struct FMorphTargetAttribCopyInfo
+{
+	TVertexAttributesConstRef<FVector3f> MeshDescAttribRef;
+	FDynamicMeshMorphTargetAttribute* DynaMeshAttribRef = nullptr;
+};
+
 
 void FMeshDescriptionToDynamicMesh::Convert(const FMeshDescription* MeshIn, FDynamicMesh3& MeshOut, bool bCopyTangents )
 {
@@ -425,6 +431,7 @@ void FMeshDescriptionToDynamicMesh::Convert(const FMeshDescription* MeshIn, FDyn
 	FDynamicMeshBoneColorAttribute* BoneColorAttrib = nullptr;
 
 	TArray<FSkinWeightsAttribCopyInfo> SkinWeightAttribs;
+	TArray<FMorphTargetAttribCopyInfo> MorphTargetAttribs;
 	if (!bDisableAttributes)
 	{
 		MeshOut.EnableAttributes(); // by default 1-UV layer and 1-normal layer
@@ -485,6 +492,23 @@ void FMeshDescriptionToDynamicMesh::Convert(const FMeshDescription* MeshIn, FDyn
 			}
 		});
 
+		// Copy all morph targets
+		MeshIn->VertexAttributes().ForEach([&](const FName InAttributeName, auto InAttributesRef)
+		{
+			if (FSkeletalMeshAttributes::IsMorphTargetAttribute(InAttributeName))
+			{
+				const FName MorphTargetName = FSkeletalMeshAttributes::GetMorphTargetNameFromAttribute(InAttributeName);
+				
+				FDynamicMeshMorphTargetAttribute *MorphTargetAttribute = new FDynamicMeshMorphTargetAttribute(&MeshOut);
+				MorphTargetAttribs.Add({
+					Attributes.GetVertexMorphPositionDelta(MorphTargetName),
+					MorphTargetAttribute                    
+				});
+				
+				MeshOut.Attributes()->AttachMorphTargetAttribute(MorphTargetName, MorphTargetAttribute);
+			}
+		});
+		
 		if (Attributes.HasBones())
 		{	
 			MeshOut.Attributes()->EnableBones(Attributes.GetNumBones());
@@ -863,6 +887,19 @@ void FMeshDescriptionToDynamicMesh::Convert(const FMeshDescription* MeshIn, FDyn
 			Pending.Add(MoveTemp(SkinWeightsTask));
 		}
 
+		for (const FMorphTargetAttribCopyInfo& MorphTargetAttribInfo : MorphTargetAttribs)
+		{
+			auto MorphTargetTask = UE::Tasks::Launch(UE_SOURCE_LOCATION, [&, MorphTargetAttribInfo]() -> void
+			{
+				for (int32 VertexID: MeshOut.VertexIndicesItr())
+				{
+					const FVector3f& V = MorphTargetAttribInfo.MeshDescAttribRef.Get(VertIDMap[VertexID]);
+					MorphTargetAttribInfo.DynaMeshAttribRef->SetValue(VertexID, V);
+				}
+			});
+			Pending.Add(MoveTemp(MorphTargetTask));
+		}
+
 		if (BoneNameAttrib != nullptr)
 		{	
 			auto BonesTask = UE::Tasks::Launch(UE_SOURCE_LOCATION, [&]() -> void
