Add IsValid() test for Oriented Box to check for valid (non-negative) extents
+ validate that boxes are valid/non-empty before testing for intersection

diff --git a/Engine/Source/Runtime/GeometryCore/Private/OrientedBoxTypes.cpp b/Engine/Source/Runtime/GeometryCore/Private/OrientedBoxTypes.cpp
--- a/Engine/Source/Runtime/GeometryCore/Private/OrientedBoxTypes.cpp
+++ b/Engine/Source/Runtime/GeometryCore/Private/OrientedBoxTypes.cpp
@@ -122,6 +122,11 @@ bool TOrientedBox3<RealType>::Intersects(const TOrientedBox3& OtherBox, RealType
 {
 	using namespace UE::OrientedBoxTypeLocals;
 
+	if (!IsValid() || !OtherBox.IsValid())
+	{
+		return false;
+	}
+
 	TQuaternion<RealType> ToFrameRot = Frame.Rotation.Inverse();
 	TQuaternion<RealType> LocalRot = ToFrameRot * OtherBox.Frame.Rotation;
 
@@ -141,6 +146,11 @@ bool TOrientedBox3<RealType>::Intersects(const TAxisAlignedBox3<RealType>& Other
 {
 	using namespace UE::OrientedBoxTypeLocals;
 
+	if (!IsValid() || OtherBox.IsEmpty())
+	{
+		return false;
+	}
+
 	TVector<RealType> LocalAxesB[3];
 	Frame.Rotation.GetAxes(LocalAxesB[0], LocalAxesB[1], LocalAxesB[2]);
 
diff --git a/Engine/Source/Runtime/GeometryCore/Public/OrientedBoxTypes.h b/Engine/Source/Runtime/GeometryCore/Public/OrientedBoxTypes.h
--- a/Engine/Source/Runtime/GeometryCore/Public/OrientedBoxTypes.h
+++ b/Engine/Source/Runtime/GeometryCore/Public/OrientedBoxTypes.h
@@ -131,6 +131,11 @@ struct TOrientedBox3
 			(TMathUtil<RealType>::Abs(InFramePoint.Z) <= Extents.Z);
 	}
 
+	/** @return true if box is valid (has non-negative extents) */
+	bool IsValid() const
+	{
+		return Extents.X >= 0 && Extents.Y >= 0 && Extents.Z >= 0;
+	}
 
 	// corners [ (-x,-y), (x,-y), (x,y), (-x,y) ], -z, then +z
 	//
