BSpline: Convert knot vector search in InsertKnot to a binary search.

diff --git a/Engine/Source/Runtime/GeometryCore/Public/Splines/BSpline.h b/Engine/Source/Runtime/GeometryCore/Public/Splines/BSpline.h
--- a/Engine/Source/Runtime/GeometryCore/Public/Splines/BSpline.h
+++ b/Engine/Source/Runtime/GeometryCore/Public/Splines/BSpline.h
@@ -568,28 +568,39 @@ protected:
 	
 	/**
 	 * Inserts a knot at the specified parameter value
-	 * 
 	 * @param InKnot - Knot to insert
+	 * @return true if knot insertion was successful, false otherwise.
 	 */
-	void InsertKnot(FKnot InKnot)
+	bool InsertKnot(FKnot InKnot)
 	{
-		for (int32 i = 0; i < PairKnots.Num(); ++i)
+		// Binary search the knot vector:
+		
+		int32 Low = 0;	// This will be our insert index
+		int32 High = PairKnots.Num() - 1;
+
+		while (Low <= High)
 		{
-			if (PairKnots[i].Value == InKnot.Value)  // exact match
+			const int32 Mid = Low + (High - Low) / 2;
+			
+			if (PairKnots[Mid].Value == InKnot.Value)
+			{
+				UE_LOG(LogSpline, Warning, TEXT("Knot insertion failed, knot already exists!"));
+				return false;
+			}
+			else if (PairKnots[Mid].Value < InKnot.Value)
 			{
-				ensureAlwaysMsgf(false, TEXT("Knot insertion failed, knot already exists!"));
-				return;
+				Low = Mid + 1;
 			}
-			else if (PairKnots[i].Value > InKnot.Value)
+			else
 			{
-				PairKnots.Insert(InKnot, i);
-				MarkFlatKnotsCacheDirty();
-				return;
+				High = Mid - 1;
 			}
 		}
 
-		PairKnots.Add(InKnot);
+		PairKnots.Insert(InKnot, Low);
 		MarkFlatKnotsCacheDirty();
+
+		return true;
 	}
 
 	struct FValidKnotSearchParams
