Spline Core: Make windows static arrays to reduce heap allocations when evaluating

diff --git a/Engine/Source/Runtime/GeometryCore/Public/Splines/BSpline.h b/Engine/Source/Runtime/GeometryCore/Public/Splines/BSpline.h
--- a/Engine/Source/Runtime/GeometryCore/Public/Splines/BSpline.h
+++ b/Engine/Source/Runtime/GeometryCore/Public/Splines/BSpline.h
@@ -117,7 +117,7 @@ public:
 	using ValueType = typename TSplineInterface<VALUETYPE>::ValueType;
 
 	static constexpr int32 WindowSize = Degree + 1;
-	using FWindow = TArray<const ValueType*>;
+	using FWindow = TStaticArray<const ValueType*, WindowSize>;
 	
 	DECLARE_SPLINE_TYPE_ID(
 		BSplineNameSelector<DEGREE>::Name,
@@ -250,10 +250,10 @@ public:
 		if (NumValues == 0) return ValueType();
 		if (NumValues == 1) return GetValue(0);
 
-		// Get window using override or evaluator
 		FWindow Window = FindWindow(Parameter);
-    	
-		if (Window.Num() == 0) return ValueType();
+
+		// If FindWindow fails, it will return an array of nullptr.
+		if (Window[0] == nullptr) return ValueType();
         
 		return InterpolateWindow(Window, Parameter);
 	}
@@ -1075,7 +1075,7 @@ private:
 
     virtual FWindow FindWindow(float Parameter) const
     {
-    	FWindow Window;
+    	FWindow Window = {};
     	if (NumKeys() < Degree + 1)
     	{
     		return Window;
@@ -1093,7 +1093,7 @@ private:
     		{
     			int32 Index = (Span - Degree + i) % NumPoints;
     			if (Index < 0) Index += NumPoints;
-    			Window.Add(&GetValue(Index));
+    			Window[i] = &GetValue(Index);
     		}
     	}
     	else
@@ -1104,7 +1104,7 @@ private:
     			const int32 Index = Span - Degree + i;
     			if (Index >= 0 && Index < NumKeys())
     			{
-    				Window.Add(&GetValue(Index));
+    				Window[i] = &GetValue(Index);
     			}
     		}
     	}
@@ -1112,6 +1112,12 @@ private:
     	return Window;
     }
 
+	/**
+	 * Evaluate the provided window using basis functions computed using the provided parameter. 
+	 * @param Window The window to interpolate with the computed basis. All elements are assumed to be valid pointers.
+	 * @param Parameter The parameter to use to compute the basis functions.
+	 * @return The interpolated value.
+	 */
 	virtual ValueType InterpolateWindow(TArrayView<const ValueType* const> Window, float Parameter) const
     {
     	TArray<float> Basis;
diff --git a/Engine/Source/Runtime/GeometryCore/Public/Splines/PolyBezierSpline.h b/Engine/Source/Runtime/GeometryCore/Public/Splines/PolyBezierSpline.h
--- a/Engine/Source/Runtime/GeometryCore/Public/Splines/PolyBezierSpline.h
+++ b/Engine/Source/Runtime/GeometryCore/Public/Splines/PolyBezierSpline.h
@@ -339,10 +339,15 @@ public:
     template<int32 Order>
     ValueType EvaluateDerivative(float Parameter) const
     {
+		FWindow Window = FindWindow(Parameter);
+
+		// If FindWindow fails, it will return an array of nullptr. InterpolateWindow assumes validity of elements.
+		if (Window[0] == nullptr) { return ValueType(); }
+		
         // Order 0 is just regular evaluation
 	    if constexpr (Order == 0)
 	    {
-	        return InterpolateWindow(FindWindow(Parameter), Parameter);
+	        return InterpolateWindow(Window, Parameter);
 	    }
     
     	// Convert global parameter to segment information
@@ -350,13 +355,6 @@ public:
     	float LocalT;
     	MapGlobalParameterToLocalSegment(Parameter, SegmentIndex, LocalT);
 		
-		// Get window for the current parameter
-	    FWindow Window = FindWindow(Parameter);
-	    if (Window.Num() < 4)  // Need all 4 control points for a cubic Bezier
-	    {
-	        return ValueType();
-	    }
-    	
 	    // For PolyBezier, each segment is normalized to unit parameter space
         constexpr float SegmentScale = 1.0f;
 	    
@@ -1593,7 +1591,7 @@ private:
 
     virtual FWindow FindWindow(float Parameter) const override
     {
-        FWindow Window;
+        FWindow Window = {};
         if (Base::NumKeys() < 4)  // Need at least one complete Bezier segment
         {
             return Window;
@@ -1625,23 +1623,24 @@ private:
                 return ((Index % NumPoints) + NumPoints) % NumPoints;
             };
 
-			Window.Add(&Entries[WrapIndex(BaseIndex + 0)]);
-        	Window.Add(&Entries[WrapIndex(BaseIndex + 1)]);
-        	Window.Add(&Entries[WrapIndex(BaseIndex + 2)]);
-        	Window.Add(&Entries[WrapIndex(BaseIndex + 3)]);
+        	for (int Idx = 0; Idx < Base::WindowSize; ++Idx)
+        	{
+        		Window[Idx] = &Entries[WrapIndex(BaseIndex + Idx)];
+        	}
         }
         else
         {
-            // Standard case - no wrapping
+        	// Return null window if BaseIndex is too large to fully populate the window.
             if (BaseIndex + 3 >= Entries.Num())
             {
                 return Window;
             }
 
-        	Window.Add(&Entries[BaseIndex + 0]);
-        	Window.Add(&Entries[BaseIndex + 1]);
-        	Window.Add(&Entries[BaseIndex + 2]);
-        	Window.Add(&Entries[BaseIndex + 3]);
+        	// Standard case - no wrapping
+        	for (int Idx = 0; Idx < Base::WindowSize; ++Idx)
+        	{
+        		Window[Idx] = &Entries[BaseIndex + Idx];
+        	}
         }
 
         return Window;
