PolyBezierSpline: Preserve parameterization of closed splines when temporarily opening the spline

diff --git a/Engine/Source/Runtime/GeometryCore/Public/Splines/PolyBezierSpline.h b/Engine/Source/Runtime/GeometryCore/Public/Splines/PolyBezierSpline.h
--- a/Engine/Source/Runtime/GeometryCore/Public/Splines/PolyBezierSpline.h
+++ b/Engine/Source/Runtime/GeometryCore/Public/Splines/PolyBezierSpline.h
@@ -635,13 +635,14 @@ public:
         const ValueType& P3, 
         EParameterizationPolicy ParameterizationPolicy = EParameterizationPolicy::Centripetal)
     {
-		bool bWasClosed = Base::IsClosedLoop();
+		const TOptional<float> LoopKnotDelta = GetLoopKnotDelta();
 		SetClosedLoop(false);
 		ON_SCOPE_EXIT
 		{
-			if (bWasClosed)
+			if (LoopKnotDelta.IsSet())
 			{
 				SetClosedLoop(true);
+				SetParameter(GetNumberOfSegments(), LoopKnotDelta.GetValue());
 			}
 		};
 		
@@ -669,13 +670,14 @@ public:
 		const ValueType& P2, 
 		EParameterizationPolicy ParameterizationPolicy = EParameterizationPolicy::Centripetal)
 	{
-		bool bWasClosed = Base::IsClosedLoop();
+		const TOptional<float> LoopKnotDelta = GetLoopKnotDelta();
 		SetClosedLoop(false);
 		ON_SCOPE_EXIT
 		{
-			if (bWasClosed)
+			if (LoopKnotDelta.IsSet())
 			{
 				SetClosedLoop(true);
+				SetParameter(GetNumberOfSegments(), LoopKnotDelta.GetValue());
 			}
 		};
 		
@@ -1771,13 +1773,14 @@ private:
 		const ValueType& Position,
 		EParameterizationPolicy ParameterizationPolicy = EParameterizationPolicy::Centripetal)
 	{
-		bool bWasClosed = Base::IsClosedLoop();
+		const TOptional<float> LoopKnotDelta = GetLoopKnotDelta();
 		SetClosedLoop(false);
 		ON_SCOPE_EXIT
 		{
-			if (bWasClosed)
+			if (LoopKnotDelta.IsSet())
 			{
 				SetClosedLoop(true);
+				SetParameter(GetNumberOfSegments(), LoopKnotDelta.GetValue());
 			}
 		};
     
@@ -1807,15 +1810,16 @@ private:
 		const ValueType& Position,
 		EParameterizationPolicy ParameterizationPolicy)
 	{
-		bool bWasClosed = Base::IsClosedLoop();
-    	SetClosedLoop(false);
-    	ON_SCOPE_EXIT
-    	{
-    		if (bWasClosed)
-    		{
-    			SetClosedLoop(true);
-    		}
-    	};
+		const TOptional<float> LoopKnotDelta = GetLoopKnotDelta();
+		SetClosedLoop(false);
+		ON_SCOPE_EXIT
+		{
+			if (LoopKnotDelta.IsSet())
+			{
+				SetClosedLoop(true);
+				SetParameter(GetNumberOfSegments(), LoopKnotDelta.GetValue());
+			}
+		};
     		
 		// Append - create a point at the end
 		int32 LastIndex = Base::NumKeys() - 1;
@@ -1920,6 +1924,14 @@ private:
 		return SegmentIndex;
 	}
 
+	/** Returns the length of the parameter range of the closing segment. The returned value is unset if the spline is open. */
+	TOptional<float> GetLoopKnotDelta() const
+	{
+		return Base::IsClosedLoop()
+			? Base::PairKnots[Base::PairKnots.Num() - 1].Value - Base::PairKnots[Base::PairKnots.Num() - 2].Value
+			: TOptional<float>();
+	}
+
 private:
 
 	/** This is the last segment fetched by parameter. Used to accelerate predictable sampling patterns. */
