GeometryCore: Fixed crash in DynamicVector::TruncateBlocks for empty vector.

diff --git a/Engine/Source/Runtime/GeometryCore/Public/Util/DynamicVector.h b/Engine/Source/Runtime/GeometryCore/Public/Util/DynamicVector.h
--- a/Engine/Source/Runtime/GeometryCore/Public/Util/DynamicVector.h
+++ b/Engine/Source/Runtime/GeometryCore/Public/Util/DynamicVector.h
@@ -353,6 +353,11 @@ private:
 
 	void TruncateBlocks(int32 NewBlockCount, EAllowShrinking AllowShrinking)
 	{
+		if (Blocks.Num() - NewBlockCount <= 0)
+		{
+			return;
+		}
+		
 		for (int32 k = NewBlockCount; k < Blocks.Num(); ++k)
 		{
 			delete Blocks[k];
