TangentBezierSpline: Fix auto-tangent computation for non-uniform knot vectors

diff --git a/Engine/Source/Runtime/GeometryCore/Public/Splines/TangentBezierSpline.h b/Engine/Source/Runtime/GeometryCore/Public/Splines/TangentBezierSpline.h
--- a/Engine/Source/Runtime/GeometryCore/Public/Splines/TangentBezierSpline.h
+++ b/Engine/Source/Runtime/GeometryCore/Public/Splines/TangentBezierSpline.h
@@ -1421,13 +1421,14 @@ public:
 	void UpdatePointTangents(int32 Index, ETangentMode Mode)
 	{
 		const int32 NumPoints = GetNumPoints();
+		const int32 NumSegments = GetNumberOfSegments();
+		
 		if (Index < 0 || Index >= NumPoints || Mode == ETangentMode::User || Mode == ETangentMode::Broken)
 		{
 			// Invalid index or user-defined mode, no action needed
 			return;
 		}
 		
-
 		// Get the position of the current point
 		const ValueType& Current = GetValue(Index);
 
@@ -1436,23 +1437,22 @@ public:
 
 		const ValueType& Next = (Index < NumPoints - 1) ? GetValue(Index + 1) : (IsClosedLoop() ? GetValue(0) : Current);
 
-		// Get actual parameter values for points
-		float PrevParam = (Index > 0)
-			                  ? static_cast<float>(Index - 1)
-			                  : (IsClosedLoop() ? static_cast<float>(NumPoints - 1) : static_cast<float>(Index));
-
-		float CurrentParam = static_cast<float>(Index);
+		const float CurrentParam = GetParameter(Index);
+		float PrevParam = Index > 0 ? GetParameter(Index - 1) : CurrentParam;
+		float NextParam = Index < NumPoints - 1 ? GetParameter(Index + 1) : CurrentParam;
 
-		float NextParam = (Index < NumPoints - 1) ? static_cast<float>(Index + 1) : (IsClosedLoop() ? 0.0f : static_cast<float>(Index));
-
-		// Adjust params for closed loop to ensure proper weighting
-		if (IsClosedLoop() && Index == 0)
-		{
-			PrevParam -= static_cast<float>(NumPoints); // Make previous parameter negative for proper delta
-		}
-		else if (IsClosedLoop() && Index == NumPoints - 1)
+		// If we are closed and updating the first or last point, we need to do some special work to get virtual PrevParam or NextParam
+		// The 'virtual' parameter is not the actual parameter of the previous or next point, but the current parameter +- the knot delta.
+		if (IsClosedLoop())
 		{
-			NextParam += static_cast<float>(NumPoints); // Make next parameter greater than NumPoints for proper delta
+			if (Index == 0)
+			{
+				PrevParam = CurrentParam - GetSegmentParameterRange(NumSegments - 1).Length();
+			}
+			else if (Index == NumPoints - 1)
+			{
+				NextParam = CurrentParam + GetSegmentParameterRange(NumSegments - 1).Length();
+			}
 		}
 
 		// Calculate tangents based on mode
