Fix MaxIterations reached issue in BSpline::FindNearestAvailableKnotValue

#was-backout 46092751

diff --git a/Engine/Source/Runtime/Engine/Private/Curves/Spline.cpp b/Engine/Source/Runtime/Engine/Private/Curves/Spline.cpp
--- a/Engine/Source/Runtime/Engine/Private/Curves/Spline.cpp
+++ b/Engine/Source/Runtime/Engine/Private/Curves/Spline.cpp
@@ -86,53 +86,52 @@ FAutoConsoleVariableRef CVarValidateRotScale(
 	
 }
 
-using namespace UE::Geometry::Spline;
 
 /**
  * Converts from EInterpCurveMode to ETangentMode
  */
-inline ETangentMode ConvertInterpCurveModeToTangentMode(EInterpCurveMode Mode)
+inline UE::Geometry::Spline::ETangentMode ConvertInterpCurveModeToTangentMode(EInterpCurveMode Mode)
 {
 	switch (Mode)
 	{
 	case CIM_Linear:
-		return ETangentMode::Linear;
+		return UE::Geometry::Spline::ETangentMode::Linear;
 	case CIM_CurveAuto:
-		return ETangentMode::Auto;
+		return UE::Geometry::Spline::ETangentMode::Auto;
 	case CIM_Constant:
-		return ETangentMode::Constant;
+		return UE::Geometry::Spline::ETangentMode::Constant;
 	case CIM_CurveUser:
-		return ETangentMode::User;
+		return UE::Geometry::Spline::ETangentMode::User;
 	case CIM_CurveBreak:
-		return ETangentMode::Broken;
+		return UE::Geometry::Spline::ETangentMode::Broken;
 	case CIM_CurveAutoClamped:
-		return ETangentMode::AutoClamped;
+		return UE::Geometry::Spline::ETangentMode::AutoClamped;
 	case CIM_Unknown:
 	default:
-		return ETangentMode::Unknown;
+		return UE::Geometry::Spline::ETangentMode::Unknown;
 	}
 }
 
 /**
  * Converts from ETangentMode to EInterpCurveMode
  */
-inline EInterpCurveMode ConvertTangentModeToInterpCurveMode(ETangentMode Mode)
+inline EInterpCurveMode ConvertTangentModeToInterpCurveMode(UE::Geometry::Spline::ETangentMode Mode)
 {
 	switch (Mode)
 	{
-	case ETangentMode::Linear:
+	case UE::Geometry::Spline::ETangentMode::Linear:
 		return CIM_Linear;
-	case ETangentMode::Auto:
+	case UE::Geometry::Spline::ETangentMode::Auto:
 		return CIM_CurveAuto;
-	case ETangentMode::Constant:
+	case UE::Geometry::Spline::ETangentMode::Constant:
 		return CIM_Constant;
-	case ETangentMode::User:
+	case UE::Geometry::Spline::ETangentMode::User:
 		return CIM_CurveUser;
-	case ETangentMode::Broken:
+	case UE::Geometry::Spline::ETangentMode::Broken:
 		return CIM_CurveBreak;
-	case ETangentMode::AutoClamped:
+	case UE::Geometry::Spline::ETangentMode::AutoClamped:
 		return CIM_CurveAutoClamped;
-	case ETangentMode::Unknown:
+	case UE::Geometry::Spline::ETangentMode::Unknown:
 	default:
 		return CIM_Unknown;
 	}
@@ -146,9 +145,9 @@ inline EInterpCurveMode ConvertTangentModeToInterpCurveMode(ETangentMode Mode)
  * tangent control and automatic tangent computation.
  */
 
-class FNewSpline: public TMultiSpline<FTangentBezierSpline3d>
+class FNewSpline: public UE::Geometry::Spline::TMultiSpline<UE::Geometry::Spline::FTangentBezierSpline3d>
 {
-	using FTangentBezierControlPoint = TTangentBezierControlPoint<FVector3d>;
+	using FTangentBezierControlPoint = UE::Geometry::Spline::TTangentBezierControlPoint<FVector3d>;
 	
 public:
 	
@@ -350,8 +349,8 @@ private:
 template <typename AttrType>
 int32 FNewSpline::AddAttributeValue(FName Name, const AttrType& Value, float Parameter)
 {
-	using AttrSplineType = TTangentBezierSpline<AttrType>;
-	using ControlPointType = TTangentBezierControlPoint<AttrType>;
+	using AttrSplineType = UE::Geometry::Spline::TTangentBezierSpline<AttrType>;
+	using ControlPointType = UE::Geometry::Spline::TTangentBezierControlPoint<AttrType>;
 	
 	if (AttrSplineType* Channel = GetTypedAttributeChannel<AttrSplineType>(Name))
 	{
@@ -374,14 +373,14 @@ int32 FNewSpline::AddAttributeValue(FName Name, const AttrType& Value, float Par
 			if (ParentSpaceParameter > MappedChildSpace.Min)
 			{
 				Channel->AppendPoint(ControlPoint);
-				Channel->Reparameterize(EParameterizationPolicy::Uniform);
+				Channel->Reparameterize(UE::Geometry::Spline::EParameterizationPolicy::Uniform);
 				SetAttributeChannelRange(Name, FInterval1f(MappedChildSpace.Min, ParentSpaceParameter), EMappingRangeSpace::Parent);
 				return 1;
 			}
 			else
 			{
 				Channel->PrependPoint(ControlPoint);
-				Channel->Reparameterize(EParameterizationPolicy::Uniform);
+				Channel->Reparameterize(UE::Geometry::Spline::EParameterizationPolicy::Uniform);
 				SetAttributeChannelRange(Name, FInterval1f(ParentSpaceParameter, MappedChildSpace.Min), EMappingRangeSpace::Parent);
 				return 0;
 			}
@@ -429,7 +428,7 @@ int32 FNewSpline::AddAttributeValue(FName Name, const AttrType& Value, float Par
 template <typename AttrType>
 void FNewSpline::SetAttributeValue(FName Name, const AttrType& Value, int32 Index)
 {
-	using AttrSplineType = TTangentBezierSpline<AttrType>;
+	using AttrSplineType = UE::Geometry::Spline::TTangentBezierSpline<AttrType>;
 
 	if (AttrSplineType* Channel = GetTypedAttributeChannel<AttrSplineType>(Name))
 	{
@@ -440,7 +439,7 @@ void FNewSpline::SetAttributeValue(FName Name, const AttrType& Value, int32 Inde
 template <typename AttrType>
 void FNewSpline::RemoveAttributeValue(FName Name, int32 Index)
 {
-	using AttrSplineType = TTangentBezierSpline<AttrType>;
+	using AttrSplineType = UE::Geometry::Spline::TTangentBezierSpline<AttrType>;
 
 	if (AttrSplineType* Channel = GetTypedAttributeChannel<AttrSplineType>(Name))
 	{
@@ -467,7 +466,7 @@ void FNewSpline::RemoveAttributeValue(FName Name, int32 Index)
 template <typename AttrType>
 AttrType FNewSpline::GetAttributeValue(const FName& Name, int32 Index) const
 {
-	using AttrSplineType = TTangentBezierSpline<AttrType>;
+	using AttrSplineType = UE::Geometry::Spline::TTangentBezierSpline<AttrType>;
 	
 	if (AttrSplineType* Channel = GetTypedAttributeChannel<AttrSplineType>(Name))
 	{
@@ -480,7 +479,7 @@ AttrType FNewSpline::GetAttributeValue(const FName& Name, int32 Index) const
 template <typename AttrType>
 int32 FNewSpline::NumAttributeValues(FName Name) const
 {
-	using AttrSplineType = TTangentBezierSpline<AttrType>;
+	using AttrSplineType = UE::Geometry::Spline::TTangentBezierSpline<AttrType>;
 	
 	if (AttrSplineType* Channel = GetTypedAttributeChannel<AttrSplineType>(Name))
 	{
@@ -493,7 +492,7 @@ int32 FNewSpline::NumAttributeValues(FName Name) const
 template <typename AttrType>
 int32 FNewSpline::SetAttributeParameter(const FName& Name, int32 Index, float Parameter)
 {
-	using AttrSplineType = TTangentBezierSpline<AttrType>;
+	using AttrSplineType = UE::Geometry::Spline::TTangentBezierSpline<AttrType>;
 	
 	if (AttrSplineType* Channel = GetTypedAttributeChannel<AttrSplineType>(Name))
 	{
@@ -505,9 +504,8 @@ int32 FNewSpline::SetAttributeParameter(const FName& Name, int32 Index, float Pa
 			const FInterval1f ChannelSpace = Channel->GetParameterSpace();
 			const float KnotOffset = -1.f * ChannelSpace.Min;
 			const float KnotScaleFactor = 1.f / (ChannelSpace.Max - ChannelSpace.Min);
-
-			TArray<FKnot> Knots = Channel->GetKnotVector();
-			for (FKnot& Knot : Knots)
+			TArray<UE::Geometry::Spline::FKnot> Knots = Channel->GetKnotVector();
+			for (UE::Geometry::Spline::FKnot& Knot : Knots)
 			{
 				Knot.Value += KnotOffset;
 				Knot.Value *= KnotScaleFactor;
@@ -703,7 +701,7 @@ int32 FNewSpline::SetAttributeParameter(const FName& Name, int32 Index, float Pa
 template <typename AttrType>
 float FNewSpline::GetAttributeParameter(const FName& Name, int32 Index) const
 {
-	using AttrSplineType = TTangentBezierSpline<AttrType>;
+	using AttrSplineType = UE::Geometry::Spline::TTangentBezierSpline<AttrType>;
 	
 	if (AttrSplineType* Channel = GetTypedAttributeChannel<AttrSplineType>(Name))
 	{
@@ -1311,10 +1309,10 @@ bool FSpline::CreateAttributeChannel(FName AttributeName) const
 {
 	if (NewData && SupportsAttributes())
 	{
-		if (TTangentBezierSpline<AttrType>* NewAttrChannel = NewData->CreateAttributeChannel<TTangentBezierSpline<AttrType>>(AttributeName); NewAttrChannel != nullptr)
+		if (UE::Geometry::Spline::TTangentBezierSpline<AttrType>* NewAttrChannel = NewData->CreateAttributeChannel<UE::Geometry::Spline::TTangentBezierSpline<AttrType>>(AttributeName); NewAttrChannel != nullptr)
 		{
-			NewAttrChannel->SetPreInfinityMode(EOutOfBoundsHandlingMode::Constant);
-			NewAttrChannel->SetPostInfinityMode(EOutOfBoundsHandlingMode::Constant);
+			NewAttrChannel->SetPreInfinityMode(UE::Geometry::Spline::EOutOfBoundsHandlingMode::Constant);
+			NewAttrChannel->SetPostInfinityMode(UE::Geometry::Spline::EOutOfBoundsHandlingMode::Constant);
 		
 			return true;
 		}
@@ -1710,11 +1708,11 @@ void FSpline::UpdateSpline()
 	UpdateSpline(CachedUpdateSplineParams);
 }
 
-TSharedPtr<TSplineInterface<FVector>> FSpline::GetSplineInterface() const
+TSharedPtr<UE::Geometry::Spline::TSplineInterface<FVector>> FSpline::GetSplineInterface() const
 {
 	if (IsNew())
 	{
-		return StaticCastSharedPtr<TSplineInterface<FVector>>(NewData);
+		return StaticCastSharedPtr<UE::Geometry::Spline::TSplineInterface<FVector>>(NewData);
 	}
 
 	return nullptr;
@@ -1780,10 +1778,10 @@ FNewSpline::FNewSpline()
 {
 	// This is a hack. We just need to make sure that this spline type registers itself with the spline registry.
 	// This should be removed in the future once we can register this spline type automatically without instantiation.
-	TTangentBezierSpline<float> AutoRegister;
+	UE::Geometry::Spline::TTangentBezierSpline<float> AutoRegister;
 	
-	CreateAttributeChannel<TTangentBezierSpline<FQuat>>(RotationAttrName);
-	CreateAttributeChannel<TTangentBezierSpline<FVector>>(ScaleAttrName);
+	CreateAttributeChannel<UE::Geometry::Spline::TTangentBezierSpline<FQuat>>(RotationAttrName);
+	CreateAttributeChannel<UE::Geometry::Spline::TTangentBezierSpline<FVector>>(ScaleAttrName);
 	
 	ValidateRotScale();
 }
@@ -1797,8 +1795,8 @@ FNewSpline::FNewSpline(const FNewSpline& Other)
 
 FNewSpline::FNewSpline(const FLegacySpline& Other)
 {
-	CreateAttributeChannel<TTangentBezierSpline<FQuat>>(RotationAttrName);
-	CreateAttributeChannel<TTangentBezierSpline<FVector>>(ScaleAttrName);
+	CreateAttributeChannel<UE::Geometry::Spline::TTangentBezierSpline<FQuat>>(RotationAttrName);
+	CreateAttributeChannel<UE::Geometry::Spline::TTangentBezierSpline<FVector>>(ScaleAttrName);
 	
 	TArray<FTangentBezierControlPoint> Points;
     for (int32 Idx = 0; Idx < Other.GetNumControlPoints(); Idx++)
@@ -1833,8 +1831,8 @@ FNewSpline::FNewSpline(const FLegacySpline& Other)
 
 FNewSpline::FNewSpline(const FSplineCurves& Other)
 {
-	CreateAttributeChannel<TTangentBezierSpline<FQuat>>(RotationAttrName);
-	CreateAttributeChannel<TTangentBezierSpline<FVector>>(ScaleAttrName);
+	CreateAttributeChannel<UE::Geometry::Spline::TTangentBezierSpline<FQuat>>(RotationAttrName);
+	CreateAttributeChannel<UE::Geometry::Spline::TTangentBezierSpline<FVector>>(ScaleAttrName);
 	
 	TArray<FTangentBezierControlPoint> Points;
     for (int32 Idx = 0; Idx < Other.Position.Points.Num(); Idx++)
@@ -1946,7 +1944,7 @@ float FNewSpline::GetSegmentLength(const int32 Index, const float Param, const F
 	const auto& T1 = GetSpline().GetTangentOut(Index == LastPoint ? 0 : Index + 1);
 
 	// Special cases for linear
-	if (GetSpline().GetTangentModes()[Index] == ETangentMode::Linear)
+	if (GetSpline().GetTangentModes()[Index] == UE::Geometry::Spline::ETangentMode::Linear)
 	{
 		return ((P1 - P0) * Scale3D).Size() * Param;
 	}
@@ -1982,11 +1980,11 @@ void FNewSpline::SetClosedLoop(bool bInClosedLoop)
 {
     GetSpline().SetClosedLoop(bInClosedLoop);
 	
-	TTangentBezierSpline<FVector>* ScaleChild = GetTypedAttributeChannel<TTangentBezierSpline<FVector>>(ScaleAttrName);
+	UE::Geometry::Spline::TTangentBezierSpline<FVector>* ScaleChild = GetTypedAttributeChannel<UE::Geometry::Spline::TTangentBezierSpline<FVector>>(ScaleAttrName);
 	ScaleChild->SetClosedLoop(bInClosedLoop);
 	ScaleChild->SetKnotVector(GetSpline().GetKnotVector());
 		
-	TTangentBezierSpline<FQuat>* RotChild = GetTypedAttributeChannel<TTangentBezierSpline<FQuat>>(RotationAttrName);
+	UE::Geometry::Spline::TTangentBezierSpline<FQuat>* RotChild = GetTypedAttributeChannel<UE::Geometry::Spline::TTangentBezierSpline<FQuat>>(RotationAttrName);
 	RotChild->SetClosedLoop(bInClosedLoop);
 	RotChild->SetKnotVector(GetSpline().GetKnotVector());
 	
@@ -1996,7 +1994,7 @@ void FNewSpline::SetClosedLoop(bool bInClosedLoop)
 void FNewSpline::Reset()
 {
 	Clear();
-	GetSpline().SetTangentModes(TArray<ETangentMode>());
+	GetSpline().SetTangentModes(TArray<UE::Geometry::Spline::ETangentMode>());
 	PositionCurve.Reset();
 	RotationCurve.Reset();
 	ScaleCurve.Reset();
@@ -2012,7 +2010,7 @@ void FNewSpline::ResetRotation()
     }
 
 	SetAttributeChannelRange(RotationAttrName, FInterval1f(0.f, 1.f), EMappingRangeSpace::Normalized);
-	TTangentBezierSpline<FQuat>* RotSpline = GetTypedAttributeChannel<TTangentBezierSpline<FQuat>>(RotationAttrName);
+	UE::Geometry::Spline::TTangentBezierSpline<FQuat>* RotSpline = GetTypedAttributeChannel<UE::Geometry::Spline::TTangentBezierSpline<FQuat>>(RotationAttrName);
 	RotSpline->SetKnotVector(GetSpline().GetKnotVector());
 }
 
@@ -2026,7 +2024,7 @@ void FNewSpline::ResetScale()
     }
 
 	SetAttributeChannelRange(ScaleAttrName, FInterval1f(0.f, 1.f), EMappingRangeSpace::Normalized);
-	TTangentBezierSpline<FVector>* ScaleSpline = GetTypedAttributeChannel<TTangentBezierSpline<FVector>>(ScaleAttrName);
+	UE::Geometry::Spline::TTangentBezierSpline<FVector>* ScaleSpline = GetTypedAttributeChannel<UE::Geometry::Spline::TTangentBezierSpline<FVector>>(ScaleAttrName);
 	ScaleSpline->SetKnotVector(GetSpline().GetKnotVector());
 }
 
@@ -2051,12 +2049,12 @@ void FNewSpline::AddPoint(const FSplinePoint& Point)
 	{
 		GetSpline().AppendPoint(ControlPoint);
 
-		TTangentBezierSpline<FVector>* ScaleChild = GetTypedAttributeChannel<TTangentBezierSpline<FVector>>(ScaleAttrName);
-		ScaleChild->AppendPoint(TTangentBezierSpline<FVector>::FTangentBezierControlPoint(FVector::OneVector));
+		UE::Geometry::Spline::TTangentBezierSpline<FVector>* ScaleChild = GetTypedAttributeChannel<UE::Geometry::Spline::TTangentBezierSpline<FVector>>(ScaleAttrName);
+		ScaleChild->AppendPoint(UE::Geometry::Spline::TTangentBezierSpline<FVector>::FTangentBezierControlPoint(FVector::OneVector));
 		ScaleChild->SetKnotVector(GetSpline().GetKnotVector());
 		
-		TTangentBezierSpline<FQuat>* RotChild = GetTypedAttributeChannel<TTangentBezierSpline<FQuat>>(RotationAttrName);
-		RotChild->AppendPoint(TTangentBezierSpline<FQuat>::FTangentBezierControlPoint(FQuat::Identity));
+		UE::Geometry::Spline::TTangentBezierSpline<FQuat>* RotChild = GetTypedAttributeChannel<UE::Geometry::Spline::TTangentBezierSpline<FQuat>>(RotationAttrName);
+		RotChild->AppendPoint(UE::Geometry::Spline::TTangentBezierSpline<FQuat>::FTangentBezierControlPoint(FQuat::Identity));
 		RotChild->SetKnotVector(GetSpline().GetKnotVector());
 		
 		MarkLegacyCurvesDirty();
@@ -2066,12 +2064,12 @@ void FNewSpline::AddPoint(const FSplinePoint& Point)
     {
     	GetSpline().PrependPoint(ControlPoint);
     	
-    	TTangentBezierSpline<FVector>* ScaleChild = GetTypedAttributeChannel<TTangentBezierSpline<FVector>>(ScaleAttrName);
-    	ScaleChild->PrependPoint(TTangentBezierSpline<FVector>::FTangentBezierControlPoint(FVector::OneVector));
+    	UE::Geometry::Spline::TTangentBezierSpline<FVector>* ScaleChild = GetTypedAttributeChannel<UE::Geometry::Spline::TTangentBezierSpline<FVector>>(ScaleAttrName);
+    	ScaleChild->PrependPoint(UE::Geometry::Spline::TTangentBezierSpline<FVector>::FTangentBezierControlPoint(FVector::OneVector));
     	ScaleChild->SetKnotVector(GetSpline().GetKnotVector());
     
-    	TTangentBezierSpline<FQuat>* RotChild = GetTypedAttributeChannel<TTangentBezierSpline<FQuat>>(RotationAttrName);
-    	RotChild->PrependPoint(TTangentBezierSpline<FQuat>::FTangentBezierControlPoint(FQuat::Identity));
+    	UE::Geometry::Spline::TTangentBezierSpline<FQuat>* RotChild = GetTypedAttributeChannel<UE::Geometry::Spline::TTangentBezierSpline<FQuat>>(RotationAttrName);
+    	RotChild->PrependPoint(UE::Geometry::Spline::TTangentBezierSpline<FQuat>::FTangentBezierControlPoint(FQuat::Identity));
     	RotChild->SetKnotVector(GetSpline().GetKnotVector());
     	
     	MarkLegacyCurvesDirty();
@@ -2083,13 +2081,13 @@ void FNewSpline::AddPoint(const FSplinePoint& Point)
 	float NewPointInternalParam = GetSpline().GetParameter(InsertedIndex);
 	
 	float ScaleParam = MapParameterToChildSpace(ScaleAttrName, NewPointInternalParam);
-	TTangentBezierSpline<FVector>* ScaleChild = GetTypedAttributeChannel<TTangentBezierSpline<FVector>>(ScaleAttrName);
-	ScaleChild->InsertPointAtGlobalParam(ScaleParam, TTangentBezierSpline<FVector>::FTangentBezierControlPoint(FVector::OneVector));
+	UE::Geometry::Spline::TTangentBezierSpline<FVector>* ScaleChild = GetTypedAttributeChannel<UE::Geometry::Spline::TTangentBezierSpline<FVector>>(ScaleAttrName);
+	ScaleChild->InsertPointAtGlobalParam(ScaleParam, UE::Geometry::Spline::TTangentBezierSpline<FVector>::FTangentBezierControlPoint(FVector::OneVector));
 	ScaleChild->SetKnotVector(GetSpline().GetKnotVector());
 	
 	float RotParam = MapParameterToChildSpace(RotationAttrName, NewPointInternalParam);
-	TTangentBezierSpline<FQuat>* RotChild = GetTypedAttributeChannel<TTangentBezierSpline<FQuat>>(RotationAttrName);
-	RotChild->InsertPointAtGlobalParam(RotParam, TTangentBezierSpline<FQuat>::FTangentBezierControlPoint(FQuat::Identity));
+	UE::Geometry::Spline::TTangentBezierSpline<FQuat>* RotChild = GetTypedAttributeChannel<UE::Geometry::Spline::TTangentBezierSpline<FQuat>>(RotationAttrName);
+	RotChild->InsertPointAtGlobalParam(RotParam, UE::Geometry::Spline::TTangentBezierSpline<FQuat>::FTangentBezierControlPoint(FQuat::Identity));
 	RotChild->SetKnotVector(GetSpline().GetKnotVector());
 	
 	MarkLegacyCurvesDirty();
@@ -2113,12 +2111,12 @@ void FNewSpline::InsertPoint(const FSplinePoint& Point, int32 Index)
 	{
 		GetSpline().AppendPoint(ControlPoint);
 		
-		TTangentBezierSpline<FVector>* ScaleChild = GetTypedAttributeChannel<TTangentBezierSpline<FVector>>(ScaleAttrName);
-		ScaleChild->AppendPoint(TTangentBezierSpline<FVector>::FTangentBezierControlPoint(FVector::OneVector));
+		UE::Geometry::Spline::TTangentBezierSpline<FVector>* ScaleChild = GetTypedAttributeChannel<UE::Geometry::Spline::TTangentBezierSpline<FVector>>(ScaleAttrName);
+		ScaleChild->AppendPoint(UE::Geometry::Spline::TTangentBezierSpline<FVector>::FTangentBezierControlPoint(FVector::OneVector));
 		ScaleChild->SetKnotVector(GetSpline().GetKnotVector());
 		
-		TTangentBezierSpline<FQuat>* RotChild = GetTypedAttributeChannel<TTangentBezierSpline<FQuat>>(RotationAttrName);
-		RotChild->AppendPoint(TTangentBezierSpline<FQuat>::FTangentBezierControlPoint(FQuat::Identity));
+		UE::Geometry::Spline::TTangentBezierSpline<FQuat>* RotChild = GetTypedAttributeChannel<UE::Geometry::Spline::TTangentBezierSpline<FQuat>>(RotationAttrName);
+		RotChild->AppendPoint(UE::Geometry::Spline::TTangentBezierSpline<FQuat>::FTangentBezierControlPoint(FQuat::Identity));
 		RotChild->SetKnotVector(GetSpline().GetKnotVector());
 		
 		MarkLegacyCurvesDirty();
@@ -2128,12 +2126,12 @@ void FNewSpline::InsertPoint(const FSplinePoint& Point, int32 Index)
 	{
 		GetSpline().PrependPoint(ControlPoint);
 		
-		TTangentBezierSpline<FVector>* ScaleChild = GetTypedAttributeChannel<TTangentBezierSpline<FVector>>(ScaleAttrName);
-		ScaleChild->PrependPoint(TTangentBezierSpline<FVector>::FTangentBezierControlPoint(FVector::OneVector));
+		UE::Geometry::Spline::TTangentBezierSpline<FVector>* ScaleChild = GetTypedAttributeChannel<UE::Geometry::Spline::TTangentBezierSpline<FVector>>(ScaleAttrName);
+		ScaleChild->PrependPoint(UE::Geometry::Spline::TTangentBezierSpline<FVector>::FTangentBezierControlPoint(FVector::OneVector));
 		ScaleChild->SetKnotVector(GetSpline().GetKnotVector());
 	
-		TTangentBezierSpline<FQuat>* RotChild = GetTypedAttributeChannel<TTangentBezierSpline<FQuat>>(RotationAttrName);
-		RotChild->PrependPoint(TTangentBezierSpline<FQuat>::FTangentBezierControlPoint(FQuat::Identity));
+		UE::Geometry::Spline::TTangentBezierSpline<FQuat>* RotChild = GetTypedAttributeChannel<UE::Geometry::Spline::TTangentBezierSpline<FQuat>>(RotationAttrName);
+		RotChild->PrependPoint(UE::Geometry::Spline::TTangentBezierSpline<FQuat>::FTangentBezierControlPoint(FQuat::Identity));
 		RotChild->SetKnotVector(GetSpline().GetKnotVector());
 		
 		MarkLegacyCurvesDirty();
@@ -2146,13 +2144,13 @@ void FNewSpline::InsertPoint(const FSplinePoint& Point, int32 Index)
 	float NewPointInternalParam = GetSpline().GetParameter(InsertedIndex);
 	
 	float ScaleParam = MapParameterToChildSpace(ScaleAttrName, NewPointInternalParam);
-	TTangentBezierSpline<FVector>* ScaleChild = GetTypedAttributeChannel<TTangentBezierSpline<FVector>>(ScaleAttrName);
-	ScaleChild->InsertPointAtGlobalParam(ScaleParam, TTangentBezierSpline<FVector>::FTangentBezierControlPoint(FVector::OneVector));
+	UE::Geometry::Spline::TTangentBezierSpline<FVector>* ScaleChild = GetTypedAttributeChannel<UE::Geometry::Spline::TTangentBezierSpline<FVector>>(ScaleAttrName);
+	ScaleChild->InsertPointAtGlobalParam(ScaleParam, UE::Geometry::Spline::TTangentBezierSpline<FVector>::FTangentBezierControlPoint(FVector::OneVector));
 	ScaleChild->SetKnotVector(GetSpline().GetKnotVector());
 	
 	float RotParam = MapParameterToChildSpace(RotationAttrName, NewPointInternalParam);
-	TTangentBezierSpline<FQuat>* RotChild = GetTypedAttributeChannel<TTangentBezierSpline<FQuat>>(RotationAttrName);
-	RotChild->InsertPointAtGlobalParam(RotParam, TTangentBezierSpline<FQuat>::FTangentBezierControlPoint(FQuat::Identity));
+	UE::Geometry::Spline::TTangentBezierSpline<FQuat>* RotChild = GetTypedAttributeChannel<UE::Geometry::Spline::TTangentBezierSpline<FQuat>>(RotationAttrName);
+	RotChild->InsertPointAtGlobalParam(RotParam, UE::Geometry::Spline::TTangentBezierSpline<FQuat>::FTangentBezierControlPoint(FQuat::Identity));
 	RotChild->SetKnotVector(GetSpline().GetKnotVector());
 	
 	MarkLegacyCurvesDirty();
@@ -2198,11 +2196,11 @@ void FNewSpline::RemovePoint(int32 Index)
 	
 	GetSpline().RemovePoint(Index);
 
-	TTangentBezierSpline<FVector>* ScaleChild = GetTypedAttributeChannel<TTangentBezierSpline<FVector>>(ScaleAttrName);
+	UE::Geometry::Spline::TTangentBezierSpline<FVector>* ScaleChild = GetTypedAttributeChannel<UE::Geometry::Spline::TTangentBezierSpline<FVector>>(ScaleAttrName);
 	ScaleChild->RemovePoint(Index);
 	ScaleChild->SetKnotVector(GetSpline().GetKnotVector());
 	
-	TTangentBezierSpline<FQuat>* RotChild = GetTypedAttributeChannel<TTangentBezierSpline<FQuat>>(RotationAttrName);
+	UE::Geometry::Spline::TTangentBezierSpline<FQuat>* RotChild = GetTypedAttributeChannel<UE::Geometry::Spline::TTangentBezierSpline<FQuat>>(RotationAttrName);
 	RotChild->RemovePoint(Index);
 	RotChild->SetKnotVector(GetSpline().GetKnotVector());
 	
@@ -2239,7 +2237,7 @@ void FNewSpline::SetInTangent(const int32 Index, const FVector& InTangent)
     	return;
     }
 
-	GetSpline().SetPointTangentMode(Index, ETangentMode::User);
+	GetSpline().SetPointTangentMode(Index, UE::Geometry::Spline::ETangentMode::User);
     GetSpline().SetTangentIn(Index, InTangent);
 
 	MarkReparamTableDirty();
@@ -2263,7 +2261,7 @@ void FNewSpline::SetOutTangent(const int32 Index, const FVector& OutTangent)
     	return;
     }
 	
-	GetSpline().SetPointTangentMode(Index, ETangentMode::User);
+	GetSpline().SetPointTangentMode(Index, UE::Geometry::Spline::ETangentMode::User);
 	GetSpline().SetTangentOut(Index, OutTangent);
 
 	MarkReparamTableDirty();
@@ -2347,7 +2345,7 @@ EInterpCurveMode FNewSpline::GetSplinePointType(int32 Index) const
     {
     	return CIM_Unknown;
     }
-	ETangentMode TangentMode = GetSpline().GetTangentMode(Index);
+	UE::Geometry::Spline::ETangentMode TangentMode = GetSpline().GetTangentMode(Index);
     return ConvertTangentModeToInterpCurveMode(TangentMode);
 }
 
@@ -2662,8 +2660,8 @@ void FNewSpline::ValidateRotScale() const
 	}
 	
 	// Validate child existence
-	TTangentBezierSpline<FQuat>* RotSpline = GetTypedAttributeChannel<TTangentBezierSpline<FQuat>>(RotationAttrName);
-	TTangentBezierSpline<FVector>* ScaleSpline = GetTypedAttributeChannel<TTangentBezierSpline<FVector>>(ScaleAttrName);
+	UE::Geometry::Spline::TTangentBezierSpline<FQuat>* RotSpline = GetTypedAttributeChannel<UE::Geometry::Spline::TTangentBezierSpline<FQuat>>(RotationAttrName);
+	UE::Geometry::Spline::TTangentBezierSpline<FVector>* ScaleSpline = GetTypedAttributeChannel<UE::Geometry::Spline::TTangentBezierSpline<FVector>>(ScaleAttrName);
 	ensure (RotSpline && ScaleSpline);
 
 	if (RotSpline && ScaleSpline)
diff --git a/Engine/Source/Runtime/GeometryCore/Public/Splines/BSpline.h b/Engine/Source/Runtime/GeometryCore/Public/Splines/BSpline.h
--- a/Engine/Source/Runtime/GeometryCore/Public/Splines/BSpline.h
+++ b/Engine/Source/Runtime/GeometryCore/Public/Splines/BSpline.h
@@ -14,13 +14,13 @@ namespace Geometry
 {
 namespace Spline
 {
-enum class EParameterizationPolicy 
+enum class EParameterizationPolicy
 {
 	/** Equal spacing between knots regardless of control point positions */
-	Uniform,        
+	Uniform,
 
 	/** Spacing based on actual distances between control points - good for uneven point distribution */
-	ChordLength,    
+	ChordLength,
 
 	/** Uses square root of distances - helps prevent loops and cusps in the curve */
 	Centripetal,
@@ -31,9 +31,17 @@ struct FKnot
 	float Value;
 	uint32 Multiplicity;
 
-	FKnot() : Value(0.0f), Multiplicity(1) {}
-	FKnot(float InValue, uint32 InMultiplicity = 1) 
-		: Value(InValue), Multiplicity(InMultiplicity) {}
+	FKnot()
+		: Value(0.0f)
+		, Multiplicity(1)
+	{
+	}
+
+	FKnot(float InValue, uint32 InMultiplicity = 1)
+		: Value(InValue)
+		, Multiplicity(InMultiplicity)
+	{
+	}
 
 	/**
 	 * Converts a pair-based knot vector to a flat knot vector
@@ -57,7 +65,7 @@ struct FKnot
 	{
 		TArray<float> FlatKnots = ConvertPairToFlatKnots(InKnots);
 		// Need at least 2*(degree+1) knots
-		if (FlatKnots.Num() < 2*(Degree + 1))
+		if (FlatKnots.Num() < 2 * (Degree + 1))
 		{
 			return false;
 		}
@@ -65,7 +73,7 @@ struct FKnot
 		// Check monotonicity
 		for (int32 i = 1; i < FlatKnots.Num(); ++i)
 		{
-			if (FlatKnots[i] < FlatKnots[i-1])
+			if (FlatKnots[i] < FlatKnots[i - 1])
 			{
 				return false;
 			}
@@ -86,31 +94,31 @@ struct FKnot
 		return Ar;
 	}
 };
-	
+
 // Generate compile-time type ID for this BSpline specialization
 // Using degree in the implementation name for uniqueness
-template<int Degree>
+template <int Degree>
 struct BSplineNameSelector
 {
 	static constexpr const TCHAR* Name = TEXT("BSpline");
 };
 
-template<>
+template <>
 struct BSplineNameSelector<2>
 {
 	static constexpr const TCHAR* Name = TEXT("BSpline2");
 };
 
-template<>
+template <>
 struct BSplineNameSelector<3>
 {
 	static constexpr const TCHAR* Name = TEXT("BSpline3");
 };
-	
+
 template <typename VALUETYPE, int32 DEGREE>
-class TBSpline :
-	public TSplineInterface<VALUETYPE>,
-	private TSelfRegisteringSpline<TBSpline<VALUETYPE, DEGREE>, VALUETYPE>
+class TBSpline
+	: public TSplineInterface<VALUETYPE>,
+	  private TSelfRegisteringSpline<TBSpline<VALUETYPE, DEGREE>, VALUETYPE>
 {
 public:
 	static constexpr int32 Degree = DEGREE;
@@ -118,21 +126,25 @@ public:
 
 	static constexpr int32 WindowSize = Degree + 1;
 	using FWindow = TStaticArray<const ValueType*, WindowSize>;
-	
+
 	DECLARE_SPLINE_TYPE_ID(
 		BSplineNameSelector<DEGREE>::Name,
 		*TSplineValueTypeTraits<VALUETYPE>::Name
 	);
-	
+
 	TBSpline() = default;
 	virtual ~TBSpline() override = default;
 
-	template<typename T, typename = void>
-	struct THasToString : std::false_type {};
+	template <typename T, typename = void>
+	struct THasToString : std::false_type
+	{
+	};
+
+	template <typename T>
+	struct THasToString<T, std::void_t<decltype(std::declval<T>().ToString())>> : std::true_type
+	{
+	};
 
-	template<typename T>
-	struct THasToString<T, std::void_t<decltype(std::declval<T>().ToString())>> : std::true_type {};
-	
 	void Dump() const
 	{
 		if constexpr (std::is_floating_point_v<ValueType>)
@@ -154,9 +166,8 @@ public:
 
 		UE_LOG(LogSpline, Verbose, TEXT("Knots:"));
 		PrintKnotVector();
-		
 	}
-	
+
 	// ISplineInterface Implementation
 	virtual void Clear() override
 	{
@@ -171,10 +182,10 @@ public:
 			const TBSpline* Other = static_cast<const TBSpline*>(OtherSpline);
 			return operator==(*Other);
 		}
-		
+
 		return false;
 	}
-	
+
 	virtual bool Serialize(FArchive& Ar) override
 	{
 		// Call immediate parent's Serialize
@@ -182,7 +193,7 @@ public:
 		{
 			return false;
 		}
-    	
+
 		Ar << Values;
 		Ar << PairKnots;
 		Ar << bIsClosedLoop;
@@ -190,36 +201,50 @@ public:
 		MarkFlatKnotsCacheDirty();
 		return true;
 	}
+
 	bool operator==(const TBSpline& Other) const
 	{
-		return	Values == Other.Values &&
-				PairKnots == Other.PairKnots &&
-				bIsClosedLoop == Other.bIsClosedLoop &&
-				bClampEnds == Other.bClampEnds;
+		return Values == Other.Values &&
+			PairKnots == Other.PairKnots &&
+			bIsClosedLoop == Other.bIsClosedLoop &&
+			bClampEnds == Other.bClampEnds;
 	}
+
 	friend FArchive& operator<<(FArchive& Ar, TBSpline& BSpline)
 	{
 		BSpline.Serialize(Ar);
-		return Ar;	
+		return Ar;
+	}
+
+	virtual FInterval1f GetParameterSpace() const override
+	{
+		return GetKnotRange();
 	}
-	
-	virtual FInterval1f GetParameterSpace() const override { return GetKnotRange(); }
-	virtual void SetClosedLoop(bool bClosed) override { bIsClosedLoop = bClosed; }
-	virtual bool IsClosedLoop() const override { return bIsClosedLoop; }
+
+	virtual void SetClosedLoop(bool bClosed) override
+	{
+		bIsClosedLoop = bClosed;
+	}
+
+	virtual bool IsClosedLoop() const override
+	{
+		return bIsClosedLoop;
+	}
+
 	virtual TUniquePtr<ISplineInterface> Clone() const override
 	{
 		TUniquePtr<TBSpline<VALUETYPE, DEGREE>> Clone = MakeUnique<TBSpline<VALUETYPE, DEGREE>>();
-    
+
 		// Copy all member variables
 		Clone->Values = this->Values;
 		Clone->PairKnots = this->PairKnots;
 		Clone->bIsClosedLoop = this->bIsClosedLoop;
 		Clone->bClampEnds = this->bClampEnds;
-    
+
 		// Copy infinity modes
 		Clone->PreInfinityMode = this->PreInfinityMode;
 		Clone->PostInfinityMode = this->PostInfinityMode;
-    
+
 		return Clone;
 	}
 
@@ -242,7 +267,7 @@ public:
 		ensureAlwaysMsgf(false, TEXT("TBSpline::GetSegmentParameterRange is unimplemented!"));
 		return FInterval1f();
 	}
-	
+
 	// TSplineInterface Implementation
 	virtual ValueType EvaluateImpl(float Parameter) const override
 	{
@@ -254,15 +279,30 @@ public:
 
 		// If FindWindow fails, it will return an array of nullptr.
 		if (Window[0] == nullptr) return ValueType();
-        
+
 		return InterpolateWindow(Window, Parameter);
 	}
+
 	virtual float FindNearest(const ValueType& Point, float& OutSquaredDistance) const override
-	{ return 0;/* todo */ }
+	{
+		return 0;/* todo */
+	}
+
+	int32 NumKeys() const
+	{
+		return Values.Num();
+	}
+
+	const ValueType& GetValue(int32 Idx) const
+	{
+		return Values[Idx];
+	}
+
+	int32 AddValue(const ValueType& NewValue)
+	{
+		return Values.Add(NewValue);
+	}
 
-	int32 NumKeys() const { return Values.Num(); }
-	const ValueType& GetValue(int32 Idx) const { return Values[Idx]; }
-	int32 AddValue(const ValueType& NewValue)  { return Values.Add(NewValue); }
 	bool SetValue(int32 Idx, const ValueType& NewValue)
 	{
 		if (!Values.IsValidIndex(Idx))
@@ -272,6 +312,7 @@ public:
 		Values[Idx] = NewValue;
 		return true;
 	}
+
 	int32 InsertValue(int32 Idx, const ValueType& NewValue)
 	{
 		if (Idx < 0 || Idx > Values.Num())
@@ -280,7 +321,8 @@ public:
 		}
 		return Values.Insert(NewValue, Idx);
 	}
-	virtual bool RemoveValue(int32 Index) 
+
+	virtual bool RemoveValue(int32 Index)
 	{
 		if (Values.IsValidIndex(Index))
 		{
@@ -289,16 +331,16 @@ public:
 		}
 		return false;
 	}
-    
+
 	// Parameterization methods
 	virtual int32 SetParameter(int32 Index, float NewParameter)
 	{
 		UE_LOG(LogSpline, Warning, TEXT("SetParameter is disabled. Use Reparameterize() or a linear spline if you need manual control."));
 		return INDEX_NONE;
 	}
-	
+
 	// Returns parameter value for a control point index using Greville Abscissae
-    virtual float GetParameter(int32 Index) const
+	virtual float GetParameter(int32 Index) const
 	{
 		UpdateFlatKnotsCache();
 		if (Index < 0 || Index + Degree >= FlatKnots.Num() - 1)
@@ -316,11 +358,11 @@ public:
 	virtual int32 FindIndexForParameter(float Parameter, float& OutLocalParam) const
 	{
 		UpdateFlatKnotsCache();
-        const TArray<float>& Knots = FlatKnots;
-		
+		const TArray<float>& Knots = FlatKnots;
+
 		if (Knots.IsEmpty())
 			return INDEX_NONE;
-        
+
 		// Clamp parameter to valid range
 		const float FirstValidKnot = Knots[Degree];
 		const float LastValidKnot = Knots[Knots.Num() - Degree - 1];
@@ -337,50 +379,59 @@ public:
 			OutLocalParam = 1.0f;
 			return NumKeys() - 2; // Last segment index
 		}
-		
+
 		if (Parameter <= FirstValidKnot + UE_SMALL_NUMBER)
 		{
 			OutLocalParam = 0.0f;
 			return 0;
 		}
-    
-    
+
+
 		// Find the appropriate span
 		int32 Span = FindKnotSpan(Parameter);
 		if (Span < Degree || Span >= Knots.Num() - 1)
 			return INDEX_NONE;
-    
+
 		int32 Index = Span - Degree;
-    
+
 		// Calculate local parameter within span, with protection against near-zero denominators
 		float SpanStart = Knots[Span];
 		float SpanEnd = Knots[Span + 1];
 		float SpanLength = SpanEnd - SpanStart;
-    
+
 		if (SpanLength < UE_SMALL_NUMBER)
 		{
 			// For nearly identical knots, use 0 or 1 based on which end is closer
-			OutLocalParam = (Parameter - SpanStart < SpanEnd - Parameter) ? 0.0f : 1.0f;
+			OutLocalParam = (Parameter - SpanStart < SpanEnd - Parameter)
+				? 0.0f
+				: 1.0f;
 		}
 		else
 		{
 			OutLocalParam = (Parameter - SpanStart) / SpanLength;
 		}
-    
+
 		return Index;
 	}
+
 	// Knot vector management
 	const TArray<FKnot>& GetKnotVector() const
 	{
 		return PairKnots;
 	}
+
 	// Get the pair-based knot vector
-    const TArray<FKnot>& GetPairKnots() const
-    {
-        return PairKnots;
-    }
-	void ResetKnotVector() { PairKnots.Reset(); FlatKnots.Reset(); }
-		
+	const TArray<FKnot>& GetPairKnots() const
+	{
+		return PairKnots;
+	}
+
+	void ResetKnotVector()
+	{
+		PairKnots.Reset();
+		FlatKnots.Reset();
+	}
+
 	/**
 	 * Returns how many times a given knot value appears.
 	 * @param KnotIndex - Knot value to check
@@ -392,7 +443,7 @@ public:
 		{
 			return PairKnots[KnotIndex].Multiplicity;
 		}
-		
+
 		return 0;
 	}
 
@@ -403,8 +454,8 @@ public:
 	 */
 	bool SetCustomKnots(const TArray<FKnot>& NewKnots)
 	{
-        const int32 RequiredKnots = GetExpectedNumKnots();
-		
+		const int32 RequiredKnots = GetExpectedNumKnots();
+
 		// Validate knot multiplicities and calculate total expanded knots  
 		int32 TotalExpandedKnots = 0;
 		for (const FKnot& Knot : NewKnots)
@@ -416,10 +467,10 @@ public:
 			}
 			TotalExpandedKnots += Knot.Multiplicity;
 		}
-    
+
 		if (TotalExpandedKnots < RequiredKnots)
 		{
-			UE_LOG(LogSpline, Warning, TEXT("Invalid knot vector size. Need %d knots, got %d."), 
+			UE_LOG(LogSpline, Warning, TEXT("Invalid knot vector size. Need %d knots, got %d."),
 				RequiredKnots, TotalExpandedKnots);
 			return false;
 		}
@@ -436,14 +487,14 @@ public:
 
 		// Directly assign the knots - no epsilon comparison needed
 		PairKnots = NewKnots;
-    
+
 		// Mark the flat knots cache as dirty
 		MarkFlatKnotsCacheDirty();
-    
+
 		UE_LOG(LogSpline, Verbose, TEXT("Set Custom knot vector - Points: %d, Degree: %d, Total Knots: %d, Unique Knots: %d"),
-			   NumKeys(), Degree, TotalExpandedKnots, NewKnots.Num());
+			NumKeys(), Degree, TotalExpandedKnots, NewKnots.Num());
 		PrintKnotVector();
-    
+
 		return true;
 	}
 
@@ -453,14 +504,14 @@ public:
 	 * @param Points - Array of control points and their Knots
 	 * @param ParameterizationPolicy - Policy for generating knot vector
 	 */
-    virtual void Reparameterize(EParameterizationPolicy ParameterizationPolicy)
+	virtual void Reparameterize(EParameterizationPolicy ParameterizationPolicy)
 	{
-    	const TArray<ValueType>& Points = Values;
+		const TArray<ValueType>& Points = Values;
 		int32 NumValues = Points.Num();
 		if (NumValues < Degree + 1)
 		{
 			UE_LOG(LogSpline, Warning, TEXT("Not enough points for B-spline: %d (need >= %d)"),
-				   NumValues, Degree + 1);
+				NumValues, Degree + 1);
 			return;
 		}
 
@@ -484,9 +535,9 @@ public:
 		MarkFlatKnotsCacheDirty();
 
 		UE_LOG(LogSpline, Verbose, TEXT("Updated knot vector - Points: %d, Degree: %d, Knots: %d"),
-			   Points.Num(), Degree, PairKnots.Num());
+			Points.Num(), Degree, PairKnots.Num());
 	}
-	
+
 	/** 
 	 * Sets whether the spline should be clamped at endpoints
 	 * When true, the curve will interpolate the first and last control points
@@ -511,9 +562,8 @@ public:
 		}
 		return FInterval1f(PairKnots[0].Value, PairKnots.Last().Value);
 	}
-	
+
 protected:
-	
 	/**
 	 *  Sets the value of a knot
 	 *  @param KnotIdx - Index of the knot to set
@@ -560,12 +610,12 @@ protected:
 		{
 			return;
 		}
-		
+
 		FKnot OldA = PairKnots[KnotIdxA];
 		PairKnots[KnotIdxA] = PairKnots[KnotIdxB];
 		PairKnots[KnotIdxB] = OldA;
 	}
-	
+
 	/**
 	 * Inserts a knot at the specified parameter value
 	 * @param InKnot - Knot to insert
@@ -574,14 +624,14 @@ protected:
 	bool InsertKnot(FKnot InKnot)
 	{
 		// Binary search the knot vector:
-		
-		int32 Low = 0;	// This will be our insert index
+
+		int32 Low = 0; // This will be our insert index
 		int32 High = PairKnots.Num() - 1;
 
 		while (Low <= High)
 		{
 			const int32 Mid = Low + (High - Low) / 2;
-			
+
 			if (PairKnots[Mid].Value == InKnot.Value)
 			{
 				UE_LOG(LogSpline, Warning, TEXT("Knot insertion failed, knot already exists!"));
@@ -606,93 +656,156 @@ protected:
 	struct FValidKnotSearchParams
 	{
 		FValidKnotSearchParams() = default;
-		
+
 		FValidKnotSearchParams(float InDesiredParameter)
 			: DesiredParameter(InDesiredParameter)
-		{}
-		
+		{
+		}
+
 		float DesiredParameter = 0.f;
 		bool bSearchRight = true;
 		bool bSearchLeft = true;
 	};
-	
+
 	/**
 	 * Finds the nearest available knot value that does not conflict with existing knots.
 	 * @param InSearchParams - Describes value to search for and how to search.
 	 * @return The nearest valid knot parameter that does not conflict with existing knots.
 	 */
 	float GetNearestAvailableKnotValue(const FValidKnotSearchParams& InSearchParams) const
-    {
+	{
 		const float& DesiredParameter = InSearchParams.DesiredParameter;
-		
-        if (!ensureAlwaysMsgf(!std::isnan(DesiredParameter), TEXT("DesiredParameter is NaN!")) ||
-        	!ensureAlwaysMsgf(InSearchParams.bSearchLeft || InSearchParams.bSearchRight, TEXT("Invalid search params!")) ||
-        	PairKnots.Num() == 0)
-        {
-            return DesiredParameter;
-        }
-
-		// Convert existing knots into a set for O(1) lookups
+
+		if (!ensureAlwaysMsgf(!std::isnan(DesiredParameter), TEXT("DesiredParameter is NaN!")) ||
+			!ensureAlwaysMsgf(InSearchParams.bSearchLeft || InSearchParams.bSearchRight, TEXT("Invalid search params!")) ||
+			PairKnots.Num() == 0)
+		{
+			return DesiredParameter;
+		}
+
+		// Build a set of existing knot values (normalize zeros to avoid -0/+0 duplicates)
 		TSet<float> KnotValuesSet;
+		KnotValuesSet.Reserve(PairKnots.Num());
+		for (const FKnot& Knot : PairKnots)
 		{
-			for (const FKnot& Knot : PairKnots)
-			{
-				KnotValuesSet.Add(Knot.Value);
-			}
+			KnotValuesSet.Add(Param::NormalizeKey(Knot.Value));
 		}
 
-		
-		// If no conflict, return immediately
-		if (!KnotValuesSet.Contains(DesiredParameter))
+		// Return the caller's exact value; NormalizeKey is only for set membership (folds -0.0f/+0.0f).
+		if (!KnotValuesSet.Contains(Param::NormalizeKey(DesiredParameter)))
 		{
 			return DesiredParameter;
 		}
-		
-        // Find the closest valid parameter by walking outward alternately
-        float Left = DesiredParameter;
-        float Right = DesiredParameter;
-        constexpr int32 MaxIterations = 100000; // Set to a high but safe value
-		int32 IterationCount = 0;
-
-		// Search until a valid parameter is found or we hit the limit
-		while (IterationCount < MaxIterations)
+
+		// Walk outward alternating sides
+		float Left = DesiredParameter;
+		float Right = DesiredParameter;
+
+		// Copy so we can disable directions mid-search
+		bool bSearchLeft = InSearchParams.bSearchLeft;
+		bool bSearchRight = InSearchParams.bSearchRight;
+
+		// Tiny dead-man budget to catch regressions; normal exits happen via guard rails.
+		int32 StepBudget = 4096;
+		for (;;)
 		{
-			IterationCount++;
+			bool AnyProgress = false;
 
-			if (InSearchParams.bSearchLeft)
+			if (bSearchLeft)
 			{
-				Left = std::nextafter(Left, -std::numeric_limits<float>::infinity());
-				if (!KnotValuesSet.Contains(Left))
+				const float NewLeft = Param::PrevDistinct(Left);
+				if (NewLeft == Left)
 				{
-					return Left;
+					bSearchLeft = false;
+				} // Guard #1: no progress
+				else
+				{
+					Left = NewLeft;
+					AnyProgress = true;
+					if (!KnotValuesSet.Contains(Param::NormalizeKey(Left)))
+					{
+						return Left;
+					}
+					if (Left <= -std::numeric_limits<float>::max()) // Guard #2: saturated
+					{
+						bSearchLeft = false;
+					}
 				}
 			}
 
-			if (InSearchParams.bSearchRight)
+			if (bSearchRight)
 			{
-				Right = std::nextafter(Right, std::numeric_limits<float>::infinity());
-				if (!KnotValuesSet.Contains(Right))
+				const float NewRight = Param::NextDistinct(Right);
+				if (NewRight == Right)
+				{
+					bSearchRight = false;
+				} // Guard #1: no progress
+				else
 				{
-					return Right;
+					Right = NewRight;
+					AnyProgress = true;
+					if (!KnotValuesSet.Contains(Param::NormalizeKey(Right)))
+					{
+						return Right;
+					}
+					if (Right >= +std::numeric_limits<float>::max()) // Guard #2: saturated
+					{
+						bSearchRight = false;
+					}
 				}
 			}
-        }
-		
-		// If no valid parameter was found within limits, return the original (or handle error)
-        // This should rarely be hit unless space is densely packed
-        ensureAlwaysMsgf(false, TEXT("Max iterations reached in GetNearestAvailableKnotValue"));
-        return DesiredParameter;
-    }
+
+			if (!bSearchLeft && !bSearchRight) break; // Guard #3: nowhere left to search
+			if (!AnyProgress) break; // Paranoid: should be covered by #1
+
+			if (--StepBudget <= 0)
+			{
+				// Dead-man: catch unexpected spins
+				ensureMsgf(false, TEXT("Step budget exhausted in GetNearestAvailableKnotValue"));
+				break;
+			}
+		}
+
+		// Fallback: never return the conflicting original. Pick a deterministic, distinct nudge.
+		{
+			// Canonicalize -0/+0, then use its *bits* to choose a side deterministically.
+			const float Canon = Param::NormalizeKey(DesiredParameter);
+			const uint32 Bits = Internal::FloatToBits(Canon);
+
+			const bool bPreferRight =
+				(InSearchParams.bSearchRight && (!InSearchParams.bSearchLeft || ((Bits & 1u) != 0)));
+
+			// Step from the *original* DesiredParameter (not Canon) so we don’t alter the value unless needed.
+			float Fallback = Param::Step(
+				DesiredParameter, bPreferRight
+				? Param::EDir::Right
+				: Param::EDir::Left);
+
+			float Probe = Param::NormalizeKey(Fallback);
+			if (KnotValuesSet.Contains(Probe))
+			{
+				// Try the opposite direction once
+				Fallback = Param::Step(
+					DesiredParameter, bPreferRight
+					? Param::EDir::Left
+					: Param::EDir::Right);
+				Probe = Param::NormalizeKey(Fallback);
+			}
+
+			ensureMsgf(!KnotValuesSet.Contains(Probe), TEXT("Fallback still conflicts; check knot packing logic."));
+			return Fallback;
+		}
+	}
 
 	virtual int32 GetExpectedNumKnots() const
 	{
 		const int32 NumValues = NumKeys();
 		// For closed loops, we repeat the first Degree points at the end
 		return this->IsClosedLoop()
-				   ? (NumValues + (2 * Degree) + 1) // closed loop: n + 2p + 1
-				   : (NumValues + Degree + 1);
+			? (NumValues + (2 * Degree) + 1) // closed loop: n + 2p + 1
+			: (NumValues + Degree + 1);
 	}
-	
+
 	/** 
      * Generates a uniform knot vector with equal spacing
      * For clamped ends, multiplicity of degree+1 is used at endpoints
@@ -700,8 +813,8 @@ protected:
 	void GenerateUniformKnots(int32 KnotCount)
 	{
 		PairKnots.Reset();
-		
-    	int32 NumPoints = NumKeys();
+
+		int32 NumPoints = NumKeys();
 		if (this->IsClosedLoop())
 		{
 			// For a closed B-spline of degree d with n points:
@@ -709,20 +822,19 @@ protected:
 			// - Need n+2d+1 knots total
 			// - Knot spacing should be uniform to maintain periodicity
 			const int32 n = NumPoints;
-			const int32 d = Degree; 
-			const int32 NumKnots = n + 2*d + 1;
-        
+			const int32 d = Degree;
+			const int32 NumKnots = n + 2 * d + 1;
+
 
 			// For periodic B-splines, use uniform knot spacing
 			for (int32 i = 0; i < NumKnots; ++i)
 			{
 				// Map to [0,n] range for n segments
-				PairKnots.Add( FKnot(static_cast<float>(i), 1) );
+				PairKnots.Add(FKnot(static_cast<float>(i), 1));
 			}
 		}
 		else
 		{
-
 			for (int32 i = 0; i < KnotCount; ++i)
 			{
 				float Value;
@@ -759,207 +871,210 @@ protected:
      * Generates a knot vector based on chord lengths between control points
      * This gives better parameterization when control points are unevenly spaced
      */
-    void GenerateChordLengthKnots(int32 KnotCount)
+	void GenerateChordLengthKnots(int32 KnotCount)
 	{
-	    const TArray<ValueType>& Points = Values;
-	    if (Points.Num() < 2)
-	    {
-	        GenerateUniformKnots(KnotCount);
-	        return;
-	    }
-	    
-	    const int32 n = Points.Num();
-	    const int32 d = Degree;
-	    const int32 NumKnots = this->IsClosedLoop() ? (n + 2*d + 1) : KnotCount;
+		const TArray<ValueType>& Points = Values;
+		if (Points.Num() < 2)
+		{
+			GenerateUniformKnots(KnotCount);
+			return;
+		}
+
+		const int32 n = Points.Num();
+		const int32 d = Degree;
+		const int32 NumKnots = this->IsClosedLoop()
+			? (n + 2 * d + 1)
+			: KnotCount;
 		PairKnots.Reset();
 
-	    // Calculate chord lengths between consecutive points
-	    TArray<float> Chords;
-	    Chords.Reserve(Points.Num() - 1);
-	    float TotalLength = 0.0f;
-	    
-	    for (int32 i = 1; i < Points.Num(); ++i)
-	    {
-	        const float Length = static_cast<float>(Math::Distance(Points[i], Points[i-1]));
-	        Chords.Add(Length);
-	        TotalLength += Length;
-	    }
-
-	    // For closed loop, add the closing segment length
-	    if (this->IsClosedLoop() && Points.Num() > 0)
-	    {
-	        const float ClosureLength = static_cast<float>(Math::Distance(Points[0], Points.Last()));
-	        Chords.Add(ClosureLength);
-	        TotalLength += ClosureLength;
-	    }
-
-	    if (this->IsClosedLoop())
-	    {
-	        // For closed loop, generate periodic knot vector based on chord lengths
-	        for (int32 i = 0; i < NumKnots; ++i)
-	        {
-	            PairKnots.Add( FKnot(static_cast<float>(i) * TotalLength / static_cast<float>(n + d)));
-	        }
-	    }
-	    else if (bClampEnds)
-	    {
-	        // For clamped ends, we need:
-	        // - First (Degree+1) knots equal to 0.0
-	        // - Last (Degree+1) knots equal to 1.0
-	        // - Middle knots based on chord lengths
-	        
-	        // First set starting knots to 0
-	    	PairKnots.Add(FKnot(0.0f, Degree+1));
-	    	
-	        // Calculate internal knots based on accumulated chord lengths
-	        // We have (KnotCount - 2*(Degree+1)) internal knots to set
-	        const int32 NumInternalKnots = KnotCount - 2 * (Degree + 1);
-	        if (NumInternalKnots > 0)
-	        {
-	            // Compute normalized Knots at each control point
-	            TArray<float> Params;
-	            Params.SetNum(Points.Num());
-	            Params[0] = 0.0f;
-	            
-	            float AccumulatedLength = 0.0f;
-	            for (int32 i = 1; i < Points.Num(); ++i)
-	            {
-	                AccumulatedLength += Chords[i-1];
-	                Params[i] = (TotalLength > UE_SMALL_NUMBER) ? 
-	                    (AccumulatedLength / TotalLength) : 
-	                    (static_cast<float>(i) / static_cast<float>(Points.Num() - 1));
-	            }
-	            
-	            // Map these Knots to internal knots
-	            for (int32 i = 0; i < NumInternalKnots; ++i)
-	            {
-	                // Map i to appropriate parameter range
-	                const float t = static_cast<float>(i + 1) / static_cast<float>(NumInternalKnots + 1);
-	                // Map t to control point indices
-	                const int32 LowIndex = FMath::FloorToInt(t * static_cast<float>(Points.Num() - 1));
-	                const int32 HighIndex = FMath::Min(LowIndex + 1, Points.Num() - 1);
-	                const float Alpha = t * static_cast<float>(Points.Num() - 1 - LowIndex);
-	                
-	                // Interpolate Knots
-	                const float KnotValue = FMath::Lerp(Params[LowIndex], Params[HighIndex], Alpha);
-	            	PairKnots.Add(FKnot(KnotValue, 1));
-	            }
-	        }
-	    	// Set the last (Degree+1) knots to 1.0
-	    	PairKnots.Add(FKnot(1.0f, Degree+1));
-
-	    }
-	    else
-	    {
-	        // Unclamped knots - simple uniform distribution based on chord lengths
-	        for (int32 i = 0; i < KnotCount; ++i)
-	        {
-	            const float Param = static_cast<float>(i) / static_cast<float>(KnotCount - 1);
-	            PairKnots.Add( FKnot(Param * TotalLength));
-	        }
-	    }
+		// Calculate chord lengths between consecutive points
+		TArray<float> Chords;
+		Chords.Reserve(Points.Num() - 1);
+		float TotalLength = 0.0f;
+
+		for (int32 i = 1; i < Points.Num(); ++i)
+		{
+			const float Length = static_cast<float>(Math::Distance(Points[i], Points[i - 1]));
+			Chords.Add(Length);
+			TotalLength += Length;
+		}
+
+		// For closed loop, add the closing segment length
+		if (this->IsClosedLoop() && Points.Num() > 0)
+		{
+			const float ClosureLength = static_cast<float>(Math::Distance(Points[0], Points.Last()));
+			Chords.Add(ClosureLength);
+			TotalLength += ClosureLength;
+		}
+
+		if (this->IsClosedLoop())
+		{
+			// For closed loop, generate periodic knot vector based on chord lengths
+			for (int32 i = 0; i < NumKnots; ++i)
+			{
+				PairKnots.Add(FKnot(static_cast<float>(i) * TotalLength / static_cast<float>(n + d)));
+			}
+		}
+		else if (bClampEnds)
+		{
+			// For clamped ends, we need:
+			// - First (Degree+1) knots equal to 0.0
+			// - Last (Degree+1) knots equal to 1.0
+			// - Middle knots based on chord lengths
+
+			// First set starting knots to 0
+			PairKnots.Add(FKnot(0.0f, Degree + 1));
+
+			// Calculate internal knots based on accumulated chord lengths
+			// We have (KnotCount - 2*(Degree+1)) internal knots to set
+			const int32 NumInternalKnots = KnotCount - 2 * (Degree + 1);
+			if (NumInternalKnots > 0)
+			{
+				// Compute normalized Knots at each control point
+				TArray<float> Params;
+				Params.SetNum(Points.Num());
+				Params[0] = 0.0f;
+
+				float AccumulatedLength = 0.0f;
+				for (int32 i = 1; i < Points.Num(); ++i)
+				{
+					AccumulatedLength += Chords[i - 1];
+					Params[i] = (TotalLength > UE_SMALL_NUMBER)
+						? (AccumulatedLength / TotalLength)
+						: (static_cast<float>(i) / static_cast<float>(Points.Num() - 1));
+				}
+
+				// Map these Knots to internal knots
+				for (int32 i = 0; i < NumInternalKnots; ++i)
+				{
+					// Map i to appropriate parameter range
+					const float t = static_cast<float>(i + 1) / static_cast<float>(NumInternalKnots + 1);
+					// Map t to control point indices
+					const int32 LowIndex = FMath::FloorToInt(t * static_cast<float>(Points.Num() - 1));
+					const int32 HighIndex = FMath::Min(LowIndex + 1, Points.Num() - 1);
+					const float Alpha = t * static_cast<float>(Points.Num() - 1 - LowIndex);
+
+					// Interpolate Knots
+					const float KnotValue = FMath::Lerp(Params[LowIndex], Params[HighIndex], Alpha);
+					PairKnots.Add(FKnot(KnotValue, 1));
+				}
+			}
+			// Set the last (Degree+1) knots to 1.0
+			PairKnots.Add(FKnot(1.0f, Degree + 1));
+		}
+		else
+		{
+			// Unclamped knots - simple uniform distribution based on chord lengths
+			for (int32 i = 0; i < KnotCount; ++i)
+			{
+				const float Param = static_cast<float>(i) / static_cast<float>(KnotCount - 1);
+				PairKnots.Add(FKnot(Param * TotalLength));
+			}
+		}
 		MarkFlatKnotsCacheDirty();
 	}
 
-    /**
-     * Generates a knot vector using centripetal parameterization
-     * Uses square root of chord lengths which helps prevent cusps and unwanted loops
-     * Often provides the most visually pleasing results for interactive curve editing
-     */
-    void GenerateCentripetalKnots(int32 KnotCount)
+	/**
+	 * Generates a knot vector using centripetal parameterization
+	 * Uses square root of chord lengths which helps prevent cusps and unwanted loops
+	 * Often provides the most visually pleasing results for interactive curve editing
+	 */
+	void GenerateCentripetalKnots(int32 KnotCount)
 	{
-    	const TArray<ValueType>& Points = Values;
-	    if (Points.Num() < 2)
-	    {
-	        GenerateUniformKnots(KnotCount);
-	        return;
-	    }
-
-	    const int32 n = Points.Num();
-	    const int32 d = Degree;
-	    const int32 NumKnots = this->IsClosedLoop() ? (n + 2*d + 1) : KnotCount;
+		const TArray<ValueType>& Points = Values;
+		if (Points.Num() < 2)
+		{
+			GenerateUniformKnots(KnotCount);
+			return;
+		}
+
+		const int32 n = Points.Num();
+		const int32 d = Degree;
+		const int32 NumKnots = this->IsClosedLoop()
+			? (n + 2 * d + 1)
+			: KnotCount;
 		PairKnots.Reset();
-	    
-	    TArray<float> Lengths;
-	    Lengths.Reserve(Points.Num() - 1);
-	    float TotalLength = 0.0f;
-	    
-	    for (int32 i = 1; i < Points.Num(); ++i)
-	    {
-	        // Square root of distance for centripetal parameterization
-	    	const float Length = static_cast<float>(Math::CentripetalDistance(Points[i], Points[i-1]));
-	        Lengths.Add(Length);
-	        TotalLength += Length;
-	    }
-	    
-	    // For closed loop, add the closing segment
-	    if (this->IsClosedLoop() && Points.Num() > 0)
-	    {
-	        const float ClosureLength = static_cast<float>(Math::CentripetalDistance(Points[0], Points.Last()));
-	        Lengths.Add(ClosureLength);
-	        TotalLength += ClosureLength;
-	    }
-
-	    if (this->IsClosedLoop())
-	    {
-	        // Generate periodic knot vector using centripetal parameterization
-	        for (int32 i = 0; i < NumKnots; ++i)
-	        {
-	            PairKnots.Add( FKnot(static_cast<float>(i) * TotalLength / static_cast<float>(n + d)));
-	        }
-	    }
-	    else if (bClampEnds)
-	    {
-	        // For clamped ends, same approach as chord length but with square root distances
-	        
-	    	// First set starting knots to 0
-	    	PairKnots.Add(FKnot(0.0f, Degree+1));
-	        	        
-	        // Calculate internal knots based on accumulated centripetal Knots
-	        const int32 NumInternalKnots = KnotCount - 2 * (Degree + 1);
-	        if (NumInternalKnots > 0)
-	        {
-	            // Compute normalized Knots at each control point
-	            TArray<float> Params;
-	            Params.SetNum(Points.Num());
-	            Params[0] = 0.0f;
-	            
-	            float AccumulatedLength = 0.0f;
-	            for (int32 i = 1; i < Points.Num(); ++i)
-	            {
-	                AccumulatedLength += Lengths[i-1];
-	                Params[i] = (TotalLength > UE_SMALL_NUMBER) ? 
-	                    (AccumulatedLength / TotalLength) : 
-	                    (static_cast<float>(i) / static_cast<float>(Points.Num() - 1));
-	            }
-	            
-	            // Map these Knots to internal knots
-	            for (int32 i = 0; i < NumInternalKnots; ++i)
-	            {
-	                // Map i to appropriate parameter range
-	                const float t = static_cast<float>(i + 1) / static_cast<float>(NumInternalKnots + 1);
-	                // Map t to control point indices
-	                const int32 LowIndex = FMath::FloorToInt(t * static_cast<float>(Points.Num() - 1));
-	                const int32 HighIndex = FMath::Min(LowIndex + 1, Points.Num() - 1);
-	                const float Alpha = t * static_cast<float>(Points.Num() - 1 - LowIndex);
-	                
-	                // Interpolate Knots
-	                const float KnotValue = FMath::Lerp(Params[LowIndex], Params[HighIndex], Alpha);
-	                PairKnots.Add(FKnot(KnotValue));
-	            }
-	        }
-	    	// Set the last (Degree+1) knots to 1.0
-	    	PairKnots.Add(FKnot(1.0f, Degree+1));
-	    }
-	    else
-	    {
-	        // Unclamped knots
-	        for (int32 i = 0; i < KnotCount; ++i)
-	        {
-	            const float Param = static_cast<float>(i) / static_cast<float>(KnotCount - 1);
-	            PairKnots.Add(FKnot(Param * TotalLength));
-	        }
-	    }
+
+		TArray<float> Lengths;
+		Lengths.Reserve(Points.Num() - 1);
+		float TotalLength = 0.0f;
+
+		for (int32 i = 1; i < Points.Num(); ++i)
+		{
+			// Square root of distance for centripetal parameterization
+			const float Length = static_cast<float>(Math::CentripetalDistance(Points[i], Points[i - 1]));
+			Lengths.Add(Length);
+			TotalLength += Length;
+		}
+
+		// For closed loop, add the closing segment
+		if (this->IsClosedLoop() && Points.Num() > 0)
+		{
+			const float ClosureLength = static_cast<float>(Math::CentripetalDistance(Points[0], Points.Last()));
+			Lengths.Add(ClosureLength);
+			TotalLength += ClosureLength;
+		}
+
+		if (this->IsClosedLoop())
+		{
+			// Generate periodic knot vector using centripetal parameterization
+			for (int32 i = 0; i < NumKnots; ++i)
+			{
+				PairKnots.Add(FKnot(static_cast<float>(i) * TotalLength / static_cast<float>(n + d)));
+			}
+		}
+		else if (bClampEnds)
+		{
+			// For clamped ends, same approach as chord length but with square root distances
+
+			// First set starting knots to 0
+			PairKnots.Add(FKnot(0.0f, Degree + 1));
+
+			// Calculate internal knots based on accumulated centripetal Knots
+			const int32 NumInternalKnots = KnotCount - 2 * (Degree + 1);
+			if (NumInternalKnots > 0)
+			{
+				// Compute normalized Knots at each control point
+				TArray<float> Params;
+				Params.SetNum(Points.Num());
+				Params[0] = 0.0f;
+
+				float AccumulatedLength = 0.0f;
+				for (int32 i = 1; i < Points.Num(); ++i)
+				{
+					AccumulatedLength += Lengths[i - 1];
+					Params[i] = (TotalLength > UE_SMALL_NUMBER)
+						? (AccumulatedLength / TotalLength)
+						: (static_cast<float>(i) / static_cast<float>(Points.Num() - 1));
+				}
+
+				// Map these Knots to internal knots
+				for (int32 i = 0; i < NumInternalKnots; ++i)
+				{
+					// Map i to appropriate parameter range
+					const float t = static_cast<float>(i + 1) / static_cast<float>(NumInternalKnots + 1);
+					// Map t to control point indices
+					const int32 LowIndex = FMath::FloorToInt(t * static_cast<float>(Points.Num() - 1));
+					const int32 HighIndex = FMath::Min(LowIndex + 1, Points.Num() - 1);
+					const float Alpha = t * static_cast<float>(Points.Num() - 1 - LowIndex);
+
+					// Interpolate Knots
+					const float KnotValue = FMath::Lerp(Params[LowIndex], Params[HighIndex], Alpha);
+					PairKnots.Add(FKnot(KnotValue));
+				}
+			}
+			// Set the last (Degree+1) knots to 1.0
+			PairKnots.Add(FKnot(1.0f, Degree + 1));
+		}
+		else
+		{
+			// Unclamped knots
+			for (int32 i = 0; i < KnotCount; ++i)
+			{
+				const float Param = static_cast<float>(i) / static_cast<float>(KnotCount - 1);
+				PairKnots.Add(FKnot(Param * TotalLength));
+			}
+		}
 		MarkFlatKnotsCacheDirty();
 	}
 
@@ -982,21 +1097,21 @@ protected:
 	 * Updates the flat knot cache from the pair representation if the cache is dirty.
 	 */
 	void UpdateFlatKnotsCache() const
-    {
-    	if (bFlatKnotsCacheDirty)
-    	{
-    		FlatKnots = FKnot::ConvertPairToFlatKnots(PairKnots);
-    		bFlatKnotsCacheDirty = false;
-    	}
-    }
+	{
+		if (bFlatKnotsCacheDirty)
+		{
+			FlatKnots = FKnot::ConvertPairToFlatKnots(PairKnots);
+			bFlatKnotsCacheDirty = false;
+		}
+	}
 
 	/**
 	 * Marks the flat knots cache as dirty.
 	 */
 	void MarkFlatKnotsCacheDirty() const
-    {
-    	bFlatKnotsCacheDirty = true;
-    }
+	{
+		bFlatKnotsCacheDirty = true;
+	}
 
 	/**
 	 * Prints the knot vector to the log for debugging.
@@ -1008,14 +1123,14 @@ protected:
 			const FKnot& Knot = PairKnots[Index];
 			UE_LOG(LogSpline, Verbose, TEXT("  [%d] Value: %f, Multiplicity: %u"), Index, Knot.Value, Knot.Multiplicity);
 		}
-		
+
 		for (int32 j = 0; j < this->FlatKnots.Num(); ++j)
 		{
 			UE_LOG(LogSpline, Verbose, TEXT("  FlatKnot[%d] = %f"), j, this->FlatKnots[j]);
 		}
 	}
+
 private:
-	
 	/** Find knot span using binary search */
 	int32 FindKnotSpan(float Parameter) const
 	{
@@ -1024,7 +1139,7 @@ private:
 		// First, check that the knot vector is populated.
 		if (Knots.Num() == 0)
 		{
-			return 0;  // Return a safe fallback.
+			return 0; // Return a safe fallback.
 		}
 
 		const int32 NumValues = NumKeys();
@@ -1073,44 +1188,44 @@ private:
 		return Degree;
 	}
 
-    virtual FWindow FindWindow(float Parameter) const
-    {
-    	FWindow Window = {};
-    	if (NumKeys() < Degree + 1)
-    	{
-    		return Window;
-    	}
-    
-    	int32 Span = FindKnotSpan(Parameter);
-    	// For closed loop, we need to ensure the span wraps properly
-    	if (this->IsClosedLoop())
-    	{
-    		// For closed loop:
-    		// - Span needs to wrap around number of control points
-    		// - Window indices wrap around array
-    		const int32 NumPoints = NumKeys();
-    		for (int32 i = 0; i <= Degree; ++i)
-    		{
-    			int32 Index = (Span - Degree + i) % NumPoints;
-    			if (Index < 0) Index += NumPoints;
-    			Window[i] = &GetValue(Index);
-    		}
-    	}
-    	else
-    	{
-    		// Get control points for the window
-    		for (int32 i = 0; i <= Degree; ++i)
-    		{
-    			const int32 Index = Span - Degree + i;
-    			if (Index >= 0 && Index < NumKeys())
-    			{
-    				Window[i] = &GetValue(Index);
-    			}
-    		}
-    	}
-
-    	return Window;
-    }
+	virtual FWindow FindWindow(float Parameter) const
+	{
+		FWindow Window = {};
+		if (NumKeys() < Degree + 1)
+		{
+			return Window;
+		}
+
+		int32 Span = FindKnotSpan(Parameter);
+		// For closed loop, we need to ensure the span wraps properly
+		if (this->IsClosedLoop())
+		{
+			// For closed loop:
+			// - Span needs to wrap around number of control points
+			// - Window indices wrap around array
+			const int32 NumPoints = NumKeys();
+			for (int32 i = 0; i <= Degree; ++i)
+			{
+				int32 Index = (Span - Degree + i) % NumPoints;
+				if (Index < 0) Index += NumPoints;
+				Window[i] = &GetValue(Index);
+			}
+		}
+		else
+		{
+			// Get control points for the window
+			for (int32 i = 0; i <= Degree; ++i)
+			{
+				const int32 Index = Span - Degree + i;
+				if (Index >= 0 && Index < NumKeys())
+				{
+					Window[i] = &GetValue(Index);
+				}
+			}
+		}
+
+		return Window;
+	}
 
 	/**
 	 * Evaluate the provided window using basis functions computed using the provided parameter. 
@@ -1119,34 +1234,34 @@ private:
 	 * @return The interpolated value.
 	 */
 	virtual ValueType InterpolateWindow(TArrayView<const ValueType* const> Window, float Parameter) const
-    {
-    	TArray<float> Basis;
-    	    
-    	if (PairKnots.IsEmpty())
-    	{
-    		return ValueType();
-    	}
+	{
+		TArray<float> Basis;
+
+		if (PairKnots.IsEmpty())
+		{
+			return ValueType();
+		}
 
 		UpdateFlatKnotsCache();
-		
-    	int32 Span = FindKnotSpan(Parameter);
-				
-    	Math::ComputeBSplineBasisFunctions(FlatKnots, Span, Parameter, Degree, Basis);
-    	return TSplineInterpolationPolicy<ValueType>::InterpolateWithBasis(Window, Basis);
-    }
-	
-protected:	
-    TArray<ValueType> Values;
+
+		int32 Span = FindKnotSpan(Parameter);
+
+		Math::ComputeBSplineBasisFunctions(FlatKnots, Span, Parameter, Degree, Basis);
+		return TSplineInterpolationPolicy<ValueType>::InterpolateWithBasis(Window, Basis);
+	}
+
+protected:
+	TArray<ValueType> Values;
 
 	// Knot representation in pair form having multiplicity
 	TArray<FKnot> PairKnots;
-    
+
 	// Cached flattened knot vector for compatible algorithms
 	mutable TArray<float> FlatKnots;
 	mutable bool bFlatKnotsCacheDirty;
-	
+
 	bool bIsClosedLoop = false;
-	
+
 	/** Whether to clamp the endpoints for interpolation of first/last control points */
 	bool bClampEnds = true;
 };
@@ -1163,4 +1278,4 @@ using FQuarticBSpline3f = TBSpline<FVector3f, 4>;
 using FQuinticBSpline2d = TBSpline<FVector2d, 5>;
 } // end namespace UE::Geometry::Spline
 } // end namespace UE::Geometry
-} // end namespace UE	
\ No newline at end of file
+} // end namespace UE	
diff --git a/Engine/Source/Runtime/GeometryCore/Public/Splines/PolyBezierSpline.h b/Engine/Source/Runtime/GeometryCore/Public/Splines/PolyBezierSpline.h
--- a/Engine/Source/Runtime/GeometryCore/Public/Splines/PolyBezierSpline.h
+++ b/Engine/Source/Runtime/GeometryCore/Public/Splines/PolyBezierSpline.h
@@ -818,12 +818,12 @@ public:
 
 		if (TargetSplitParam == SegmentRange.Min)
 		{
-			TargetSplitParam = std::nextafter(TargetSplitParam, std::numeric_limits<float>::infinity());
+			TargetSplitParam = Param::NextDistinct(TargetSplitParam);
 		}
 
 		if (TargetSplitParam == SegmentRange.Max)
 		{
-			TargetSplitParam = std::nextafter(TargetSplitParam, -std::numeric_limits<float>::infinity());
+			TargetSplitParam = Param::PrevDistinct(TargetSplitParam);
 		}
 		
 		for (int32 i = 0; i < 4; ++i)
@@ -1547,12 +1547,12 @@ private:
 			if (Offset == 0.f) return Val;
 				
 			const float Result = Val + Offset;
-				
-			return Result == Val
-				? std::nextafter(Val, Offset > 0.f
-					? std::numeric_limits<float>::infinity()
-					: -std::numeric_limits<float>::infinity())
-				: Result;
+			if (Result != Val) return Result;
+
+			// Guaranteed one-step progress even under FTZ/DAZ:
+			return UE::Geometry::Spline::Param::Step(
+				Val, Offset > 0.f ? UE::Geometry::Spline::Param::EDir::Right
+								  : UE::Geometry::Spline::Param::EDir::Left);
 		};
 		
 		// Reset pair knots
@@ -1905,13 +1905,17 @@ private:
 		else if (bAppendToEnd)
 		{
 			const float LastParameter = Base::PairKnots.Last().Value;
-			const float NextParameter = FMath::Max(LastParameter + SegmentLength, std::nextafter(LastParameter, std::numeric_limits<float>::infinity()));
-			this->InsertKnot( FKnot(NextParameter, this->bClampEnds ? this->GetDegree() : 1));
+			const float CandidateByLength = LastParameter + SegmentLength;
+			const float CandidateByStep   = UE::Geometry::Spline::Param::NextDistinct(LastParameter);
+			const float NextParameter = FMath::Max(CandidateByLength, CandidateByStep);
+			this->InsertKnot(FKnot(NextParameter, this->bClampEnds ? this->GetDegree() : 1));
 		}
 		else if (bPrependToStart)
 		{
 			const float FirstParameter = Base::PairKnots[0].Value;
-			const float NextParameter = FMath::Min(FirstParameter - SegmentLength, std::nextafter(FirstParameter, -std::numeric_limits<float>::infinity()));
+			const float CandidateByLength = FirstParameter - SegmentLength;
+			const float CandidateByStep   = UE::Geometry::Spline::Param::PrevDistinct(FirstParameter);
+			const float NextParameter = FMath::Min(CandidateByLength, CandidateByStep);
 			this->InsertKnot( FKnot(NextParameter, this->bClampEnds ? this->GetDegree() : 1));
 		}
 	
diff --git a/Engine/Source/Runtime/GeometryCore/Public/Splines/SplineMath.h b/Engine/Source/Runtime/GeometryCore/Public/Splines/SplineMath.h
--- a/Engine/Source/Runtime/GeometryCore/Public/Splines/SplineMath.h
+++ b/Engine/Source/Runtime/GeometryCore/Public/Splines/SplineMath.h
@@ -5,11 +5,8 @@
 #include "CoreMinimal.h"
 #include "SplineInterfaces.h"
 #include "ParameterizedTypes.h"
-
-namespace UE::Geometry::Spline
-{
-	struct FKnot;
-}
+#include <cstdint>
+#include <limits>
 
 namespace UE
 {
@@ -17,697 +14,828 @@ namespace Geometry
 {
 namespace Spline
 {
-namespace Math
+struct FKnot;
+
+// -------- Internal: numeric primitives (not public API) -----------------
+namespace Internal
 {
+	inline uint32 FloatToBits(float f)
+	{
+		uint32 u;
+		FMemory::Memcpy(&u, &f, sizeof(float));
+		return u;
+	}
 
+	inline float BitsToFloat(uint32 u)
+	{
+		float f;
+		FMemory::Memcpy(&f, &u, sizeof(float));
+		return f;
+	}
 
-/**
- * Size/magnitude of a value
- */
-template<typename T>
-double Size(const T& Value)
-{
-    if constexpr (THasSizeMethod<T>::value)
-        return static_cast<double>(Value.Size());
-    else if constexpr (std::is_arithmetic_v<T>)
-        return FMath::Abs(static_cast<double>(Value));
-    else
-        return 1.0; // Default for other types
-}
-
-/**
- * Squared size (optimization to avoid sqrt)
- */
-template<typename T>
-double SizeSquared(const T& Value)
-{
-    if constexpr (THasSizeSquaredMethod<T>::value)
-        return static_cast<double>(Value.SizeSquared());
-    else if constexpr (std::is_arithmetic_v<T>)
-        return static_cast<double>(Value * Value);
-    else
-        return Size(Value) * Size(Value);
-}
-
-/**
- * Distance between values
- */
-template<typename T>
-double Distance(const T& A, const T& B)
-{
-    if constexpr (THasSubtractionOperator<T>::value)
-        return Size(B - A);
-    else
-        return 1.0; // Default fallback
-}
-	
-/**
- * Centripetal distance calculation (for parameterization)
- */
-template<typename T>
-double CentripetalDistance(const T& A, const T& B)
-{
-	return FMath::Sqrt(Distance(A, B));
-}
-	
-/**
- * Dot product between values
- */
-template<typename T>
-double Dot(const T& A, const T& B)
-{
-    if constexpr (THasDotMethod<T>::value)
-        return static_cast<double>(A.Dot(B));
-    else if constexpr (std::is_arithmetic_v<T>)
-        return static_cast<double>(A * B);
-    else
-        return 0.0; // Fallback
-}
-
-/**
- * Get normalized version of a value
- */
-template<typename T>
-T GetSafeNormal(const T& Value)
-{
-    if constexpr (THasGetSafeNormalMethod<T>::value)
-        return Value.GetSafeNormal();
-    else if constexpr (std::is_arithmetic_v<T>)
-    {
-        const double AbsValue = FMath::Abs(static_cast<double>(Value));
-        return AbsValue < UE_KINDA_SMALL_NUMBER ? T(0) : (Value > 0 ? T(1) : T(-1));
-    }
-    else
-        return Value; // Fallback
-}
-
-/**
- * Check if two values are equal within tolerance
- */
-template<typename T>
-bool Equals(const T& A, const T& B, float Tolerance = UE_KINDA_SMALL_NUMBER)
-{
-    if constexpr (THasEqualsMethod<T>::value)
-        return A.Equals(B, Tolerance);
-    else if constexpr (std::is_arithmetic_v<T>)
-        return FMath::Abs(static_cast<float>(A - B)) <= Tolerance;
-    else if constexpr (THasSubtractionOperator<T>::value)
-        return Distance(A, B) <= Tolerance;
-    else
-        return false; // Fallback
-}
-
-	
-/** Helper for validating types and inputs */
-struct FSplineValidation
-{
-	/** Get the default/zero value for a type */
-	template<typename T>
-	static T GetDefaultValue()
+	// Map IEEE754 float to a monotonic integer ordering and back.
+	inline uint32 ToOrdered(float f)
 	{
-		if constexpr (TIsConstructible<T, EForceInit>::Value)
+		const uint32 u = FloatToBits(f);
+		const bool neg = (u & 0x80000000u) != 0;
+		return neg
+			? ~u
+			: (u ^ 0x80000000u);
+	}
+
+	inline float FromOrdered(uint32 ordered)
+	{
+		const uint32 u = (ordered & 0x80000000u)
+			? (ordered ^ 0x80000000u)
+			: ~ordered;
+		return BitsToFloat(u);
+	}
+
+	/** Adjacent representable *normal-or-better* float (FTZ-proof).
+	 *  Skips subnormals (snaps to ±FLT_MIN) and saturates at ±FLT_MAX.
+	 */
+	inline float NextafterNoSubnormal(float from, bool bStepRight)
+	{
+		if (FMath::IsNaN(from) || !FMath::IsFinite(from))
 		{
-			return T(ForceInit);
+			return from; // keep NaN/Inf unchanged
+		}
+		if (from == 0.0f)
+		{
+			const float minNorm = std::numeric_limits<float>::min();
+			return bStepRight
+				? +minNorm
+				: -minNorm;
+		}
+
+		uint32 ordered = ToOrdered(from);
+		if (bStepRight)
+		{
+			++ordered;
 		}
 		else
 		{
-			return T(); 
-		} 
-	}
+			--ordered;
+		}
+		float y = FromOrdered(ordered);
 
-	template<typename T>
-	static bool IsValidWindow(TArrayView<const T* const> Window, int32 RequiredSize)
-	{
-		if (Window.Num() < RequiredSize)
+		const float minNorm = std::numeric_limits<float>::min();
+		const float ay = FMath::Abs(y);
+		if (ay > 0.0f && ay < minNorm)
 		{
-			return false;
+			y = bStepRight
+				? +minNorm
+				: -minNorm;
 		}
-    
-		// Check for null pointers
-		for (const T* Ptr : Window)
+		if (!FMath::IsFinite(y))
 		{
-			if (!Ptr)
-			{
-				return false;
-			}
+			y = bStepRight
+				? +std::numeric_limits<float>::max()
+				: -std::numeric_limits<float>::max();
 		}
-    
-		return true;
+		return y;
 	}
 
-	static bool IsValidParameter(float Parameter)
+	/** Normalize -0.0f to +0.0f for use as a key (avoids phantom duplicates). */
+	inline float NormalizeZeroForKey(float v)
 	{
-		return !FPlatformMath::IsNaN(Parameter) && 
-			   FPlatformMath::IsFinite(Parameter);
+		return (v == 0.0f)
+			? 0.0f
+			: v;
 	}
+} // namespace Internal
+
+// -------- Public: parameter helpers (domain-centric names) ---------------
+namespace Param
+{
+	enum class EDir : uint8 { Left, Right };
 
-	template<typename T>
-	static bool IsValidGeometry(TArrayView<const T* const> Window, float Tolerance = UE_KINDA_SMALL_NUMBER)
+	/** Step to a distinct adjacent parameter (float), FTZ-proof. */
+	inline float NextDistinct(float t)
 	{
-		if (Window.Num() < 2) return false;
+		return Internal::NextafterNoSubnormal(t, /*right=*/true);
+	}
 
-		// Check for degenerate/repeated points
-		for (int32 i = 1; i < Window.Num(); ++i)
-		{
-			if (DistanceSquared(*Window[i], *Window[i-1]) < Tolerance)
-			{
-				return false;
-			}
-		}
+	inline float PrevDistinct(float t)
+	{
+		return Internal::NextafterNoSubnormal(t, /*right=*/false);
+	}
 
-		return true;
+	/** Step one distinct param in given direction. */
+	inline float Step(float t, EDir d)
+	{
+		return Internal::NextafterNoSubnormal(t, d == EDir::Right);
 	}
-};
-	
-// Cached factorial table access
-template <int32 N>
-const TArray<double>& GetFactorialTable()
-{
-	static TArray<double> Table;
-	if (Table.Num() == 0)
+
+	/** Step and clamp to an interval (when you want to stay inside bounds). */
+	inline float StepInside(float t, EDir d, const FInterval1f& bounds)
 	{
-		Table.SetNum(N + 1);
-		Table[0] = 1.0;
-		for (int32 i = 1; i <= N; ++i)
-		{
-			Table[i] = Table[i - 1] * i;
-		}
+		return FMath::Clamp(Internal::NextafterNoSubnormal(t, d == EDir::Right), bounds.Min, bounds.Max);
 	}
-	return Table;
-}
 
-// Cached binomial coefficient table
-template<int32 N>
-const TArray<TArray<double>>& GetBinomialTable()
-{
-	static TArray<TArray<double>> Table;
-	if (Table.Num() == 0)
+	/** Key sanitizer for sets/maps that use exact float equality. */
+	inline float NormalizeKey(float t)
 	{
-		Table.SetNum(N + 1);
-		for (int32 i = 0; i <= N; ++i)
-		{
-			Table[i].SetNum(i + 1);
-			Table[i][0] = Table[i][i] = 1.0;
-			for (int32 j = 1; j < i; ++j)
-			{
-				Table[i][j] = Table[i-1][j-1] + Table[i-1][j];
-			}
-		}
+		return Internal::NormalizeZeroForKey(t);
 	}
-	return Table;
-}
+} // namespace Param
 
-// Bezier basis computation
-template <int32 Order>
-static void ComputeBezierBasis(float t, TArray<float>& OutBasis)
+namespace Math
 {
-	OutBasis.SetNumZeroed(4); // Cubic Bezier
-
-	const float mt = 1.0f - t;
-
-	if constexpr (Order == 0)
+	/**
+	 * Size/magnitude of a value
+	 */
+	template <typename T>
+	double Size(const T& Value)
 	{
-		// Position basis
-		OutBasis[0] = mt * mt * mt;
-		OutBasis[1] = 3.0f * mt * mt * t;
-		OutBasis[2] = 3.0f * mt * t * t;
-		OutBasis[3] = t * t * t;
+		if constexpr (THasSizeMethod<T>::value)
+			return static_cast<double>(Value.Size());
+		else if constexpr (std::is_arithmetic_v<T>)
+			return FMath::Abs(static_cast<double>(Value));
+		else
+			return 1.0; // Default for other types
 	}
-	else if constexpr (Order == 1)
+
+	/**
+	 * Squared size (optimization to avoid sqrt)
+	 */
+	template <typename T>
+	double SizeSquared(const T& Value)
 	{
-		// First derivative basis
-		OutBasis[0] = -3.0f * mt * mt;
-		OutBasis[1] = 3.0f * mt * (1.0f - 3.0f * t);
-		OutBasis[2] = 3.0f * t * (2.0f - 3.0f * t);
-		OutBasis[3] = 3.0f * t * t;
+		if constexpr (THasSizeSquaredMethod<T>::value)
+			return static_cast<double>(Value.SizeSquared());
+		else if constexpr (std::is_arithmetic_v<T>)
+			return static_cast<double>(Value * Value);
+		else
+			return Size(Value) * Size(Value);
 	}
-	else if constexpr (Order == 2)
+
+	/**
+	 * Distance between values
+	 */
+	template <typename T>
+	double Distance(const T& A, const T& B)
 	{
-		// Second derivative basis
-		OutBasis[0] = 6.0f * mt;
-		OutBasis[1] = -12.0f + 18.0f * t;
-		OutBasis[2] = 6.0f - 18.0f * t;
-		OutBasis[3] = 6.0f * t;
+		if constexpr (THasSubtractionOperator<T>::value)
+			return Size(B - A);
+		else
+			return 1.0; // Default fallback
 	}
-	else if constexpr (Order == 3)
+
+	/**
+	 * Centripetal distance calculation (for parameterization)
+	 */
+	template <typename T>
+	double CentripetalDistance(const T& A, const T& B)
 	{
-		// Third derivative basis (constant)
-		OutBasis[0] = -6.0f;
-		OutBasis[1] = 18.0f;
-		OutBasis[2] = -18.0f;
-		OutBasis[3] = 6.0f;
+		return FMath::Sqrt(Distance(A, B));
 	}
-}
-	
-// compute basis functions for a given degree
-static void ComputeBSplineBasisFunctions(
-	const TArray<float>& KnotVector,
-	int32 Span,
-	float Parameter, 
-	int32 Degree,
-	TArray<float>& OutBasis)
-{
-	OutBasis.SetNum(Degree + 1);
-	TArray<float> Left;
-	TArray<float> Right;
-	Left.SetNum(Degree + 1);
-	Right.SetNum(Degree + 1);
-
-	OutBasis[0] = 1.0f;
 
-	for (int32 j = 1; j <= Degree; ++j)
+	/**
+	 * Dot product between values
+	 */
+	template <typename T>
+	double Dot(const T& A, const T& B)
 	{
-		const int32 LeftIndex = Span + 1 - j;
-		const int32 RightIndex = Span + j;
+		if constexpr (THasDotMethod<T>::value)
+			return static_cast<double>(A.Dot(B));
+		else if constexpr (std::is_arithmetic_v<T>)
+			return static_cast<double>(A * B);
+		else
+			return 0.0; // Fallback
+	}
 
-		if (LeftIndex < 0 || LeftIndex >= KnotVector.Num() ||
-			RightIndex < 0 || RightIndex >= KnotVector.Num())
+	/**
+	 * Get normalized version of a value
+	 */
+	template <typename T>
+	T GetSafeNormal(const T& Value)
+	{
+		if constexpr (THasGetSafeNormalMethod<T>::value)
+			return Value.GetSafeNormal();
+		else if constexpr (std::is_arithmetic_v<T>)
 		{
-			UE_LOG(LogTemp, Warning, TEXT("Invalid knot indices in ComputeBasisFunctions: LeftIndex=%d, RightIndex=%d"),
-				   LeftIndex, RightIndex);
-			// Instead of continuing, we should zero out the basis and return.
-			OutBasis.SetNumZeroed(Degree + 1);
-			return;
+			const double AbsValue = FMath::Abs(static_cast<double>(Value));
+			return AbsValue < UE_KINDA_SMALL_NUMBER
+				? T(0)
+				: (Value > 0
+					? T(1)
+					: T(-1));
 		}
+		else
+			return Value; // Fallback
+	}
 
-		Left[j] = Parameter - KnotVector[LeftIndex];
-		Right[j] = KnotVector[RightIndex] - Parameter;
+	/**
+	 * Check if two values are equal within tolerance
+	 */
+	template <typename T>
+	bool Equals(const T& A, const T& B, float Tolerance = UE_KINDA_SMALL_NUMBER)
+	{
+		if constexpr (THasEqualsMethod<T>::value)
+			return A.Equals(B, Tolerance);
+		else if constexpr (std::is_arithmetic_v<T>)
+			return FMath::Abs(static_cast<float>(A - B)) <= Tolerance;
+		else if constexpr (THasSubtractionOperator<T>::value)
+			return Distance(A, B) <= Tolerance;
+		else
+			return false; // Fallback
+	}
 
-		float Saved = 0.0f;
 
-		for (int32 r = 0; r < j; ++r)
+	/** Helper for validating types and inputs */
+	struct FSplineValidation
+	{
+		/** Get the default/zero value for a type */
+		template <typename T>
+		static T GetDefaultValue()
 		{
-			float Temp;
-			const float Denominator = Right[r + 1] + Left[j - r];
-            // Correctly handle zero denominator
-            if (FMath::IsNearlyZero(Denominator))
+			if constexpr (TIsConstructible<T, EForceInit>::Value)
 			{
-				Temp = 0.0f;
+				return T(ForceInit);
 			}
 			else
 			{
-				Temp = OutBasis[r] / Denominator;
+				return T();
 			}
-
-			OutBasis[r] = Saved + Right[r + 1] * Temp;
-			Saved = Left[j - r] * Temp;
 		}
 
-		OutBasis[j] = Saved;
-	}
-}
-/** Template-based derivative calculators for spline types */
-
-// Helper struct for shared derivative implementation
-template<typename T, int32 Order>
-struct TGenericDerivativeHelper
-{
-	static T Compute(TArrayView<const T* const> Window, float Parameter)
-	{
-		// Basic validation
-		if (!FSplineValidation::IsValidWindow(Window, Order + 1))
+		template <typename T>
+		static bool IsValidWindow(TArrayView<const T* const> Window, int32 RequiredSize)
 		{
-			return FSplineValidation::GetDefaultValue<T>();
-		}
+			if (Window.Num() < RequiredSize)
+			{
+				return false;
+			}
 
-		if (!FSplineValidation::IsValidParameter(Parameter))
-		{
-#if ENABLE_NAN_DIAGNOSTIC
-			logOrEnsureNanError(TEXT("Invalid parameter in derivative calculation"));
-#endif
-			return FSplineValidation::GetDefaultValue<T>();
-		}
-		
-		// Get pre-computed binomial coefficients
-		static const TArray<TArray<double>>& BinomialTable = GetBinomialTable<20>();
-		static const TArray<double>& FactorialTable = GetFactorialTable<20>();
+			// Check for null pointers
+			for (const T* Ptr : Window)
+			{
+				if (!Ptr)
+				{
+					return false;
+				}
+			}
 
-		const int32 WindowSize = Window.Num();
-		if (Order >= WindowSize || Order >= BinomialTable.Num())
-		{
-			return FSplineValidation::GetDefaultValue<T>();
+			return true;
 		}
 
-		// Optional geometry validation for higher derivatives
-		if (Order > 1 && !FSplineValidation::IsValidGeometry(Window))
+		static bool IsValidParameter(float Parameter)
 		{
-			return FSplineValidation::GetDefaultValue<T>();
+			return !FPlatformMath::IsNaN(Parameter) &&
+				FPlatformMath::IsFinite(Parameter);
 		}
 
-		T Result = FSplineValidation::GetDefaultValue<T>();
-		const float t = Parameter;
-		const float mt = 1.0f - t;
-
-		// Compute derivative using binomial expansion
-		for (int32 i = 0; i <= WindowSize - Order - 1; ++i)
+		template <typename T>
+		static bool IsValidGeometry(TArrayView<const T* const> Window, float Tolerance = UE_KINDA_SMALL_NUMBER)
 		{
-			double Coeff = BinomialTable[WindowSize - 1][i];
+			if (Window.Num() < 2) return false;
 
-			// Protect against overflow
-			if (Coeff > DBL_MAX / (WindowSize - 1))
+			// Check for degenerate/repeated points
+			for (int32 i = 1; i < Window.Num(); ++i)
 			{
-				return FSplineValidation::GetDefaultValue<T>();
-			}
-			
-			// Apply derivative scaling
-			Coeff *= FactorialTable[WindowSize - 1] / FactorialTable[WindowSize - Order - 1];
-        
-			// Apply parameter terms with underflow protection
-			if (mt > 0.0f) 
-			{
-				const float PowMT = FMath::Pow(mt, WindowSize - Order - 1 - i);
-				if (FPlatformMath::IsFinite(PowMT))
+				if (DistanceSquared(*Window[i], *Window[i - 1]) < Tolerance)
 				{
-					Coeff *= PowMT;
+					return false;
 				}
 			}
-			if (t > 0.0f)
+
+			return true;
+		}
+	};
+
+	// Cached factorial table access
+	template <int32 N>
+	const TArray<double>& GetFactorialTable()
+	{
+		static TArray<double> Table;
+		if (Table.Num() == 0)
+		{
+			Table.SetNum(N + 1);
+			Table[0] = 1.0;
+			for (int32 i = 1; i <= N; ++i)
+			{
+				Table[i] = Table[i - 1] * i;
+			}
+		}
+		return Table;
+	}
+
+	// Cached binomial coefficient table
+	template <int32 N>
+	const TArray<TArray<double>>& GetBinomialTable()
+	{
+		static TArray<TArray<double>> Table;
+		if (Table.Num() == 0)
+		{
+			Table.SetNum(N + 1);
+			for (int32 i = 0; i <= N; ++i)
 			{
-				const float PowT = FMath::Pow(t, i);
-				if (FPlatformMath::IsFinite(PowT))
+				Table[i].SetNum(i + 1);
+				Table[i][0] = Table[i][i] = 1.0;
+				for (int32 j = 1; j < i; ++j)
 				{
-					Coeff *= PowT;
+					Table[i][j] = Table[i - 1][j - 1] + Table[i - 1][j];
 				}
 			}
-
-			Result += *Window[i] * Coeff;
 		}
-		return Result;
+		return Table;
 	}
-};
 
-// Computes the nth derivative (where n == DerivOrder) of a BSpline curve.
-// Note: This implementation assumes that DerivOrder <= Degree.
-template <typename T, int32 DerivOrder>
-struct TBSplineDerivativeCalculator
-{
-	/**
-	 * Computes the DerivOrder-th derivative of a BSpline curve at a given parameter.
-	 *
-	 * @param Window     A view of pointers to control points for the current BSpline segment.
-	 *                   For a BSpline of degree p, this view should contain p+1 control points.
-	 * @param Knots      The global knot vector.
-	 * @param Span       The knot span index such that Parameter ∈ [Knots[Span], Knots[Span+1]).
-	 * @param Parameter  The parameter at which to evaluate.
-	 * @param Degree     The degree of the current BSpline (should be ≥ DerivOrder).
-	 *
-	 * @return The DerivOrder-th derivative (of type T) evaluated at Parameter.
-	 */
-	static T Compute(
-		TArrayView<const T* const> Window,
-		const TArray<float>& Knots,
-		int32 Span,
-		float Parameter,
-		int32 Degree)
+	// Bezier basis computation
+	template <int32 Order>
+	static void ComputeBezierBasis(float t, TArray<float>& OutBasis)
 	{
-		// Basic validation
-        if (Window.Num() < Degree + 1 || DerivOrder > Degree)
-        {
-            return T();
-        }
-
-		TArray<T> DerivControlPoints;
-        DerivControlPoints.SetNum(Degree);
-        const int32 KnotStart = Span - Degree;
-
-        // Determine if knots are uniformly spaced
-        bool bIsUniform = true;
-        const float h = Knots[1] - Knots[0];
-        for (int32 i = 2; i < Knots.Num(); ++i) {
-            if (!FMath::IsNearlyEqual(Knots[i] - Knots[i-1], h, UE_KINDA_SMALL_NUMBER)) {
-                bIsUniform = false;
-                break;
-            }
-        }
-
-        // Compute derivative control points
-        for (int32 i = 0; i < Degree; ++i)
-        {
-            const int32 KnotDenom = KnotStart + i + Degree - DerivOrder + 2; //Knot index for the denominator
-            const int32 KnotNum = KnotStart + i + 1; //Knot index in the numerator
-
-            if (KnotDenom >= Knots.Num() || KnotNum >= Knots.Num() ||
-                KnotDenom < 0 || KnotNum < 0)
-            {
-                return T();
-            }
-
-            const float Denom = Knots[KnotDenom] - Knots[KnotNum];
-
-            // Handle uniform knots differently
-            if (bIsUniform)
-            {
-                // For uniform knots, use the fixed interval h instead of Denom
-                DerivControlPoints[i] = ((*Window[i + 1] - *Window[i])) * (static_cast<float>(Degree) / h); //Correct scaling for uniform knots
-            }
-            else
-            {
-                DerivControlPoints[i] = ((*Window[i + 1] - *Window[i])) * (static_cast<float>(Degree) / Denom); //Correct scaling for non-uniform knots
-            }
-        }
-
-        // Adjust scale factor for uniform knots
-        //Removing the scale factor
-        float ScaleFactor = 1.0;
-
-        // Recurse for higher derivatives with adjusted knots and parameters
-        if constexpr (DerivOrder > 1)
-        {
-           //No change here
-            TArray<const T*> DerivWindow;
-            DerivWindow.SetNum(DerivControlPoints.Num());
-            for (int32 i = 0; i < DerivControlPoints.Num(); ++i)
-            {
-                DerivWindow[i] = &DerivControlPoints[i];
-            }
-
-            // Trim the knot vector for uniform case
-            TArray<float> TrimmedKnots;
-            if (bIsUniform)
-            {
-                TrimmedKnots.Reserve(Knots.Num() - 2);
-                for (int32 i = 1; i < Knots.Num() - 1; ++i)
-                {
-                    TrimmedKnots.Add(Knots[i]);
-                }
-            }
-            else
-            {
-                TrimmedKnots = Knots;
-            }
-
-            return TBSplineDerivativeCalculator<T, DerivOrder - 1>::Compute(
-                DerivWindow, //Reduced degree control points
-                TrimmedKnots, //Remove the first and the last knot as the degree is now p-1
-                Span, // Keep Span the same for uniform knots
-                Parameter, //Same parameter
-                Degree - 1 //Reduce the degree
-            ) * ScaleFactor; //Multiply with scale factor
-        }
-        else
-        {
-            // Base case: First derivative
-            TArray<float> Basis;
-            Math::ComputeBSplineBasisFunctions(Knots, Span, Parameter, Degree - 1, Basis);
-
-            T Result = T();
-            for (int32 i = 0; i < Degree; ++i)
-            {
-                Result += DerivControlPoints[i] * Basis[i];
-            }
-            return Result * ScaleFactor;
-        }
-    }
-};
-
-// Base case: 0th derivative is just the original BSpline evaluation.
-template <typename T>
-struct TBSplineDerivativeCalculator<T, 0>
-{
-	/**
-	 * Evaluates the BSpline curve at Parameter using de Boor's algorithm.
-	 *
-	 * @param Window     A view of pointers to the BSpline's control points (should be Degree+1 points).
-	 * @param Knots      The (possibly trimmed) knot vector.
-	 * @param Span       The knot span index.
-	 * @param Parameter  The parameter at which to evaluate.
-	 * @param Degree     The degree of the BSpline.
-	 *
-	 * @return The BSpline value (of type T) at Parameter.
-	 */
-	static T Compute(
-		TArrayView<const T* const> Window,
-		const TArray<float>& Knots,
+		OutBasis.SetNumZeroed(4); // Cubic Bezier
+
+		const float mt = 1.0f - t;
+
+		if constexpr (Order == 0)
+		{
+			// Position basis
+			OutBasis[0] = mt * mt * mt;
+			OutBasis[1] = 3.0f * mt * mt * t;
+			OutBasis[2] = 3.0f * mt * t * t;
+			OutBasis[3] = t * t * t;
+		}
+		else if constexpr (Order == 1)
+		{
+			// First derivative basis
+			OutBasis[0] = -3.0f * mt * mt;
+			OutBasis[1] = 3.0f * mt * (1.0f - 3.0f * t);
+			OutBasis[2] = 3.0f * t * (2.0f - 3.0f * t);
+			OutBasis[3] = 3.0f * t * t;
+		}
+		else if constexpr (Order == 2)
+		{
+			// Second derivative basis
+			OutBasis[0] = 6.0f * mt;
+			OutBasis[1] = -12.0f + 18.0f * t;
+			OutBasis[2] = 6.0f - 18.0f * t;
+			OutBasis[3] = 6.0f * t;
+		}
+		else if constexpr (Order == 3)
+		{
+			// Third derivative basis (constant)
+			OutBasis[0] = -6.0f;
+			OutBasis[1] = 18.0f;
+			OutBasis[2] = -18.0f;
+			OutBasis[3] = 6.0f;
+		}
+	}
+
+	// compute basis functions for a given degree
+	static void ComputeBSplineBasisFunctions(
+		const TArray<float>& KnotVector,
 		int32 Span,
 		float Parameter,
-		int32 Degree)
+		int32 Degree,
+		TArray<float>& OutBasis)
 	{
-		// Regular evaluation using current basis functions
-		TArray<float> Basis;
-		ComputeBSplineBasisFunctions(Knots, Span, Parameter, Degree, Basis);
+		OutBasis.SetNum(Degree + 1);
+		TArray<float> Left;
+		TArray<float> Right;
+		Left.SetNum(Degree + 1);
+		Right.SetNum(Degree + 1);
+
+		OutBasis[0] = 1.0f;
 
-		T Result = T();
-		for (int32 i = 0; i <= Degree; ++i)
+		for (int32 j = 1; j <= Degree; ++j)
 		{
-			Result += *Window[i] * Basis[i];
+			const int32 LeftIndex = Span + 1 - j;
+			const int32 RightIndex = Span + j;
+
+			if (LeftIndex < 0 || LeftIndex >= KnotVector.Num() ||
+				RightIndex < 0 || RightIndex >= KnotVector.Num())
+			{
+				UE_LOG(LogTemp, Warning, TEXT("Invalid knot indices in ComputeBasisFunctions: LeftIndex=%d, RightIndex=%d"),
+					LeftIndex, RightIndex);
+				// Instead of continuing, we should zero out the basis and return.
+				OutBasis.SetNumZeroed(Degree + 1);
+				return;
+			}
+
+			Left[j] = Parameter - KnotVector[LeftIndex];
+			Right[j] = KnotVector[RightIndex] - Parameter;
+
+			float Saved = 0.0f;
+
+			for (int32 r = 0; r < j; ++r)
+			{
+				float Temp;
+				const float Denominator = Right[r + 1] + Left[j - r];
+				// Correctly handle zero denominator
+				if (FMath::IsNearlyZero(Denominator))
+				{
+					Temp = 0.0f;
+				}
+				else
+				{
+					Temp = OutBasis[r] / Denominator;
+				}
+
+				OutBasis[r] = Saved + Right[r + 1] * Temp;
+				Saved = Left[j - r] * Temp;
+			}
+
+			OutBasis[j] = Saved;
 		}
-		return Result;
 	}
-};
 
+	/** Template-based derivative calculators for spline types */
 
-// Bezier curve specific calculations
-template <typename T, int32 Order>
-struct TBezierDerivativeCalculator
-{
-	static T Compute(TArrayView<const T* const> Window, float Parameter, float SegmentScale)
+	// Helper struct for shared derivative implementation
+	template <typename T, int32 Order>
+	struct TGenericDerivativeHelper
 	{
-		if constexpr (Order > 3) 
+		static T Compute(TArrayView<const T* const> Window, float Parameter)
 		{
-			return FSplineValidation::GetDefaultValue<T>();
-		}
+			// Basic validation
+			if (!FSplineValidation::IsValidWindow(Window, Order + 1))
+			{
+				return FSplineValidation::GetDefaultValue<T>();
+			}
 
-		if (!FSplineValidation::IsValidWindow(Window, 4) ||  // Cubic Bezier needs 4 points
-			!FSplineValidation::IsValidParameter(Parameter))
-		{
-			return FSplineValidation::GetDefaultValue<T>();
+			if (!FSplineValidation::IsValidParameter(Parameter))
+			{
+#if ENABLE_NAN_DIAGNOSTIC
+				logOrEnsureNanError(TEXT("Invalid parameter in derivative calculation"));
+#endif
+				return FSplineValidation::GetDefaultValue<T>();
+			}
+
+			// Get pre-computed binomial coefficients
+			static const TArray<TArray<double>>& BinomialTable = GetBinomialTable<20>();
+			static const TArray<double>& FactorialTable = GetFactorialTable<20>();
+
+			const int32 WindowSize = Window.Num();
+			if (Order >= WindowSize || Order >= BinomialTable.Num())
+			{
+				return FSplineValidation::GetDefaultValue<T>();
+			}
+
+			// Optional geometry validation for higher derivatives
+			if (Order > 1 && !FSplineValidation::IsValidGeometry(Window))
+			{
+				return FSplineValidation::GetDefaultValue<T>();
+			}
+
+			T Result = FSplineValidation::GetDefaultValue<T>();
+			const float t = Parameter;
+			const float mt = 1.0f - t;
+
+			// Compute derivative using binomial expansion
+			for (int32 i = 0; i <= WindowSize - Order - 1; ++i)
+			{
+				double Coeff = BinomialTable[WindowSize - 1][i];
+
+				// Protect against overflow
+				if (Coeff > DBL_MAX / (WindowSize - 1))
+				{
+					return FSplineValidation::GetDefaultValue<T>();
+				}
+
+				// Apply derivative scaling
+				Coeff *= FactorialTable[WindowSize - 1] / FactorialTable[WindowSize - Order - 1];
+
+				// Apply parameter terms with underflow protection
+				if (mt > 0.0f)
+				{
+					const float PowMT = FMath::Pow(mt, WindowSize - Order - 1 - i);
+					if (FPlatformMath::IsFinite(PowMT))
+					{
+						Coeff *= PowMT;
+					}
+				}
+				if (t > 0.0f)
+				{
+					const float PowT = FMath::Pow(t, i);
+					if (FPlatformMath::IsFinite(PowT))
+					{
+						Coeff *= PowT;
+					}
+				}
+
+				Result += *Window[i] * Coeff;
+			}
+			return Result;
 		}
+	};
 
-		// Use existing basis computation
-		TArray<float> Basis;
-		ComputeBezierBasis<Order>(Parameter, Basis);
-    
-		// Apply basis to control points
-		T Result = FSplineValidation::GetDefaultValue<T>();
-		for (int32 i = 0; i < Window.Num(); ++i)
+	// Computes the nth derivative (where n == DerivOrder) of a BSpline curve.
+	// Note: This implementation assumes that DerivOrder <= Degree.
+	template <typename T, int32 DerivOrder>
+	struct TBSplineDerivativeCalculator
+	{
+		/**
+		 * Computes the DerivOrder-th derivative of a BSpline curve at a given parameter.
+		 *
+		 * @param Window     A view of pointers to control points for the current BSpline segment.
+		 *                   For a BSpline of degree p, this view should contain p+1 control points.
+		 * @param Knots      The global knot vector.
+		 * @param Span       The knot span index such that Parameter ∈ [Knots[Span], Knots[Span+1]).
+		 * @param Parameter  The parameter at which to evaluate.
+		 * @param Degree     The degree of the current BSpline (should be ≥ DerivOrder).
+		 *
+		 * @return The DerivOrder-th derivative (of type T) evaluated at Parameter.
+		 */
+		static T Compute(
+			TArrayView<const T* const> Window,
+			const TArray<float>& Knots,
+			int32 Span,
+			float Parameter,
+			int32 Degree)
 		{
-			Result += *Window[i] * Basis[i];
+			// Basic validation
+			if (Window.Num() < Degree + 1 || DerivOrder > Degree)
+			{
+				return T();
+			}
+
+			TArray<T> DerivControlPoints;
+			DerivControlPoints.SetNum(Degree);
+			const int32 KnotStart = Span - Degree;
+
+			// Determine if knots are uniformly spaced
+			bool bIsUniform = true;
+			const float h = Knots[1] - Knots[0];
+			for (int32 i = 2; i < Knots.Num(); ++i)
+			{
+				if (!FMath::IsNearlyEqual(Knots[i] - Knots[i - 1], h, UE_KINDA_SMALL_NUMBER))
+				{
+					bIsUniform = false;
+					break;
+				}
+			}
+
+			// Compute derivative control points
+			for (int32 i = 0; i < Degree; ++i)
+			{
+				const int32 KnotDenom = KnotStart + i + Degree - DerivOrder + 2; //Knot index for the denominator
+				const int32 KnotNum = KnotStart + i + 1; //Knot index in the numerator
+
+				if (KnotDenom >= Knots.Num() || KnotNum >= Knots.Num() ||
+					KnotDenom < 0 || KnotNum < 0)
+				{
+					return T();
+				}
+
+				const float Denom = Knots[KnotDenom] - Knots[KnotNum];
+
+				// Handle uniform knots differently
+				if (bIsUniform)
+				{
+					// For uniform knots, use the fixed interval h instead of Denom
+					DerivControlPoints[i] = ((*Window[i + 1] - *Window[i])) * (static_cast<float>(Degree) / h);
+					//Correct scaling for uniform knots
+				}
+				else
+				{
+					DerivControlPoints[i] = ((*Window[i + 1] - *Window[i])) * (static_cast<float>(Degree) / Denom);
+					//Correct scaling for non-uniform knots
+				}
+			}
+
+			// Adjust scale factor for uniform knots
+			//Removing the scale factor
+			float ScaleFactor = 1.0;
+
+			// Recurse for higher derivatives with adjusted knots and parameters
+			if constexpr (DerivOrder > 1)
+			{
+				//No change here
+				TArray<const T*> DerivWindow;
+				DerivWindow.SetNum(DerivControlPoints.Num());
+				for (int32 i = 0; i < DerivControlPoints.Num(); ++i)
+				{
+					DerivWindow[i] = &DerivControlPoints[i];
+				}
+
+				// Trim the knot vector for uniform case
+				TArray<float> TrimmedKnots;
+				if (bIsUniform)
+				{
+					TrimmedKnots.Reserve(Knots.Num() - 2);
+					for (int32 i = 1; i < Knots.Num() - 1; ++i)
+					{
+						TrimmedKnots.Add(Knots[i]);
+					}
+				}
+				else
+				{
+					TrimmedKnots = Knots;
+				}
+
+				return TBSplineDerivativeCalculator<T, DerivOrder - 1>::Compute(
+					DerivWindow, //Reduced degree control points
+					TrimmedKnots, //Remove the first and the last knot as the degree is now p-1
+					Span, // Keep Span the same for uniform knots
+					Parameter, //Same parameter
+					Degree - 1 //Reduce the degree
+				) * ScaleFactor; //Multiply with scale factor
+			}
+			else
+			{
+				// Base case: First derivative
+				TArray<float> Basis;
+				Math::ComputeBSplineBasisFunctions(Knots, Span, Parameter, Degree - 1, Basis);
+
+				T Result = T();
+				for (int32 i = 0; i < Degree; ++i)
+				{
+					Result += DerivControlPoints[i] * Basis[i];
+				}
+				return Result * ScaleFactor;
+			}
 		}
+	};
 
-		// Apply segment scaling based on derivative order
-		// For derivatives, we need to divide by SegmentScale^Order because 
-		// derivatives scale inversely with parameter range
-		if constexpr (Order > 0)
+	// Base case: 0th derivative is just the original BSpline evaluation.
+	template <typename T>
+	struct TBSplineDerivativeCalculator<T, 0>
+	{
+		/**
+		 * Evaluates the BSpline curve at Parameter using de Boor's algorithm.
+		 *
+		 * @param Window     A view of pointers to the BSpline's control points (should be Degree+1 points).
+		 * @param Knots      The (possibly trimmed) knot vector.
+		 * @param Span       The knot span index.
+		 * @param Parameter  The parameter at which to evaluate.
+		 * @param Degree     The degree of the BSpline.
+		 *
+		 * @return The BSpline value (of type T) at Parameter.
+		 */
+		static T Compute(
+			TArrayView<const T* const> Window,
+			const TArray<float>& Knots,
+			int32 Span,
+			float Parameter,
+			int32 Degree)
 		{
-			// Calculate scaling factor: 1 / SegmentScale^Order
-			float ScalingFactor = 1.0f;
-			for (int32 i = 0; i < Order; ++i)
+			// Regular evaluation using current basis functions
+			TArray<float> Basis;
+			ComputeBSplineBasisFunctions(Knots, Span, Parameter, Degree, Basis);
+
+			T Result = T();
+			for (int32 i = 0; i <= Degree; ++i)
 			{
-				ScalingFactor /= SegmentScale;
+				Result += *Window[i] * Basis[i];
 			}
-			return Result * ScalingFactor;
+			return Result;
 		}
-		else
+	};
+
+
+	// Bezier curve specific calculations
+	template <typename T, int32 Order>
+	struct TBezierDerivativeCalculator
+	{
+		static T Compute(TArrayView<const T* const> Window, float Parameter, float SegmentScale)
 		{
-			return Result; // No scaling needed for position (Order = 0)
-		}
-	}
-};
+			if constexpr (Order > 3)
+			{
+				return FSplineValidation::GetDefaultValue<T>();
+			}
 
+			if (!FSplineValidation::IsValidWindow(Window, 4) || // Cubic Bezier needs 4 points
+				!FSplineValidation::IsValidParameter(Parameter))
+			{
+				return FSplineValidation::GetDefaultValue<T>();
+			}
 
+			// Use existing basis computation
+			TArray<float> Basis;
+			ComputeBezierBasis<Order>(Parameter, Basis);
 
+			// Apply basis to control points
+			T Result = FSplineValidation::GetDefaultValue<T>();
+			for (int32 i = 0; i < Window.Num(); ++i)
+			{
+				Result += *Window[i] * Basis[i];
+			}
 
-/** 
- * Find nearest point on a cubic Bezier curve segment to a target point
- * @param Coeffs - Array of 4 coefficients in Bezier form relative to target point
- * @param StartParam - Start parameter of segment
- * @param EndParam - End parameter of segment  
- * @param OutSquaredDistance - Output squared distance to nearest point
- * @return Parameter at nearest point between StartParam and EndParam
- */
+			// Apply segment scaling based on derivative order
+			// For derivatives, we need to divide by SegmentScale^Order because 
+			// derivatives scale inversely with parameter range
+			if constexpr (Order > 0)
+			{
+				// Calculate scaling factor: 1 / SegmentScale^Order
+				float ScalingFactor = 1.0f;
+				for (int32 i = 0; i < Order; ++i)
+				{
+					ScalingFactor /= SegmentScale;
+				}
+				return Result * ScalingFactor;
+			}
+			else
+			{
+				return Result; // No scaling needed for position (Order = 0)
+			}
+		}
+	};
 
-template<typename ValueType>
-float FindNearestPoint_Cubic(
-	const TArrayView<ValueType> Coeffs,
-	float StartParam,
-	float EndParam, 
-	float& OutSquaredDistance)
-{
-	// Test endpoints first
-	const ValueType& P0 = Coeffs[0]; // Already relative to target
-	const ValueType At1 = P0 + Coeffs[1] + Coeffs[2] + Coeffs[3]; // Curve at t=1
 
-	float BestDistSq = static_cast<float>(Math::SizeSquared(P0));
-	float BestParam = StartParam;
+	/** 
+	 * Find nearest point on a cubic Bezier curve segment to a target point
+	 * @param Coeffs - Array of 4 coefficients in Bezier form relative to target point
+	 * @param StartParam - Start parameter of segment
+	 * @param EndParam - End parameter of segment  
+	 * @param OutSquaredDistance - Output squared distance to nearest point
+	 * @return Parameter at nearest point between StartParam and EndParam
+	 */
 
-	const float EndDistSq = static_cast<float>(Math::SizeSquared(At1));
-	if (EndDistSq < BestDistSq)
+	template <typename ValueType>
+	float FindNearestPoint_Cubic(
+		const TArrayView<ValueType> Coeffs,
+		float StartParam,
+		float EndParam,
+		float& OutSquaredDistance)
 	{
-		BestDistSq = EndDistSq;
-		BestParam = EndParam;
-	}
+		// Test endpoints first
+		const ValueType& P0 = Coeffs[0]; // Already relative to target
+		const ValueType At1 = P0 + Coeffs[1] + Coeffs[2] + Coeffs[3]; // Curve at t=1
 
-	// Find roots of (C(t) - P).Dot(C'(t))
-	const ValueType& P1 = Coeffs[1];
-	const ValueType& P2 = Coeffs[2];
-	const ValueType& P3 = Coeffs[3];
+		float BestDistSq = static_cast<float>(Math::SizeSquared(P0));
+		float BestParam = StartParam;
 
-	TStaticArray<double, 6> RootCoeffs;
-	RootCoeffs[5] = 3.0 * Dot(P3, P3);
-	RootCoeffs[4] = 5.0 * Dot(P3, P2);
-	RootCoeffs[3] = 4.0 * Dot(P3, P1) + 2.0 * Dot(P2, P2);
-	RootCoeffs[2] = 3.0 * Dot(P2, P1) + 3.0 * Dot(P3, P0);
-	RootCoeffs[1] = Dot(P1, P1) + 2.0 * Dot(P2, P0);
-	RootCoeffs[0] = Dot(P1, P0);
+		const float EndDistSq = static_cast<float>(Math::SizeSquared(At1));
+		if (EndDistSq < BestDistSq)
+		{
+			BestDistSq = EndDistSq;
+			BestParam = EndParam;
+		}
 
-	// Find roots in [0,1]  
-	UE::Math::TPolynomialRootSolver<double, 5> RootSolver(RootCoeffs, 0, 1);
+		// Find roots of (C(t) - P).Dot(C'(t))
+		const ValueType& P1 = Coeffs[1];
+		const ValueType& P2 = Coeffs[2];
+		const ValueType& P3 = Coeffs[3];
 
-	// Check each root
-	for (int32 RootIdx = 0; RootIdx < RootSolver.Roots.Num(); ++RootIdx)
-	{
-		const float t = static_cast<float>(RootSolver.Roots[RootIdx]);
-		const float Param = FMath::Lerp(StartParam, EndParam, t);
+		TStaticArray<double, 6> RootCoeffs;
+		RootCoeffs[5] = 3.0 * Dot(P3, P3);
+		RootCoeffs[4] = 5.0 * Dot(P3, P2);
+		RootCoeffs[3] = 4.0 * Dot(P3, P1) + 2.0 * Dot(P2, P2);
+		RootCoeffs[2] = 3.0 * Dot(P2, P1) + 3.0 * Dot(P3, P0);
+		RootCoeffs[1] = Dot(P1, P1) + 2.0 * Dot(P2, P0);
+		RootCoeffs[0] = Dot(P1, P0);
 
-		const ValueType Point = P0 +
-								P1 * t +
-								P2 * (t * t) +
-								P3 * (t * t * t);
+		// Find roots in [0,1]  
+		UE::Math::TPolynomialRootSolver<double, 5> RootSolver(RootCoeffs, 0, 1);
 
-		const float DistSq = static_cast<float>(Math::SizeSquared(Point));
-		if (DistSq < BestDistSq)
+		// Check each root
+		for (int32 RootIdx = 0; RootIdx < RootSolver.Roots.Num(); ++RootIdx)
 		{
-			BestDistSq = DistSq;
-			BestParam = Param;
+			const float t = static_cast<float>(RootSolver.Roots[RootIdx]);
+			const float Param = FMath::Lerp(StartParam, EndParam, t);
+
+			const ValueType Point = P0 +
+				P1 * t +
+				P2 * (t * t) +
+				P3 * (t * t * t);
+
+			const float DistSq = static_cast<float>(Math::SizeSquared(Point));
+			if (DistSq < BestDistSq)
+			{
+				BestDistSq = DistSq;
+				BestParam = Param;
+			}
 		}
+
+		OutSquaredDistance = BestDistSq;
+		return BestParam;
 	}
 
-	OutSquaredDistance = BestDistSq;
-	return BestParam;
-}
+	template <typename ValueType>
+	ValueType ComputeFirstDerivativeCentralDifference(const TSplineInterface<ValueType>& Spline, double Param,
+													  double Step = UE_KINDA_SMALL_NUMBER)
+	{
+		FInterval1f ParameterSpace = Spline.GetParameterSpace();
+		Step = FMath::Max(Step, UE_KINDA_SMALL_NUMBER);
 
-template <typename ValueType>
-ValueType ComputeFirstDerivativeCentralDifference(const TSplineInterface<ValueType>& Spline, double Param, double Step = UE_KINDA_SMALL_NUMBER)
-{
-	FInterval1f ParameterSpace = Spline.GetParameterSpace();
-	Step = FMath::Max(Step, UE_KINDA_SMALL_NUMBER);
+		const double ParamRight = FMath::Clamp(Param + Step, ParameterSpace.Min, ParameterSpace.Max);
+		const double ParamLeft = FMath::Clamp(Param - Step, ParameterSpace.Min, ParameterSpace.Max);
 
-	const double ParamRight = FMath::Clamp(Param + Step, ParameterSpace.Min, ParameterSpace.Max);
-	const double ParamLeft = FMath::Clamp(Param - Step, ParameterSpace.Min, ParameterSpace.Max);
-	
-	const ValueType ValueRight  = Spline.Evaluate(ParamRight);
-	const ValueType ValueLeft = Spline.Evaluate(ParamLeft);
+		const ValueType ValueRight = Spline.Evaluate(ParamRight);
+		const ValueType ValueLeft = Spline.Evaluate(ParamLeft);
 
-	return (ValueRight - ValueLeft) / (2.0 * Step);
-}
+		return (ValueRight - ValueLeft) / (2.0 * Step);
+	}
 
-template <typename ValueType>
-ValueType ComputeSecondDerivativeCentralDifference(const TSplineInterface<ValueType>& Spline, double Param, double Step = 0.1)
-{
-	FInterval1f ParameterSpace = Spline.GetParameterSpace();
-	Step = FMath::Max(Step, UE_KINDA_SMALL_NUMBER);
-
-	const double ParamRight = FMath::Clamp(Param + Step, ParameterSpace.Min, ParameterSpace.Max);
-	Param = FMath::Clamp(Param, ParameterSpace.Min, ParameterSpace.Max);
-	const double ParamLeft = FMath::Clamp(Param - Step, ParameterSpace.Min, ParameterSpace.Max);
-	
-	const ValueType ValueRight  = Spline.Evaluate(ParamRight);
-	const ValueType Value  = Spline.Evaluate(Param);
-	const ValueType ValueLeft = Spline.Evaluate(ParamLeft);
-
-	return (ValueRight - 2.0 * Value + ValueLeft) / (Step * Step);
-}
-	
+	template <typename ValueType>
+	ValueType ComputeSecondDerivativeCentralDifference(const TSplineInterface<ValueType>& Spline, double Param, double Step = 0.1)
+	{
+		FInterval1f ParameterSpace = Spline.GetParameterSpace();
+		Step = FMath::Max(Step, UE_KINDA_SMALL_NUMBER);
+
+		const double ParamRight = FMath::Clamp(Param + Step, ParameterSpace.Min, ParameterSpace.Max);
+		Param = FMath::Clamp(Param, ParameterSpace.Min, ParameterSpace.Max);
+		const double ParamLeft = FMath::Clamp(Param - Step, ParameterSpace.Min, ParameterSpace.Max);
+
+		const ValueType ValueRight = Spline.Evaluate(ParamRight);
+		const ValueType Value = Spline.Evaluate(Param);
+		const ValueType ValueLeft = Spline.Evaluate(ParamLeft);
+
+		return (ValueRight - 2.0 * Value + ValueLeft) / (Step * Step);
+	}
 } // end namespace UE::Geometry::Spline::Math
 } // end namespace UE::Geometry::Spline
 } // end namespace UE::Geometry
