BugHawk bug: Internal::NextafterNoSubnormal snaps subnormal results to ?FLT_MIN based solely on step direction (bStepRight), which can invert the sign when approaching zero and break monotonic progression. The snapped sign should be derived from the computed result's sign to preserve directionality.

#was-backout 46092739

diff --git a/Engine/Source/Runtime/GeometryCore/Public/Splines/SplineMath.h b/Engine/Source/Runtime/GeometryCore/Public/Splines/SplineMath.h
--- a/Engine/Source/Runtime/GeometryCore/Public/Splines/SplineMath.h
+++ b/Engine/Source/Runtime/GeometryCore/Public/Splines/SplineMath.h
@@ -83,9 +83,18 @@ namespace Internal
 		const float ay = FMath::Abs(y);
 		if (ay > 0.0f && ay < minNorm)
 		{
-			y = bStepRight
-				? +minNorm
-				: -minNorm;
+			// Preserve the computed result's sign to maintain monotonicity
+			const float snapped = std::copysign(minNorm, y);
+
+			if (snapped == from)
+			{
+				// We're at ±FLT_MIN stepping toward zero; make monotone progress to signed zero
+				y = std::copysign(0.0f, y);
+			}
+			else
+			{
+				y = snapped;
+			}
 		}
 		if (!FMath::IsFinite(y))
 		{
