GeometryCore: Ensure DynamicVector serialization always results with at least one allocated block.

diff --git a/Engine/Source/Runtime/GeometryCore/Public/Util/DynamicVector.h b/Engine/Source/Runtime/GeometryCore/Public/Util/DynamicVector.h
--- a/Engine/Source/Runtime/GeometryCore/Public/Util/DynamicVector.h
+++ b/Engine/Source/Runtime/GeometryCore/Public/Util/DynamicVector.h
@@ -959,6 +959,11 @@ void TDynamicVector<Type, BlockSize>::Serialize(FArchive& Ar)
 			}
 		}
 	}
+
+	if (Ar.IsLoading() && Blocks.Num() == 0)
+	{
+		AddAllocatedBlock();
+	}
 }
 
 } // end namespace UE::Geometry
