Fix TangentBezier::RemovePoint bug when removing the 2nd last point in a closed spline.

diff --git a/Engine/Source/Runtime/GeometryCore/Public/Splines/TangentBezierSpline.h b/Engine/Source/Runtime/GeometryCore/Public/Splines/TangentBezierSpline.h
--- a/Engine/Source/Runtime/GeometryCore/Public/Splines/TangentBezierSpline.h
+++ b/Engine/Source/Runtime/GeometryCore/Public/Splines/TangentBezierSpline.h
@@ -1036,7 +1036,7 @@ public:
 			TangentModes.Reset();
 			return;
 		}
-		else if (NumPoints == 2)
+		else if (NumPoints == 2 && !IsClosedLoop())
 		{
 			// Removing the last point and tangents
 			// For CP = 0, we remove P0, P1, P2
