Remove equality-with-tolerance check from TPolyBezierSpline::FindIndexForParameter

diff --git a/Engine/Source/Runtime/GeometryCore/Public/Splines/PolyBezierSpline.h b/Engine/Source/Runtime/GeometryCore/Public/Splines/PolyBezierSpline.h
--- a/Engine/Source/Runtime/GeometryCore/Public/Splines/PolyBezierSpline.h
+++ b/Engine/Source/Runtime/GeometryCore/Public/Splines/PolyBezierSpline.h
@@ -1298,7 +1298,7 @@ public:
 			float Start = Base::PairKnots[i].Value;
 			float End = Base::PairKnots[i + 1].Value;
 			// Exact match on last segment end point
-			if (i == NumSegments - 1 && FMath::IsNearlyEqual(Parameter, End))
+			if (i == NumSegments - 1 && Parameter == End)
 			{
 				OutLocalParam = 1.0f;
 				return i * 4;
