TangentBezier: Cache tangents when spline has a single point, use these cached tangents when adding the second spline point

diff --git a/Engine/Source/Runtime/GeometryCore/Public/Splines/TangentBezierSpline.h b/Engine/Source/Runtime/GeometryCore/Public/Splines/TangentBezierSpline.h
--- a/Engine/Source/Runtime/GeometryCore/Public/Splines/TangentBezierSpline.h
+++ b/Engine/Source/Runtime/GeometryCore/Public/Splines/TangentBezierSpline.h
@@ -219,6 +219,15 @@ public:
 		if (Ar.IsLoading())
 		{
 			TangentModes.SetNum(NumTangentModes);
+			
+			// It was once valid for a spline to have a single bezier point, but we now expect 3 for this case.
+			// It is never valid to have 1 or 2 points.
+			switch (InternalSpline.NumKeys())
+			{
+			case 1: InternalSpline.AddValue(ValueType());	// fallthrough
+			case 2: InternalSpline.AddValue(ValueType());	// fallthrough
+			default: break;
+			}
 		}
         
 		for (int32 i = 0; i < NumTangentModes; ++i)
@@ -1038,6 +1047,8 @@ public:
 		}
 		else if (NumPoints == 2 && !IsClosedLoop())
 		{
+			ValueType PadValue = Index == 0 ? InternalSpline.GetValue(2) : InternalSpline.GetValue(1);
+		
 			// Removing the last point and tangents
 			// For CP = 0, we remove P0, P1, P2
 			// For CP = 1, we remove P1, P2, P3
@@ -1045,6 +1056,16 @@ public:
 			InternalSpline.RemoveValue(Index+1);
 			InternalSpline.RemoveValue(Index);
 			TangentModes.RemoveAt(Index);
+			
+			// If we remove point 0:
+			// Values = < p3, p2, p2 >
+			
+			// Ff we remove point 1:
+			// Values = < p0, p1, p1 >
+
+			InternalSpline.AddValue(PadValue);
+			InternalSpline.AddValue(PadValue);
+			
 			return;
 		}
     
@@ -1312,12 +1333,19 @@ public:
     {
     	const int32 NumKeys = InternalSpline.NumKeys();
     
-    	if (NumKeys == 0)
+    	switch (NumKeys)
+    	{
+    	case 0:
     		return 0;
-    	else if (NumKeys == 1)
-    		return 1;  // Special case when we just have the first point
-    	else
-    		return IsClosedLoop() ? GetNumberOfSegments() : GetNumberOfSegments() + 1;
+    	case 1:
+    		// intentional fallthrough!
+    	case 2:
+    		ensureAlwaysMsgf(false, TEXT("Invalid number of bezier points!"));
+    	case 3:
+			return 1;
+    	default:
+			return IsClosedLoop() ? GetNumberOfSegments() : GetNumberOfSegments() + 1;
+    	}
     }
 
 	virtual int32 GetNumberOfSegments() const override
@@ -1655,10 +1683,15 @@ private:
 	 */
 	bool InitializeFirstSegment(const FTangentBezierControlPoint& ControlPoint, int32& OutPointIndex)
     {
-    	if (InternalSpline.NumKeys() == 0)
+		const int32 NumPoints = GetNumPoints();
+    
+    	if (NumPoints == 0)
     	{
-    		// First point - just store it directly
+    		// First point - just store it directly with both tangents, these are not used until the first segment is completed
     		InternalSpline.AddValue(ControlPoint.Position);
+    		InternalSpline.AddValue(ControlPoint.Position + (ControlPoint.TangentOut / 3.0f));
+			InternalSpline.AddValue(ControlPoint.Position - (ControlPoint.TangentIn  / 3.0f));
+    		
     		// update tangent mode
             if (TangentModes.IsEmpty())
             {
@@ -1668,43 +1701,62 @@ private:
     		{
     			TangentModes[0] = ControlPoint.TangentMode;
     		}
+    		
     		OutPointIndex = 0;
     		return true;
     	}
-
-    	if (InternalSpline.NumKeys() == 1)
+    	else if (NumPoints == 1)
     	{
-    		// We have exactly one point stored - this is our second point
-    		// Get the first point, then clear to start building segments
-    		ValueType P0 = InternalSpline.GetValue(0);
-    		ValueType P3 = ControlPoint.Position;
-    		// Calculate control points
-    		ValueType P1, P2;
-    		if (ControlPoint.TangentMode == ETangentMode::Auto ||
-				ControlPoint.TangentMode == ETangentMode::AutoClamped)
+    		constexpr bool bAppend = true;
+    		if (bAppend)
     		{
-    			// For first segment with auto tangents, use simple interpolation
-    			ValueType AutoTangent = (P3 - P0) * (1.0f - Tension);
-    			P1 = P0 + (AutoTangent / 3.0f);
-    			P2 = P3 - (AutoTangent / 3.0f);
+				// We have exactly one point stored - this is our second point
+				// Get the first point, then clear to start building segments
+				
+				// Calculate control points
+				ValueType P0 = InternalSpline.GetValue(0);
+				ValueType P1, P2;
+				ValueType P3 = ControlPoint.Position;
+				
+				ValueType AutoTangent = (P3 - P0) * (1.0f - Tension);
+				
+				if (ControlPoint.TangentMode == ETangentMode::Auto ||
+					ControlPoint.TangentMode == ETangentMode::AutoClamped)
+				{
+					P2 = P3 - AutoTangent / 3.0f;
+				}
+				else
+				{
+					// Use provided tangents
+					P2 = P3 - ControlPoint.TangentIn / 3.0f;
+				}
+				
+				if (IsAutoTangent(0))
+				{
+					P1 = P0 + AutoTangent / 3.0f;
+				}
+				else
+				{
+					// Values[1] is already the proper out tangent value (Values[2] is the in tangent we'd use if prepending instead of appending here)
+					P1 = InternalSpline.GetValue(1);
+				}
+				
+				// Add the first segment
+				InternalSpline.SetControlPoints( {P0, P1, P2, P3}, EParameterizationPolicy::Centripetal);
+				
+				// update tangent mode
+				TangentModes.AddDefaulted();
+				SetPointTangentMode(1, ControlPoint.TangentMode);
+				
+				OutPointIndex = 1;
+				return true;
     		}
     		else
     		{
-    			// Use provided tangents
-    			P1 = P0 + (ControlPoint.TangentOut / 3.0f); // Need to save first point's OutTangent
-    			P2 = P3 - (ControlPoint.TangentIn / 3.0f);
+    			// todo: prepend
     		}
-    		
-    		// Add the first segment
-    		InternalSpline.SetControlPoints( {P0, P1, P2, P3}, EParameterizationPolicy::Centripetal);
-
-    		// update tangent mode
-    		TangentModes.AddDefaulted();
-    		SetPointTangentMode(1, ControlPoint.TangentMode);
-    		
-    		OutPointIndex = 1;
-    		return true;
     	}
+    	
     	return false;
     }
 
