Spline Core: Remove early outs causing wonky evaluation near segment end points

These were an attempt at bailing out of a slightly more expensive code path, but they behave strange in certain cases.

preflight 68f126b4385729fbf920003f

diff --git a/Engine/Source/Runtime/GeometryCore/Public/Splines/InterpolationPolicies.h b/Engine/Source/Runtime/GeometryCore/Public/Splines/InterpolationPolicies.h
--- a/Engine/Source/Runtime/GeometryCore/Public/Splines/InterpolationPolicies.h
+++ b/Engine/Source/Runtime/GeometryCore/Public/Splines/InterpolationPolicies.h
@@ -223,15 +223,6 @@ public:
 			return TQuat::Identity;
 		}
 
-		// Optimization: single quaternion with full weight
-		for (int32 Index = 0; Index < Count; ++Index)
-		{
-			if (FMath::IsNearlyEqual(Basis[Index], 1.0f))
-			{
-				return *Window[Index];
-			}
-		}
-
 		// Use first quaternion as reference for hemisphere checks
 		const TQuat& ReferenceQuat = *Window[0];
 
diff --git a/Engine/Source/Runtime/GeometryCore/Public/Splines/PolyBezierSpline.h b/Engine/Source/Runtime/GeometryCore/Public/Splines/PolyBezierSpline.h
--- a/Engine/Source/Runtime/GeometryCore/Public/Splines/PolyBezierSpline.h
+++ b/Engine/Source/Runtime/GeometryCore/Public/Splines/PolyBezierSpline.h
@@ -1659,11 +1659,6 @@ private:
     	int32 SegmentIndex;
     	float LocalT;
     	MapGlobalParameterToLocalSegment(Parameter, SegmentIndex, LocalT);
-    	
-    	if (SegmentIndex == NumSegments - 1 && FMath::IsNearlyEqual(LocalT, 1.0f))
-    	{
-    		return *Window[3];  // Last control point of final segment
-    	}
 
         // Calculate Bernstein basis functions
         const float t = LocalT;
