FSparseDynamicPointOctree3: export FindKClosestPoints function

diff --git a/Engine/Source/Runtime/GeometryCore/Public/Spatial/SparseDynamicPointOctree3.h b/Engine/Source/Runtime/GeometryCore/Public/Spatial/SparseDynamicPointOctree3.h
--- a/Engine/Source/Runtime/GeometryCore/Public/Spatial/SparseDynamicPointOctree3.h
+++ b/Engine/Source/Runtime/GeometryCore/Public/Spatial/SparseDynamicPointOctree3.h
@@ -236,7 +236,7 @@ public:
 	 * @param DistSqFunc User-provided squared distance between the point with the given ID and the QueryPt
 	 * @param TempBuffer Optional temporary buffer to store a queue of cells (to avoid memory allocations)
 	 */
-	void FindKClosestPoints(FVector3d QueryPt, double DistanceThreshold, int32 NumToFind,
+	GEOMETRYCORE_API void FindKClosestPoints(FVector3d QueryPt, double DistanceThreshold, int32 NumToFind,
 		TArray<TPair<int32, double>>& FoundPoints, TFunctionRef<bool(int32)> PredicateFunc, TFunctionRef<double(int32)> DistSqFunc,
 		TArray<const FSparsePointOctreeCell*>* TempBuffer = nullptr) const;
 
