ParallelAdaptiveRefinement: respect constraints for inherited neighbor splits

#rnx

diff --git a/Engine/Source/Runtime/GeometryCore/Public/Tessellation/ParallelAdaptiveRefinement.h b/Engine/Source/Runtime/GeometryCore/Public/Tessellation/ParallelAdaptiveRefinement.h
--- a/Engine/Source/Runtime/GeometryCore/Public/Tessellation/ParallelAdaptiveRefinement.h
+++ b/Engine/Source/Runtime/GeometryCore/Public/Tessellation/ParallelAdaptiveRefinement.h
@@ -213,11 +213,12 @@ public:
 
 					if (Options.bApplyNeighborhoodRefinement)
 					{
-						// now the same for all the neighors
+						// for more gradual mesh regularity, tag neighbors for refinement.
 						for (int NbIdx=0; NbIdx<3; ++NbIdx)
 						{
 							const int32 AdjTriIndex = Mesh.GetAdjTriangle(TriIndex, NbIdx);
-							if (AdjTriIndex != IndexConstants::InvalidID && !ConcurrencyState.TriSplitRecords[AdjTriIndex].bClaimed.test_and_set())
+							if (AdjTriIndex != IndexConstants::InvalidID && RefinementPolicy.AllowRefine(AdjTriIndex, Level) && 
+								!ConcurrencyState.TriSplitRecords[AdjTriIndex].bClaimed.test_and_set())
 							{
 								check(Mesh.IsValidTri(AdjTriIndex));
 
