Use StartIndex in the fmoduloiterator struct

diff --git a/Engine/Source/Runtime/GeometryCore/Public/Util/IteratorUtil.h b/Engine/Source/Runtime/GeometryCore/Public/Util/IteratorUtil.h
--- a/Engine/Source/Runtime/GeometryCore/Public/Util/IteratorUtil.h
+++ b/Engine/Source/Runtime/GeometryCore/Public/Util/IteratorUtil.h
@@ -437,7 +437,7 @@ struct FModuloIteration
 	{
 		MaxIndex = (uint64)FMath::Max((uint32)0, MaxIndexIn);
 		StartIndex = (uint64)FMath::Max((uint32)0, StartIndexIn);
-		CurIndex = 0;
+		CurIndex = StartIndex;
 		Count = 0;
 		ModuloNum = FMath::Max((uint64)1, MaxIndex);  // can't be zero or we hit integer-divide. If MaxIndex is 0 we will terminate on first iteration anyway
 		ModuloPrime = ModuloPrimeIn;
