fixed issue where spaced point sampling methods were not using the random seed to determine the downsample ordering, so would be highly correlated across runs regardless of seed

+ make the PCG node pass through the sampling method version

diff --git a/Engine/Plugins/PCGInterops/PCGGeometryScriptInterop/Source/PCGGeometryScriptInterop/Private/Elements/PCGMeshSampler.cpp b/Engine/Plugins/PCGInterops/PCGGeometryScriptInterop/Source/PCGGeometryScriptInterop/Private/Elements/PCGMeshSampler.cpp
--- a/Engine/Plugins/PCGInterops/PCGGeometryScriptInterop/Source/PCGGeometryScriptInterop/Private/Elements/PCGMeshSampler.cpp
+++ b/Engine/Plugins/PCGInterops/PCGGeometryScriptInterop/Source/PCGGeometryScriptInterop/Private/Elements/PCGMeshSampler.cpp
@@ -210,6 +210,7 @@ namespace PCGMeshSampler
 					PointSampling.MaxSamples = Settings->SamplingOptions.MaxNumSamples;
 					PointSampling.RandomSeed = PCGHelpers::ComputeSeed(Seed, Settings->SamplingOptions.RandomSeed);
 					PointSampling.SubSampleDensity = Settings->SamplingOptions.SubSampleDensity;
+					PointSampling.SamplingMethodVersion = Settings->SamplingOptions.SamplingMethodVersion;
 
 					if (Settings->NonUniformSamplingOptions.MaxSamplingRadius > PointSampling.SampleRadius)
 					{
diff --git a/Engine/Plugins/Runtime/GeometryScripting/Source/GeometryScriptingCore/Public/GeometryScript/MeshSamplingFunctions.h b/Engine/Plugins/Runtime/GeometryScripting/Source/GeometryScriptingCore/Public/GeometryScript/MeshSamplingFunctions.h
--- a/Engine/Plugins/Runtime/GeometryScripting/Source/GeometryScriptingCore/Public/GeometryScript/MeshSamplingFunctions.h
+++ b/Engine/Plugins/Runtime/GeometryScripting/Source/GeometryScriptingCore/Public/GeometryScript/MeshSamplingFunctions.h
@@ -38,7 +38,8 @@ public:
 	 * Set this for more consistent results across UE versions, at the risk of worse performance or quality.
 	 * Valid versions are:
 	 * Method 0: UE 5.5 and earlier sampling method. Slower initial (dense) point sampling, less robust to degenerate triangles.
-	 * Method 1: Currently the default method.
+	 * Method 1: Point downsampling is performed in a fixed order regardless of seed, which gives overly-similar sample patterns across multiple runs
+	 * Method 2: Currently the default method.
 	 */
 	UPROPERTY(EditAnywhere, BlueprintReadWrite, Category = Options, AdvancedDisplay)
 	int SamplingMethodVersion = INDEX_NONE;
diff --git a/Engine/Source/Runtime/GeometryCore/Private/Sampling/MeshSurfacePointSampling.cpp b/Engine/Source/Runtime/GeometryCore/Private/Sampling/MeshSurfacePointSampling.cpp
--- a/Engine/Source/Runtime/GeometryCore/Private/Sampling/MeshSurfacePointSampling.cpp
+++ b/Engine/Source/Runtime/GeometryCore/Private/Sampling/MeshSurfacePointSampling.cpp
@@ -14,6 +14,44 @@ using namespace UE::Geometry;
 namespace UELocal
 {
 
+// Generate a semi-random point ordering (via shuffling) 
+// Note this does create some bias in favor of more-densely
+// triangulated regions, which can have more points in the
+// initial dense sampling.
+// Version 1 and earlier also ignore the RandomSeed, so will tend to draw
+// samples from the same triangles across differently-seeded runs.
+void GetShuffledPointOrdering(TArray<int32>& PointOrdering, int32 NumDensePoints, int32 RandomSeed, int32 Version)
+{
+	bool bUseSeed = Version > 1 || Version == INDEX_NONE;
+
+	// as of Version 2, we use the random stream to generate a shuffled order
+	if (bUseSeed)
+	{
+		PointOrdering.SetNumUninitialized(NumDensePoints);
+		FRandomStream Stream(RandomSeed);
+		for (int32 Index = 0; Index < NumDensePoints; ++Index)
+		{
+			// Insert Index at a random position via swap
+			const int32 SwapIndex = Stream.RandRange(0, Index);
+			if (SwapIndex < Index)
+			{
+				PointOrdering[Index] = PointOrdering[SwapIndex];
+			}
+			PointOrdering[SwapIndex] = Index;
+		}
+	}
+	else // older version; use modulo iteration w/out random seed
+	{
+		PointOrdering.Reserve(NumDensePoints);
+		FModuloIteration Iter(NumDensePoints);
+		uint32 NextIndex;
+		while (Iter.GetNextIndex(NextIndex))
+		{
+			PointOrdering.Add(NextIndex);
+		}
+	}
+}
+	
 
 // cache of per-triangle information for a triangle mesh
 struct FTriangleInfoCache
@@ -309,21 +347,15 @@ void UniformMeshPointSampling(
 	// Generate a random point ordering for sampling 
 	//
 
-	// currently only generating a semi-random point ordering (via shuffling) 
-	// if a subset of points is requested.  This likely does create some bias, 
 	TArray<int32> PointOrdering;
-	PointOrdering.Reserve(NumDensePoints);
+
 	if (bShuffle)
 	{
-		FModuloIteration Iter(NumDensePoints);
-		uint32 NextIndex;
-		while (Iter.GetNextIndex(NextIndex))
-		{
-			PointOrdering.Add(NextIndex);
-		}
+		GetShuffledPointOrdering(PointOrdering, NumDensePoints, RandomSeed, SamplingMethodVersion);
 	}
 	else
 	{
+		PointOrdering.Reserve(NumDensePoints);
 		for (int32 k = 0; k < NumDensePoints; ++k)
 		{
 			PointOrdering.Add(k);
@@ -472,15 +504,7 @@ void NonUniformMeshPointSampling(
 	// currently only generating a semi-random point ordering (via shuffling) 
 	// if a subset of points is requested.  This likely does create some bias, 
 	TArray<int32> PointOrdering;
-	PointOrdering.Reserve(NumDensePoints);
-	FModuloIteration Iter(NumDensePoints);
-	{
-		uint32 NextIndex;
-		while (Iter.GetNextIndex(NextIndex))
-		{
-			PointOrdering.Add(NextIndex);
-		}
-	}
+	GetShuffledPointOrdering(PointOrdering, NumDensePoints, RandomSeed, SamplingMethodVersion);
 	int32 CurOrderingIndex = 0;
 
 	// if the bounds are small enough that we'd have at most a single point (i.e., min radius sphere from any point in bounds will cover bounds), early out here
diff --git a/Engine/Source/Runtime/GeometryCore/Public/Sampling/MeshSurfacePointSampling.h b/Engine/Source/Runtime/GeometryCore/Public/Sampling/MeshSurfacePointSampling.h
--- a/Engine/Source/Runtime/GeometryCore/Public/Sampling/MeshSurfacePointSampling.h
+++ b/Engine/Source/Runtime/GeometryCore/Public/Sampling/MeshSurfacePointSampling.h
@@ -176,7 +176,8 @@ public:
 	 * Set this for more consistent results across UE versions, at the risk of worse performance or quality.
 	 * Valid versions are:
 	 * Method 0: UE 5.5 and earlier sampling method. Slower initial (dense) point sampling, less robust to degenerate triangles.
-	 * Method 1: Currently the default method.
+	 * Method 1: Point downsampling is performed in a fixed order regardless of seed, which gives overly-similar sample patterns across multiple runs
+	 * Method 2: Currently the default method.
 	 */
 	int32 SamplingMethodVersion = INDEX_NONE;
 
