Fix compilation issue introduced in 46616872

diff --git a/Engine/Plugins/Runtime/GeometryProcessing/Source/DynamicMesh/Private/Tessellation/AdaptiveDisplacement.cpp b/Engine/Plugins/Runtime/GeometryProcessing/Source/DynamicMesh/Private/Tessellation/AdaptiveDisplacement.cpp
new file mode 100644
--- /dev/null
+++ b/Engine/Plugins/Runtime/GeometryProcessing/Source/DynamicMesh/Private/Tessellation/AdaptiveDisplacement.cpp
@@ -0,0 +1,356 @@
+// Copyright Epic Games, Inc. All Rights Reserved.
+
+#include "Tessellation/AdaptiveDisplacement.h"
+#include "Image/DisplacementMap.h"
+#include "Tessellation/Affine.h"
+#include "Tessellation/FeatureAdaptive.h"
+#include "Tessellation/AdaptiveMetrics.h"
+#include "Tessellation/Regularization.h"
+#include "MeshConstraints.h"
+#include "MeshConstraintsUtil.h"
+#include "Tessellation/DynamicMeshPolicy.h"
+#include "Containers/BitArray.h"
+
+namespace UE {
+namespace Geometry {
+namespace DisplacementInternals {
+	
+struct FDynamicMeshDisplacementPolicy
+{
+	using FIndex3i = Geometry::FIndex3i;
+
+	FDynamicMeshDisplacementPolicy(
+		const FDynamicMesh3& InMesh,
+		const FMeshConstraints& InMeshConstraints,
+		const TBitArray<>& InActiveTriangles,
+		const FDisplacementMap& InDisplacementMap,
+		const double InSampleRate,
+		const double InTargetError,
+		const float InCenter,
+		const float InMagnitude,
+		const float InFeatureSensitivity,
+		const FDynamicMeshUVOverlay* InUVOverlay) // 
+		: Mesh(InMesh)
+		, MeshConstraints(InMeshConstraints)
+		, ActiveTriangles(InActiveTriangles)
+		, DisplacementMap(InDisplacementMap)
+		, SampleRate(InSampleRate)
+		, TargetError(InTargetError)
+		, Center(InCenter)
+		, Magnitude(InMagnitude)
+		, FeatureSensitivity(InFeatureSensitivity)
+		, UVOverlay(InUVOverlay)
+		, MaxRefinementLevel(10)
+		, DisplacementBoundsRefinements(1)
+		, LongestEdgeRatioThreshold(0.95f)
+	{
+		check(Mesh.HasVertexNormals());
+	}
+
+	[[nodiscard]] inline FVector2f TransformUVs(const FVector2f UV) const
+	{
+		// We don't arbitary UV transforms here, it would require cutting the mesh along the internal seams, 
+		// to introduce edge constraints.
+
+		return UV;
+	}
+
+	[[nodiscard]] inline FVector3d GetVertexDisplacement(int32 VertexIndex, const int TriIndex) const
+	{
+		const FIndex3i Triangle = Mesh.GetTriangle(TriIndex);
+		const int LocalVtxIdx = IndexUtil::FindTriIndex(VertexIndex, Triangle);
+
+		const FIndex3i UVTri = UVOverlay->GetTriangle(TriIndex);
+		const FVector2f UV = TransformUVs(UVOverlay->GetElement(UVTri[LocalVtxIdx]));
+
+		if (UV[0] < 0.f || UV[0] > 1.f || UV[1] < 0.f || UV[1] > 1.f)
+			return FVector3d(0.f);
+
+		return FVector3d(Mesh.GetVertexNormal(VertexIndex) * (DisplacementMap.Sample(UV) - Center) * Magnitude);
+	}
+
+	inline void GetTransformedTriUVs(const int TriIndex, FVector2f UVs[3]) const
+	{
+		UVOverlay->GetTriElements(TriIndex, UVs[0], UVs[1], UVs[2]);
+
+		UVs[0] = TransformUVs(UVs[0]);
+		UVs[1] = TransformUVs(UVs[1]);
+		UVs[2] = TransformUVs(UVs[2]);
+	}
+
+	[[nodiscard]] inline FVector3d GetDisplacement(const FVector3d Barycentrics, const int TriIndex) const
+	{
+		const FIndex3i Triangle = Mesh.GetTriangle(TriIndex);
+
+		FVector3f TriNormals[3] = {	Mesh.GetVertexNormal(Triangle[0]),
+			Mesh.GetVertexNormal(Triangle[1]),
+			Mesh.GetVertexNormal(Triangle[2]) };
+
+		FVector3f Normal = TriNormals[0] * static_cast<float>(Barycentrics.X) +
+						   TriNormals[1] * static_cast<float>(Barycentrics.Y) +
+						   TriNormals[2] * static_cast<float>(Barycentrics.Z);
+		Normal.Normalize();
+
+		FVector2f UVs[3];
+		GetTransformedTriUVs(TriIndex, UVs);
+		
+		const FVector2f UV = UVs[0] * static_cast<float>(Barycentrics.X) +
+							 UVs[1] * static_cast<float>(Barycentrics.Y) +
+							 UVs[2] * static_cast<float>(Barycentrics.Z);
+
+		if (UV[0] < 0.f || UV[0] > 1.f || UV[1] < 0.f || UV[1] > 1.f)
+			return FVector3d(0.);
+
+		return FVector3d(Normal) * (DisplacementMap.Sample(UV) - Center) * Magnitude;
+	}
+
+	[[nodiscard]] inline FVector2d GetErrorBounds( const FVector3d* const Barycentrics,
+									 const FVector3d Displacement0,
+									 const FVector3d Displacement1,
+									 const FVector3d Displacement2,
+									 const int TriIndex ) const
+	{
+		const FIndex3i Triangle = Mesh.GetTriangle(TriIndex);
+
+		FVector3f VertexNormals[3] = { 
+			Mesh.GetVertexNormal(Triangle[0]),
+			Mesh.GetVertexNormal(Triangle[1]),
+			Mesh.GetVertexNormal(Triangle[2]) };
+
+		FVector3d VertexNormals64[3] = { FVector3d(VertexNormals[0]), FVector3d(VertexNormals[1]), FVector3d(VertexNormals[2]) };
+
+		FVector2f UVs[3];
+		GetTransformedTriUVs(TriIndex, UVs);
+
+		const FVector3d Displacements[] = { Displacement0, Displacement1, Displacement2 };
+
+		FVector3f Barycentrics3f[3] = { FVector3f(Barycentrics[0]), FVector3f(Barycentrics[1]), FVector3f(Barycentrics[2]) };
+	
+		FVector2f UVBounds[2];
+		UE::Geometry::TessellationUtil::CalculateUVBounds(Barycentrics3f, UVs, UVBounds);
+			
+		FVector2f MagnitudeScaleBounds(Magnitude);
+			
+		return UE::Geometry::TessellationUtil::CalculateErrorBounds<double>(Barycentrics, VertexNormals64, UVBounds, Displacements, DisplacementMap, Center, MagnitudeScaleBounds, DisplacementBoundsRefinements);
+	}
+
+	[[nodiscard]] inline int32 GetNumSamples(const FVector3d* const Barycentrics, const FIndex3i& Triangle, const int TriIndex) const
+	{
+		FVector2f UVs[3];
+		GetTransformedTriUVs(TriIndex, UVs);
+
+		return TessellationUtil::GetNumSamples( Barycentrics, UVs, FVector2f(static_cast<float>(DisplacementMap.GetSizeX()), static_cast<float>(DisplacementMap.GetSizeY())) );
+	}
+
+	[[nodiscard]] inline bool ShouldRefine(const int32 TriIndex, const TArray<FVector3d>& Displacements, FVector3d& SplitVertexBary, const int32 Level) const
+	{
+		if (Level >= MaxRefinementLevel) 
+		{
+			return false;
+		}
+
+		if (TriIndex < ActiveTriangles.Num() && !ActiveTriangles[TriIndex])
+		{
+			// ActiveTriangles may be empty -> no constraint, or if the the index is larger, it is a newly added triangle, that is
+			// implicitly active
+			return false;
+		}
+
+		FVector3d P0, P1, P2;
+		Mesh.GetTriVertices(TriIndex, P0, P1, P2);
+
+		const FVector3d Edge01 = P1 - P0;
+		const FVector3d Edge12 = P2 - P1;
+		const FVector3d Edge20 = P0 - P2;
+
+		const FVector3d EdgeLengths( Edge01.Length(), Edge12.Length(), Edge20.Length() );
+
+		SplitVertexBary = FVector3d( 1./3., 1./3., 1./3. );
+
+		const FIndex3i Triangle = Mesh.GetTriangle(TriIndex);
+		
+		FVector2f UVs[3];
+		GetTransformedTriUVs(TriIndex, UVs);
+		
+		// check for padding region
+		const FVector2f MinUV(UVs[0].ComponentMin(UVs[1]).ComponentMin(UVs[2]));
+		const FVector2f MaxUV(UVs[0].ComponentMax(UVs[1]).ComponentMax(UVs[2]));
+		
+		if (MaxUV[0] < 0.f || MaxUV[1] < 0.f || MinUV[0] > 1.f || MinUV[1] > 1.f)
+		{
+			return false;
+		}
+
+		const FVector3d Displacement0 = Displacements[Triangle.A];
+		const FVector3d Displacement1 = Displacements[Triangle.B];
+		const FVector3d Displacement2 = Displacements[Triangle.C];
+
+		if (FeatureSensitivity > 0.f)
+		{
+			TFeatureSensitivityRefinementCheck<double> RefinementCheck(FeatureSensitivity, SampleRate, P0 + Displacement0, P1 + Displacement1, P2 + Displacement2);
+
+			if (RefinementCheck.NeedsCheck())
+			{
+				bool bShouldSplit = false;
+				for (int32 EdgeIdx = 0; EdgeIdx < 3; ++EdgeIdx)
+				{
+					const int EID = Mesh.GetTriEdge(TriIndex, EdgeIdx);
+					
+					UE::Geometry::FEdgeConstraint EdgeConstraint = MeshConstraints.GetEdgeConstraint(EID);
+
+					// we only apply the feature criterion if all of the edges are unconstrained
+					if (!EdgeConstraint.IsUnconstrained())
+					{
+						bShouldSplit = false;
+						break;
+					}
+
+					const Geometry::FDynamicMesh3::FEdge Edge = Mesh.GetEdge(EID);
+
+					if (bShouldSplit == false && Edge.Tri[1] != IndexConstants::InvalidID)
+					{
+						const int32 AdjTriIdx = (Edge.Tri[0] == TriIndex) ? Edge.Tri[1] : Edge.Tri[0];
+
+						const FIndex3i AdjTri = Mesh.GetTriangle(AdjTriIdx);
+						const int32 FarLocalIndex = IndexUtil::FindTriOtherIndex(Edge.Vert[0], Edge.Vert[1], AdjTri);
+						const int32 FarVID = AdjTri[FarLocalIndex];
+
+						check(AdjTri[(FarLocalIndex+1)%3] == Edge.Vert[0] || AdjTri[(FarLocalIndex+1)%3] == Edge.Vert[1]);
+						check(AdjTri[(FarLocalIndex+2)%3] == Edge.Vert[0] || AdjTri[(FarLocalIndex+2)%3] == Edge.Vert[1]);
+
+						const FVector3d P = Mesh.GetVertex(FarVID) + Displacements[FarVID];
+
+						if (RefinementCheck.CheckCriterion(EdgeIdx, P))
+						{
+							bShouldSplit = true;
+						}
+					}
+				}
+
+				if (bShouldSplit)
+				{
+					return true;
+				}
+			}
+		}
+
+		if( EdgeLengths[0] < SampleRate &&
+			EdgeLengths[1] < SampleRate &&
+			EdgeLengths[2] < SampleRate )
+		{
+			return false;
+		}
+		
+		// avoid skinny triangles
+		const int32 LongestEdgeIdx = FMath::Max3Index( EdgeLengths[0], EdgeLengths[1], EdgeLengths[2] );
+		const double LongestEdge = EdgeLengths[LongestEdgeIdx];
+		const double RemainingEdgeLen = EdgeLengths[(LongestEdgeIdx+1)%3] + EdgeLengths[(LongestEdgeIdx+2)%3];
+			
+		if (LongestEdge > LongestEdgeRatioThreshold * RemainingEdgeLen)
+		{
+			// try edge split
+			if (!Mesh.IsBoundaryEdge(Mesh.GetTriEdge(TriIndex, LongestEdgeIdx)))
+			{
+				SplitVertexBary[LongestEdgeIdx]       = 0.5f;
+				SplitVertexBary[(LongestEdgeIdx+1)%3] = 0.5f;
+				SplitVertexBary[(LongestEdgeIdx+2)%3] = 0.0f;
+				
+				return true;
+			}
+
+			return false;
+		}
+	
+		FVector3d Barycentrics[3] =
+		{
+			{ 1.0f, 0.0f, 0.0f },
+			{ 0.0f, 1.0f, 0.0f },
+			{ 0.0f, 0.0f, 1.0f },
+		};
+
+		const FVector2d ErrorBounds = GetErrorBounds(
+			Barycentrics,
+			Displacement0,
+			Displacement1,
+			Displacement2,
+			TriIndex);
+
+		if (ErrorBounds[1] > TargetError)
+		{
+			return true;
+		}
+
+		return false;
+	}
+
+private:
+	const FDynamicMesh3&           Mesh;
+	const FMeshConstraints&        MeshConstraints;
+	const TBitArray<>&             ActiveTriangles;
+	const FDisplacementMap&        DisplacementMap;
+	const double                   SampleRate;
+	const double                   TargetError;
+	const float                    Center;
+	const float                    Magnitude;
+	const float                    FeatureSensitivity;
+		
+	const FDynamicMeshUVOverlay*   UVOverlay        { nullptr };
+	
+	const int32                    MaxRefinementLevel;
+	const int32                    DisplacementBoundsRefinements;
+	const float                    LongestEdgeRatioThreshold;
+};
+
+} // namespace DisplacementInternals
+
+bool ApplyAdaptiveDisplacement(
+	FDynamicMesh3& Mesh,
+	FDynamicMeshUVOverlay* UVOverlay,
+	const FDisplacementMap& DisplacementMap,
+	const FAdaptiveTessellatorOptions& Options,
+	const float FeatureSensitivity,
+	FMeshConstraints& MeshConstraints,
+	TBitArray<>& ActiveTriangles,
+	const bool bApplyPostOptimization )
+{
+	if (!Mesh.HasVertexNormals())
+	{
+		return false;
+	}
+
+	FDynamicMesh3Policy DynamicMesh3Policy(Mesh, MeshConstraints);
+
+	DisplacementInternals::FDynamicMeshDisplacementPolicy DisplacementPolicy(
+		Mesh, MeshConstraints, ActiveTriangles, DisplacementMap, 
+		Options.SampleRate, Options.TargetError, 
+		0.f, 1.f, 
+		FeatureSensitivity, 
+		UVOverlay );
+
+	using TessellatorT = TAdaptiveTessellator<FDynamicMesh3Policy, DisplacementInternals::FDynamicMeshDisplacementPolicy>;
+
+	TessellatorT Tessellator(DynamicMesh3Policy, DisplacementPolicy, Options);
+
+	if (bApplyPostOptimization)
+	{
+		SplitLongEdges(Mesh);
+
+		IntrinsicDelaunayFlipEdge(Mesh);
+
+		TBitArray<> IsInteriorVertex;
+		IsInteriorVertex.Init(true, Mesh.MaxVertexID()); 
+		for (const int32 EID : Mesh.BoundaryEdgeIndicesItr())
+		{
+			FDynamicMesh3::FEdge Edge = Mesh.GetEdge(EID);
+			IsInteriorVertex[Edge.Vert[0]] = false;
+			IsInteriorVertex[Edge.Vert[1]] = false;
+		}
+
+		RegularizeVolumePreserving(Mesh, true, IsInteriorVertex);
+	}
+
+	return true;
+}
+	
+} // namespace Geometry
+} // namespace UE
diff --git a/Engine/Plugins/Runtime/GeometryProcessing/Source/DynamicMesh/Public/Tessellation/AdaptiveDisplacement.h b/Engine/Plugins/Runtime/GeometryProcessing/Source/DynamicMesh/Public/Tessellation/AdaptiveDisplacement.h
new file mode 100644
--- /dev/null
+++ b/Engine/Plugins/Runtime/GeometryProcessing/Source/DynamicMesh/Public/Tessellation/AdaptiveDisplacement.h
@@ -0,0 +1,42 @@
+// Copyright Epic Games, Inc. All Rights Reserved.
+
+#pragma once
+
+#include "Containers/BitArray.h"
+#include "DynamicMesh/DynamicMeshAttributeSet.h"
+
+namespace UE {
+namespace Geometry {
+
+class FDynamicMesh3;
+struct FAdaptiveTessellatorOptions;
+class FMeshConstraints;
+class FDisplacementMap;
+
+/**
+ * Apply adaptive displacement on given mesh.
+ * 
+ * @param Mesh               mesh to be displaced (requires per-vertex normals to be present)
+ * @param UVOverlay          UVs on which to evaluate displacement map
+ * @param DisplacementMap    displacement map
+ * @param Options            control sample rate, error, etc.
+ * @param FeatureSensitivity value in [0, 1] to control additional adaptive refinement at sharp features
+ * @param MeshConstraints    constraints to preserve
+ * @param ActiveTriangles    optional per-triangle selection
+ * @param bApplyPostOptimization whether to apply some post regularization
+ * @return                   whether the adaptive displacement succeeded
+ */
+
+DYNAMICMESH_API
+bool ApplyAdaptiveDisplacement(
+	FDynamicMesh3& Mesh,
+	FDynamicMeshUVOverlay* UVOverlay,
+	const FDisplacementMap& DisplacementMap,
+	const FAdaptiveTessellatorOptions& Options,
+	const float FeatureSensitivity,
+	FMeshConstraints& MeshConstraints,
+	TBitArray<>& ActiveTriangles,
+	const bool bApplyPostOptimization = true );
+	
+} // namespace Geometry
+} // namespace UE
diff --git a/Engine/Source/Runtime/GeometryCore/Private/Tessellation/Regularization.cpp b/Engine/Source/Runtime/GeometryCore/Private/Tessellation/Regularization.cpp
new file mode 100644
--- /dev/null
+++ b/Engine/Source/Runtime/GeometryCore/Private/Tessellation/Regularization.cpp
@@ -0,0 +1,746 @@
+// Copyright Epic Games, Inc. All Rights Reserved.
+
+#include "Tessellation/Regularization.h"
+
+#include "DynamicMesh/DynamicMesh3.h"
+#include "Containers/Array.h"
+#include "Util/IndexUtil.h"
+#include "VectorUtil.h"
+#include "VectorTypes.h"
+#include "HAL/IConsoleManager.h"
+#include "Async/ParallelFor.h"
+#include "MathUtil.h"
+
+namespace UE {
+namespace Geometry {
+
+/**
+ * A binary heap data structure with external swap tracking capability.
+ *
+ * This class implements a min-heap (based on operator<) and allows external systems
+ * to track element movement via a swap policy. This is useful in scenarios like
+ * priority queues where an external index mapping needs to be updated on swaps.
+ *
+ * @tparam T               The type of elements stored in the heap.
+ * @tparam SwapPolicyType  A policy class that defines Swap(int32, int32, T&, T&) to track swaps.
+ */
+template <typename T, typename SwapPolicyType>
+class THeapWithTracking
+{
+public:
+	using SizeType = typename TArray<T>::SizeType;
+	
+	/**
+	 * @param InSwapPolicy Reference to a swap policy instance used to track heap element swaps.
+	 */
+	THeapWithTracking(SwapPolicyType& InSwapPolicy)
+	: SwapPolicy(InSwapPolicy)
+	{
+	}
+
+	/**
+	 * Adds an element to the heap and reorders to maintain the heap property.
+	 *
+	 * @param Element The element to add.
+	 */
+	void Add(const T& Element)
+	{
+		HeapUp(Heap.Add(Element));
+	}
+	/**
+	 * Removes and returns the root (minimum) element from the heap.
+	 *
+	 * @return The removed root element.
+	 */
+	T Pop()
+	{
+		check(!Heap.IsEmpty());
+		const T Root = Heap[0];
+		if (Heap.Num() > 1)
+		{
+			Swap(0, Heap.Num()-1);
+			Heap.Pop();
+			HeapDown(0);
+		}
+		else
+		{
+			Heap.Pop();
+		}
+		return Root;
+	}
+
+	/**
+	 * Updates an element at a given index and reorders the heap accordingly.
+	 *
+	 * @param Idx     Index of the element to update.
+	 * @param Element The new element value.
+	 */
+	void Update(const SizeType Idx, const T& Element)
+	{
+		Heap[Idx] = Element;
+		HeapUp(Idx);
+		HeapDown(Idx);
+	}
+
+	/**
+	 * Removes the element at the specified index and reorders the heap to maintain
+	 * heap property.
+	 *
+	 * @param Idx Index of the element to remove.
+	 */
+	void Remove(const SizeType Idx)
+	{
+		check(Idx < Heap.Num());
+		if (Heap.Num() > 1)
+		{
+			Swap(Idx, Heap.Num()-1);
+			Heap.Pop();
+			HeapDown(Idx);
+		}
+		else
+		{
+			Heap.Pop();
+		}
+	}
+
+	/**
+	 * Checks if the heap is empty.
+	 *
+	 * @return true if the heap is empty; false otherwise.
+	 */
+	bool IsEmpty() const
+	{
+		return Heap.IsEmpty();
+	}
+
+	/**
+	 * Returns the number of elements currently in the heap.
+	 *
+	 * @return The number of heap elements.
+	 */
+	SizeType Num() const 
+	{ 
+		return Heap.Num(); 
+	}
+
+	/**
+	 * Provides read-only access to an element by index.
+	 *
+	 * @param Idx Index to access.
+	 * @return Reference to the element at the specified index.
+	 */
+	const T& operator[](const SizeType Idx) const
+	{
+		return Heap[Idx];
+	}
+
+private:
+
+	/**
+	 * Bubbles an element up the heap to restore the heap property.
+	 *
+	 * @param Idx Index of the element to heapify up.
+	 */
+	void HeapUp(SizeType Idx)
+	{
+		while (Idx > 0)
+		{
+			const SizeType Parent = (Idx - 1) >> 1;
+			if (!(Heap[Idx] < Heap[Parent]))
+			{
+				break;
+			}
+
+			Swap(Idx, Parent);
+			Idx = Parent;
+		};
+	}
+
+	/**
+	 * Pushes an element down the heap to restore the heap property.
+	 *
+	 * @param Idx Index of the element to heapify down.
+	 */
+	void HeapDown(SizeType Idx)
+	{
+		const int32 Num = Heap.Num();
+
+		while (true)
+		{
+			const SizeType Child0 = (Idx << 1) + 1;
+			const SizeType Child1 = Child0 + 1;
+
+			SizeType MinIdx = Idx;
+			if (Child0 < Num && Heap[Child0] < Heap[MinIdx])
+			{
+				MinIdx = Child0;
+			}
+			if (Child1 < Num && Heap[Child1] < Heap[MinIdx])
+			{
+				MinIdx = Child1;
+			}
+
+			if (MinIdx == Idx)
+			{
+				break;
+			}
+
+            Swap(Idx, MinIdx);
+            Idx = MinIdx;
+		}
+	}
+
+	/**
+	 * Swaps two elements in the heap and informs the swap policy.
+	 *
+	 * @param A Index of the first element.
+	 * @param B Index of the second element.
+	 */
+	void Swap(const SizeType A, const SizeType B)
+	{
+		SwapPolicy.Swap(A, B, Heap[A], Heap[B]);
+		::Swap(Heap[A], Heap[B]);
+	}
+
+	// The underlying array representing the heap.
+	TArray<T> Heap;
+
+	// Reference to the swap policy for tracking element movements.
+	SwapPolicyType& SwapPolicy;
+};
+
+// [ Baensch et al. "A Finite Element Method for Surface Diffusion", 2004]
+void RegularizeVolumePreserving(FDynamicMesh3& Mesh, bool bPreserveVolume, const TOptional<TBitArray<>>& ActiveVertices)
+{
+	constexpr double ClampT = 0.5;
+	TArray<FVector3d> Vertices;
+
+	Vertices.SetNum(Mesh.MaxVertexID());
+
+	for (int VID : Mesh.VertexIndicesItr())
+	{
+		Vertices[VID] = Mesh.GetVertex(VID);		
+	}
+
+	ParallelFor(Mesh.MaxVertexID(), [&](int32 VID)
+	{
+		if (!Mesh.IsVertex(VID))
+		{
+			return;
+		}
+
+		if (ActiveVertices && !(*ActiveVertices)[VID])
+		{
+			return;
+		}
+
+		const FVector3d Z = Vertices[VID];
+
+		FVector3d ZHat(0.);
+		int TriCount = 0;
+		FVector3d NormalDirection(0.);
+
+		for (int TID : Mesh.VtxTrianglesItr(VID))
+		{
+			const FIndex3i Tri = Mesh.GetTriangle(TID);
+			const FVector3d V[3] = { Vertices[Tri[0]], Vertices[Tri[1]], Vertices[Tri[2]] };			
+
+			FVector3d Edge1(V[1] - V[0]);
+			FVector3d Edge2(V[2] - V[0]);
+			
+			FVector3d vCross(Edge2.Cross(Edge1));
+			NormalDirection += vCross;
+
+			ZHat += (V[0] + V[1] + V[2]) * 1./3.;
+			TriCount++;
+		}
+
+		if (TriCount == 0) 
+		{
+			return;
+		}
+
+		ZHat /= static_cast<double>(TriCount);
+
+
+		const double NormalDirectionLength = NormalDirection.Length();
+		if (NormalDirectionLength < UE_DOUBLE_SMALL_NUMBER)
+		{
+			return;
+		}
+        // scale the normal direction such that t becomes invariant of mesh scaling (note, this is not unit normal, but proportional to edge length) 
+		NormalDirection /= FMath::Sqrt(NormalDirectionLength);
+
+		double t = 0.;
+
+		if (bPreserveVolume)
+		{
+			double nom(0.), denom(0.);
+			for (int TID : Mesh.VtxTrianglesItr(VID))
+			{
+				const FIndex3i Tri = Mesh.GetTriangle(TID);
+				const int CID = IndexUtil::FindTriIndex(VID, Tri);
+
+				const FVector3d Z2 = Vertices[Tri[(CID+1)%3]];
+				const FVector3d Z3 = Vertices[Tri[(CID+2)%3]];
+
+				nom   += ((Z-ZHat) ^ (Z2-ZHat)) | (Z3-ZHat);
+				denom += (NormalDirection   ^ (Z2-ZHat)) | (Z3-ZHat);
+			}
+			t = nom / denom;
+		}
+
+		if (TMathUtil<double>::IsFinite(t)) 
+		{
+			Mesh.SetVertex(VID, ZHat + NormalDirection * FMath::Clamp(t, -ClampT, ClampT), false);
+		}
+		else
+		{
+			Mesh.SetVertex(VID, ZHat, false);
+		}
+	});
+}
+
+// Applies edge splits to edges for which |longest edge| > SplitThreshold (|other edge a| + |other edge b|)
+// in order of importance.
+//
+void SplitLongEdges(FDynamicMesh3& Mesh, TArray<FVector3d>& Displacements, FDisplaceFunc DisplacementFunctor, const double SplitThreshold)
+{
+	check(Displacements.IsEmpty() || Displacements.Num() == Mesh.MaxVertexID());
+
+	const bool bHasDisplacements = !Displacements.IsEmpty();
+
+	check(!bHasDisplacements || DisplacementFunctor);
+
+	const auto ComputeWeight = [&](const int32 EID) -> double
+	{
+		const FDynamicMesh3::FEdge& Edge = Mesh.GetEdgeRef(EID);
+	
+		if (Edge.Tri[1] == IndexConstants::InvalidID)
+		{
+			return 0.;
+		}
+
+		const FIndex3i Tri0 = Mesh.GetTriangle(Edge.Tri[0]);
+		const FIndex3i Tri1 = Mesh.GetTriangle(Edge.Tri[1]);
+		
+		const int SharedEdge0 = IndexUtil::FindEdgeIndexInTri(Edge.Vert[0], Edge.Vert[1], Tri0);
+		const int SharedEdge1 = IndexUtil::FindEdgeIndexInTri(Edge.Vert[0], Edge.Vert[1], Tri1);
+
+		const int VID0 = Tri0[SharedEdge0];
+		const int VID1 = Tri1[SharedEdge1];
+
+		check(VID1 == Tri0[(SharedEdge0+1)%3]);
+		check(VID0 == Tri1[(SharedEdge1+1)%3]);
+
+		const int VIDA = Tri0[(SharedEdge0+2)%3];
+		const int VIDB = Tri1[(SharedEdge1+2)%3];
+
+		// Edge (V0, V1) is the split candidate
+		//
+		//       o VA
+		//      / \
+		//     /   \
+		// V0 o=====o V1
+		//     \   /
+		//      \ /
+	    //       o VB
+		// 
+
+		FVector3d V0 = Mesh.GetVertex(VID0);
+		FVector3d V1 = Mesh.GetVertex(VID1);
+		FVector3d VA = Mesh.GetVertex(VIDA);
+		FVector3d VB = Mesh.GetVertex(VIDB);
+
+		if (bHasDisplacements)
+		{
+			V0 += Displacements[VID0];
+			V1 += Displacements[VID1];
+			VA += Displacements[VIDA];
+			VB += Displacements[VIDB];
+		}
+
+		double E01 = (V1 - V0).Length();
+		double E0A = (VA - V0).Length();
+		double E0B = (VB - V0).Length();
+		double E1A = (VA - V1).Length();
+		double E1B = (VB - V1).Length();
+		
+		return E01 / FMath::Max(E0A + E1A, E0B + E1B);		
+	};
+
+	// Entry of the heap
+	struct FEntry
+	{
+		double Weight;
+		int32 EID;
+
+		bool operator<(const FEntry& Other) const
+		{
+			return Weight < Other.Weight;
+		}
+	};
+
+	// keeps track of changed heap indices
+	struct FSwapPolicy
+	{
+		FSwapPolicy(TArray<int32>& InHeapIndices)
+			: HeapIndices(InHeapIndices)
+		{
+		}
+
+		void Swap(const int32 HeapIdxA, const int32 HeapIdxB, const FEntry& EntryA, const FEntry& EntryB)
+		{
+			check(HeapIndices[EntryA.EID] == HeapIdxA);
+			check(HeapIndices[EntryB.EID] == HeapIdxB);
+
+			::Swap(HeapIndices[EntryA.EID], HeapIndices[EntryB.EID]);
+		}
+
+		TArray<int32>& HeapIndices;
+	};
+
+	TArray<int32> HeapIndices;
+	FSwapPolicy SwapPolicy(HeapIndices);
+	THeapWithTracking<FEntry, FSwapPolicy> Heap(SwapPolicy);
+
+	int32 SplitEdgeCount = 0;
+	for (int EID : Mesh.EdgeIndicesItr())
+	{
+		const double Weight = ComputeWeight(EID);
+		if (Weight > SplitThreshold)
+		{
+			HeapIndices.Add(Heap.Num());
+			Heap.Add({ -Weight, EID });
+		}
+		else
+		{
+			HeapIndices.Add(IndexConstants::InvalidID);
+		}
+	}
+
+	check(HeapIndices.Num() == Mesh.MaxEdgeID());
+	
+	while(!Heap.IsEmpty())
+	{
+		const FEntry Entry = Heap.Pop();
+		HeapIndices[Entry.EID] = IndexConstants::InvalidID;
+
+		if (Entry.Weight < -SplitThreshold)
+		{
+			int EdgesToUpdate[4];
+			{
+				const FDynamicMesh3::FEdge& Edge = Mesh.GetEdgeRef(Entry.EID);
+				const FIndex3i Tri0 = Mesh.GetTriangle(Edge.Tri[0]);
+				const FIndex3i Tri1 = Mesh.GetTriangle(Edge.Tri[1]);
+
+				const int SharedEdge0 = IndexUtil::FindEdgeIndexInTri(Edge.Vert[0], Edge.Vert[1], Tri0);
+				const int SharedEdge1 = IndexUtil::FindEdgeIndexInTri(Edge.Vert[0], Edge.Vert[1], Tri1);
+
+				EdgesToUpdate[0] = Mesh.GetTriEdge(Edge.Tri[0], (SharedEdge0+1)%3);
+				EdgesToUpdate[1] = Mesh.GetTriEdge(Edge.Tri[0], (SharedEdge0+2)%3);
+				EdgesToUpdate[2] = Mesh.GetTriEdge(Edge.Tri[1], (SharedEdge1+1)%3);
+				EdgesToUpdate[3] = Mesh.GetTriEdge(Edge.Tri[1], (SharedEdge1+2)%3);
+			}
+
+			FDynamicMesh3::FEdgeSplitInfo SplitInfo;
+			auto Result = Mesh.SplitEdge(Entry.EID, SplitInfo);
+		
+			if (Result == EMeshResult::Ok)
+			{
+				if (bHasDisplacements)
+				{
+					const FVector3d SplitVertexDisplacement = DisplacementFunctor(SplitInfo.NewVertex);
+					if (SplitInfo.NewVertex == Displacements.Num())
+					{
+						Displacements.Add(SplitVertexDisplacement);
+					}
+					else
+					{
+						Displacements.SetNum(FMath::Max(Displacements.Num(), SplitInfo.NewVertex + 1));
+						Displacements[SplitInfo.NewVertex] = SplitVertexDisplacement;
+					}
+				}
+
+				HeapIndices[Entry.EID] = IndexConstants::InvalidID;
+
+				for (int32 I=0; I<3; ++I)
+				{
+					const int32 EID = SplitInfo.NewEdges[I];
+					HeapIndices.Add(IndexConstants::InvalidID);
+				}
+
+				for (const int EID : EdgesToUpdate)
+				{
+					check(EID < Mesh.MaxEdgeID());
+
+					if (EID == IndexConstants::InvalidID || HeapIndices[EID] == IndexConstants::InvalidID)
+					{
+						continue;
+					}
+
+					const FDynamicMesh3::FEdge& Edge = Mesh.GetEdgeRef(EID);
+
+					if (Edge.Tri[1] != IndexConstants::InvalidID)
+					{
+						Heap.Update(HeapIndices[EID], { -ComputeWeight(EID), EID });
+					}
+				}
+			}
+		}
+	}
+}
+
+// Apply edge flips on non-Delaunay edges in order of importance.
+// Intrinsic edge-flips would split the flipped edge and produce edges that are on the same surface.
+// This function relaxes this and applies honest flips when the surface is planar enough (when the
+// change in the angle is smaller than FlipAngleThreshold)
+// 
+// [Fisher, Springborn, Bobenko, Schroeder, "An Algorithm for the Construction of Intrinsic Delaunay Triangulations
+// with Applications to Digital Geometry Processing", 2006]
+//
+void IntrinsicDelaunayFlipEdge(FDynamicMesh3& Mesh, TArray<FVector3d>& Displacements, FDisplaceFunc DisplacementFunctor, const double FlipAngleThreshold)
+{
+	check(Displacements.IsEmpty() || Displacements.Num() == Mesh.MaxVertexID());
+
+	const bool bHasDisplacements = !Displacements.IsEmpty();
+
+	check(!bHasDisplacements || DisplacementFunctor);
+
+	const auto ComputeWeight = [&](const int32 EID) -> std::tuple<double, bool>
+	{
+		const FDynamicMesh3::FEdge& Edge = Mesh.GetEdgeRef(EID);
+		
+		const FIndex3i Tri0 = Mesh.GetTriangle(Edge.Tri[0]);
+		const FIndex3i Tri1 = Mesh.GetTriangle(Edge.Tri[1]);
+		
+		const int SharedEdge0 = IndexUtil::FindEdgeIndexInTri(Edge.Vert[0], Edge.Vert[1], Tri0);
+		const int SharedEdge1 = IndexUtil::FindEdgeIndexInTri(Edge.Vert[0], Edge.Vert[1], Tri1);
+
+		const int VID0 = Tri0[SharedEdge0];
+		const int VID1 = Tri1[SharedEdge1];
+
+		check(VID1 == Tri0[(SharedEdge0+1)%3]);
+		check(VID0 == Tri1[(SharedEdge1+1)%3]);
+
+		const int VIDA = Tri0[(SharedEdge0+2)%3];
+		const int VIDB = Tri1[(SharedEdge1+2)%3];
+
+		FVector3d V0 = Mesh.GetVertex(VID0);
+		FVector3d V1 = Mesh.GetVertex(VID1);
+		FVector3d VA = Mesh.GetVertex(VIDA);
+		FVector3d VB = Mesh.GetVertex(VIDB);
+
+		if (bHasDisplacements)
+		{
+			V0 += Displacements[VID0];
+			V1 += Displacements[VID1];
+			VA += Displacements[VIDA];
+			VB += Displacements[VIDB];
+		}
+
+		const FVector3d EA0 = V0 - VA;
+		const FVector3d EA1 = V1 - VA;
+		const FVector3d EB0 = V0 - VB;
+		const FVector3d EB1 = V1 - VB;
+		const FVector3d E01 = V1 - V0;
+
+		const double CotanA = Geometry::VectorUtil::VectorCot(EA0, EA1);
+		const double CotanB = Geometry::VectorUtil::VectorCot(EB0, EB1);
+
+		const double LaplacianWeight  = 0.5 * (CotanA + CotanB);
+
+		// unflipped
+		FVector3d N0 = Normalized(EA1 ^ EA0);
+		FVector3d N1 = Normalized(EB0 ^ EB1);
+		FVector3d E  = Normalized(E01);
+
+		// flipped
+		FVector3d N0f = Normalized(EA0 ^ EB0);
+		FVector3d N1f = Normalized(EB1 ^ EA1);
+		FVector3d Ef  = Normalized(VA - VB); 
+
+		const double OrientedDihedral  = Geometry::VectorUtil::OrientedDihedralAngle(N0 , N1 , E );
+		const double OrientedDihedralf = Geometry::VectorUtil::OrientedDihedralAngle(N0f, N1f, Ef);
+
+		if (LaplacianWeight > 0.)
+		{
+			return { -LaplacianWeight, false };
+		}
+
+		const bool bNeedsSplit = (FMath::Abs(OrientedDihedralf - OrientedDihedral) > FlipAngleThreshold);
+
+		return { -LaplacianWeight, bNeedsSplit };
+	};
+
+	struct FEntry
+	{
+		double Weight;
+		bool bNeedsSplit { false };
+		int32 EID;
+
+		bool operator<(const FEntry& Other) const
+		{
+			return Weight < Other.Weight;
+		}
+	};
+
+	struct FSwapPolicy
+	{
+		FSwapPolicy(TArray<int32>& InHeapIndices)
+			: HeapIndices(InHeapIndices)
+		{
+		}
+
+		void Swap(const int32 HeapIdxA, const int32 HeapIdxB, const FEntry& EntryA, const FEntry& EntryB)
+		{
+			check(HeapIndices[EntryA.EID] == HeapIdxA);
+			check(HeapIndices[EntryB.EID] == HeapIdxB);
+
+			::Swap(HeapIndices[EntryA.EID], HeapIndices[EntryB.EID]);
+		}
+
+		TArray<int32>& HeapIndices;
+	};
+
+	TArray<int32> HeapIndices;
+	FSwapPolicy SwapPolicy(HeapIndices);
+	THeapWithTracking<FEntry, FSwapPolicy> Heap(SwapPolicy);
+
+	for (int EID : Mesh.EdgeIndicesItr())
+	{
+		const FDynamicMesh3::FEdge& Edge = Mesh.GetEdgeRef(EID);
+		if (Edge.Tri[1] == IndexConstants::InvalidID)
+		{
+			HeapIndices.Add(IndexConstants::InvalidID);
+		}
+		else
+		{
+			double Weight;
+			bool bNeedsSplit;
+			std::tie(Weight, bNeedsSplit) = ComputeWeight(EID);
+
+			if (Weight > 0)
+			{
+				const int32 Idx = Heap.Num();
+				HeapIndices.Add(Idx);
+				Heap.Add({ -Weight, bNeedsSplit, EID });
+			}
+			else
+			{
+				HeapIndices.Add(IndexConstants::InvalidID);
+			}
+		}
+	}
+
+	check(HeapIndices.Num() == Mesh.MaxEdgeID());
+
+	while(!Heap.IsEmpty())
+	{
+		const FEntry Entry = Heap.Pop();
+		HeapIndices[Entry.EID] = IndexConstants::InvalidID;
+
+		if (Entry.Weight < 0.)
+		{
+			int EdgesToUpdate[4];
+			{
+				const FDynamicMesh3::FEdge& Edge = Mesh.GetEdgeRef(Entry.EID);
+				const FIndex3i Tri0 = Mesh.GetTriangle(Edge.Tri[0]);
+				const FIndex3i Tri1 = Mesh.GetTriangle(Edge.Tri[1]);
+
+				const int SharedEdge0 = IndexUtil::FindEdgeIndexInTri(Edge.Vert[0], Edge.Vert[1], Tri0);
+				const int SharedEdge1 = IndexUtil::FindEdgeIndexInTri(Edge.Vert[0], Edge.Vert[1], Tri1);
+
+				EdgesToUpdate[0] = Mesh.GetTriEdge(Edge.Tri[0], (SharedEdge0+1)%3);
+				EdgesToUpdate[1] = Mesh.GetTriEdge(Edge.Tri[0], (SharedEdge0+2)%3);
+				EdgesToUpdate[2] = Mesh.GetTriEdge(Edge.Tri[1], (SharedEdge1+1)%3);
+				EdgesToUpdate[3] = Mesh.GetTriEdge(Edge.Tri[1], (SharedEdge1+2)%3);
+			}
+
+			if (Entry.bNeedsSplit)
+			{
+				FDynamicMesh3::FEdgeSplitInfo SplitInfo;
+				EMeshResult Result = Mesh.SplitEdge(Entry.EID, SplitInfo);
+
+				if (Result == EMeshResult::Ok)
+				{
+					if (bHasDisplacements)
+					{
+						const FVector3d SplitVertexDisplacement = DisplacementFunctor(SplitInfo.NewVertex);
+						if (SplitInfo.NewVertex == Displacements.Num())
+						{
+							Displacements.Add(SplitVertexDisplacement);
+						}
+						else
+						{
+							Displacements.SetNum(FMath::Max(Displacements.Num(), SplitInfo.NewVertex+1));
+							Displacements[SplitInfo.NewVertex] = SplitVertexDisplacement;
+						}
+					}
+
+					HeapIndices[Entry.EID] = IndexConstants::InvalidID;
+
+					for (int32 I = 0; I < 3; ++I)
+					{
+						const int32 EID = SplitInfo.NewEdges[I];
+						HeapIndices.Add(IndexConstants::InvalidID);
+					}
+
+					for (const int EID : EdgesToUpdate)
+					{
+						check(EID < Mesh.MaxEdgeID());
+
+						if (EID == IndexConstants::InvalidID || HeapIndices[EID] == IndexConstants::InvalidID)
+						{
+							continue;
+						}
+
+						const FDynamicMesh3::FEdge& Edge = Mesh.GetEdgeRef(EID);
+
+						if (Edge.Tri[1] != IndexConstants::InvalidID)
+						{
+							double Weight;
+							bool bNeedsSplit;
+							std::tie(Weight, bNeedsSplit) = ComputeWeight(EID);
+
+							Heap.Update(HeapIndices[EID], { Weight, bNeedsSplit, EID });
+						}
+					}
+				}
+			}
+			else
+			{
+				// Flip Only
+				FDynamicMesh3::FEdgeFlipInfo EdgeFlipInfo;
+				const EMeshResult Result = Mesh.FlipEdge(Entry.EID, EdgeFlipInfo);
+
+				if (Result == EMeshResult::Ok)
+				{
+					for (const int EID : EdgesToUpdate)
+					{
+						check(EID < Mesh.MaxEdgeID());
+
+						if (HeapIndices[EID] == IndexConstants::InvalidID)
+						{
+							continue;
+						}
+
+						const FDynamicMesh3::FEdge& Edge = Mesh.GetEdgeRef(EID);
+
+						if (Edge.Tri[1] != IndexConstants::InvalidID)
+						{
+							double Weight;
+							bool bNeedsSplit;
+							std::tie(Weight, bNeedsSplit) = ComputeWeight(EID);
+							
+							Heap.Update(HeapIndices[EID], { Weight, bNeedsSplit, EID });
+						}
+					}
+				}
+			}
+		}
+	}
+}
+
+} // namespace Geometry
+} // namespace UE
\ No newline at end of file
diff --git a/Engine/Source/Runtime/GeometryCore/Public/Tessellation/AdaptiveMetrics.h b/Engine/Source/Runtime/GeometryCore/Public/Tessellation/AdaptiveMetrics.h
new file mode 100644
--- /dev/null
+++ b/Engine/Source/Runtime/GeometryCore/Public/Tessellation/AdaptiveMetrics.h
@@ -0,0 +1,149 @@
+// Copyright Epic Games, Inc. All Rights Reserved.
+
+#pragma once
+
+#include "Tessellation/Affine.h"
+#include "Math/Vector.h"
+#include "Math/UnrealMathUtility.h"
+#include "Image/DisplacementMap.h"
+
+namespace UE {
+namespace Geometry {
+namespace TessellationUtil {
+
+/**
+ * Helper methods for use in adaptive tessellation callbacks in displacement policy for 
+ * TAdaptiveTesselator and TParallelRefinement.
+ */
+
+/**
+ * Get bounds of UVs in a UV subtriangle
+ * 
+ * @param Barycentrics  Barycentric coordinates of three vertices of a subtriangle
+ * @param VertexUVs     UVs of a triangle
+ * @param UVBounds      component-wise bounds of UV coordinates of subtriangle
+ */
+inline void CalculateUVBounds(
+	const FVector3f Barycentrics[3],
+	const FVector2f VertexUVs[3],
+	FVector2f OutUVBounds[2])
+{
+	OutUVBounds[0] = FVector2f( std::numeric_limits<float>::max(),  std::numeric_limits<float>::max());
+	OutUVBounds[1] = FVector2f(-std::numeric_limits<float>::max(), -std::numeric_limits<float>::max());
+							 
+	for (int k = 0; k < 3; k++)
+	{
+		FVector2f UV;
+		UV =  VertexUVs[0] * Barycentrics[k].X;
+		UV += VertexUVs[1] * Barycentrics[k].Y;
+		UV += VertexUVs[2] * Barycentrics[k].Z;
+
+		OutUVBounds[0] = FVector2f::Min(OutUVBounds[0], UV);
+		OutUVBounds[1] = FVector2f::Max(OutUVBounds[1], UV);
+	}
+}
+
+/**
+ * Calculate conservative bounds of the squared error of the displaced position.
+ *  
+ * VertexNormals and Displacements are given on the corners of a base-triangle.
+ * The subtriangle to evaluate on is defined by the Barycentrics.
+ * UVBounds are given with respect to subtriangle.
+ * 
+ * DisplacementBoundsRefinements controls the accuracy of the bounds when retrieving
+ * bounds from the displacement map (quality speed/tradeoff). See 
+ * FDisplacementMap::SampleHierarchical.
+ * 
+ * Magnitude, Center are the usual affine maps on FDisplacementMap values. Set to
+ * 1 and 0 if they are already set on the displacement map. They are only provided here,
+ * to allow modification without changing the displacement map.
+ */
+template <typename T>
+inline UE::Math::TVector2<T> CalculateErrorBounds(
+	const UE::Math::TVector<T> Barycentrics[3],
+	const UE::Math::TVector<T> VertexNormals[3],
+	const FVector2f UVBounds[2], 
+	const UE::Math::TVector<T> Displacements[3],
+	const FDisplacementMap& DisplacementMap,
+	const float Center,
+	const FVector2f MagnitudeBounds, // bounds 
+	const int32 DisplacementBoundsRefinements)
+{
+	T MinBarycentric0 = FMath::Min3(Barycentrics[0].X, Barycentrics[1].X, Barycentrics[2].X);
+	T MaxBarycentric0 = FMath::Max3(Barycentrics[0].X, Barycentrics[1].X, Barycentrics[2].X);
+
+	T MinBarycentric1 = FMath::Min3(Barycentrics[0].Y, Barycentrics[1].Y, Barycentrics[2].Y);
+	T MaxBarycentric1 = FMath::Max3(Barycentrics[0].Y, Barycentrics[1].Y, Barycentrics[2].Y);
+
+	TAffine<T, 2> Barycentric0(MinBarycentric0, MaxBarycentric0, 0);
+	TAffine<T, 2> Barycentric1(MinBarycentric1, MaxBarycentric1, 1);
+	TAffine<T, 2> Barycentric2 = TAffine<T, 2 >(1.0) - Barycentric0 - Barycentric1;
+
+	TAffine<UE::Math::TVector<T>, 2> LerpedDisplacement;
+	LerpedDisplacement  = TAffine<UE::Math::TVector<T>, 2>(Displacements[0]) * Barycentric0;
+	LerpedDisplacement += TAffine<UE::Math::TVector<T>, 2>(Displacements[1]) * Barycentric1;
+	LerpedDisplacement += TAffine<UE::Math::TVector<T>, 2>(Displacements[2]) * Barycentric2;
+
+	TAffine<UE::Math::TVector<T>, 2> Normal;
+	Normal  = TAffine<UE::Math::TVector<T>, 2>(VertexNormals[0]) * Barycentric0;
+	Normal += TAffine<UE::Math::TVector<T>, 2>(VertexNormals[1]) * Barycentric1;
+	Normal += TAffine<UE::Math::TVector<T>, 2>(VertexNormals[2]) * Barycentric2;
+	Normal = Normalize(Normal);
+
+	const FVector2f DisplacementBounds = DisplacementMap.SampleHierarchical(UVBounds[0], UVBounds[1], DisplacementBoundsRefinements) - FVector2f(Center);
+
+	// Displacement can be negative, but Magnitude is always positive.
+	const FVector2f ScaledDisplacementBounds( FMath::Min(MagnitudeBounds[0]*DisplacementBounds[0], MagnitudeBounds[1]*DisplacementBounds[0]),
+		                                      FMath::Max(MagnitudeBounds[0]*DisplacementBounds[1], MagnitudeBounds[1]*DisplacementBounds[1]));
+
+	TAffine<T, 2> Displacement( ScaledDisplacementBounds[0], ScaledDisplacementBounds[1] );
+	TAffine<T, 2> Error = (Normal * Displacement - LerpedDisplacement).SizeSquared();
+
+	return UE::Math::TVector2<T>(Error.GetMin(), Error.GetMax());
+}
+
+
+/**
+ * TAdaptiveTessellation samples the triangle at the finest level to determine
+ * the optimal split position in terms of maximizing the error reduction.
+ *  
+ * This method can be used to compute the number of samples for this sub-sampling step.
+ * 
+ * @param Barycentrics  Barycentric coordinates of three vertices of a subtriangle
+ * @param VertexUVs     UVs of a triangle
+ * @param MapRes        Resolution of FDisplacementMap used
+ */
+template <typename T>
+inline int32 GetNumSamples(
+	const UE::Math::TVector<T> Barycentrics[3],
+	const FVector2f VertexUVs[3],
+	const FVector2f MapRes) 
+{
+	FVector2f UVs[3];
+	for (int k = 0; k < 3; k++)
+	{
+		UVs[k] =  VertexUVs[0] * static_cast<float>(Barycentrics[k].X);
+		UVs[k] += VertexUVs[1] * static_cast<float>(Barycentrics[k].Y);
+		UVs[k] += VertexUVs[2] * static_cast<float>(Barycentrics[k].Z);
+
+		UVs[k].X *= MapRes[0];
+		UVs[k].Y *= MapRes[1];
+	}
+
+	FVector2f Edge01 = UVs[1] - UVs[0];
+	FVector2f Edge12 = UVs[2] - UVs[1];
+	FVector2f Edge20 = UVs[0] - UVs[2];
+
+	T MaxEdgeLength = FMath::Sqrt(FMath::Max3(
+		Edge01.SizeSquared(),
+		Edge12.SizeSquared(),
+		Edge20.SizeSquared()));
+
+	T AreaInTexels = FMath::Abs(0.5 * (Edge01 ^ Edge12));
+
+	return static_cast<int32>(FMath::CeilToInt(FMath::Max(MaxEdgeLength, AreaInTexels)));
+}
+
+} // namespace TessellationUtil
+} // namespace Geometry
+} // namespace UE
\ No newline at end of file
diff --git a/Engine/Source/Runtime/GeometryCore/Public/Tessellation/FeatureAdaptive.h b/Engine/Source/Runtime/GeometryCore/Public/Tessellation/FeatureAdaptive.h
new file mode 100644
--- /dev/null
+++ b/Engine/Source/Runtime/GeometryCore/Public/Tessellation/FeatureAdaptive.h
@@ -0,0 +1,70 @@
+// Copyright Epic Games, Inc. All Rights Reserved.
+
+#pragma once
+
+#include "Math/Vector.h"
+#include "Math/UnrealMathUtility.h"
+
+namespace UE {
+namespace Geometry {
+
+// [ Baensch et al. "A Finite Element Method for Surface Diffusion", 2004] Sec 5.4
+// Feature adaptive refinement criterion
+template <class T>
+class TFeatureSensitivityRefinementCheck
+{
+public:
+	
+	// @param InFeatureSensitivity   sensitivity of feature check, between 0 and 1
+	// @param InFeatureSize          the grid size parameter (h_T in the paper)
+	// @param V0, V1, V2             triangle vertices
+	TFeatureSensitivityRefinementCheck(const T InFeatureSensitivity, const T InFeatureSize,
+		const Math::TVector<T>& V0, 
+		const Math::TVector<T>& V1,
+		const Math::TVector<T>& V2)
+		: Threshold(FMath::Max(T(1.) / (InFeatureSensitivity*T(0.5) + T(0.5)) - T(0.95), T(0.05)))
+		, TriVertices{ V0, V1, V2 }
+		, TriNormal( (V2-V0) ^ (V1-V0) )
+		, LSqN(TriNormal.SquaredLength())
+		, IncircleRadius(CalculateInRadius(V0, V1, V2))
+		, FeatureSize(InFeatureSize)
+	{
+	}
+
+	[[nodiscard]] FORCEINLINE bool NeedsCheck() const
+	{
+		return (Threshold * FeatureSize < IncircleRadius * UE_PI); // Early-out (also if IncircleRadius is NaN)
+	}
+
+	// check whether the edge should be refined refined
+	// @param LocalEdge local edge index with respect to the triangle vertices ordering
+	// @param FarVertex vertex on the other side of that edge
+	//  
+	[[nodiscard]] FORCEINLINE bool CheckCriterion(const int32 LocalEdge, const UE::Math::TVector<T>& FarVertex ) const 
+	{
+		const Math::TVector<T> NT = (TriVertices[(LocalEdge + 1) % 3] - TriVertices[LocalEdge]) ^ (FarVertex - TriVertices[LocalEdge]);
+		const T LSqNT = NT.SquaredLength();
+
+		return ((TriNormal | NT) < FMath::Sqrt(LSqN * LSqNT) * FMath::Cos(Threshold * FeatureSize / IncircleRadius));
+	}
+
+private:
+
+	static T CalculateInRadius( const UE::Math::TVector<T>& V0, const UE::Math::TVector<T>& V1, const UE::Math::TVector<T>& V2 )
+	{
+		const T EdgeLengths_D[] = { (V1 - V0).Size(), (V2 - V1).Size(), (V0 - V2).Size() };
+		const T SemiPerimeter = T(0.5) * (EdgeLengths_D[0] + EdgeLengths_D[1] + EdgeLengths_D[2]);
+		const T TriArea = FMath::Sqrt(FMath::Max(T(0.), SemiPerimeter * (SemiPerimeter - EdgeLengths_D[0]) * (SemiPerimeter - EdgeLengths_D[1]) * (SemiPerimeter - EdgeLengths_D[2]))); // Heron
+		return TriArea / SemiPerimeter;
+	}
+
+    const T                Threshold;
+	const Math::TVector<T> TriVertices[3];
+	const Math::TVector<T> TriNormal;
+	const T                LSqN;            // SquaredLength of Normal vector
+	const T                IncircleRadius;
+	const T                FeatureSize;
+};
+
+} // namespace Geometry
+} // namespace UE
\ No newline at end of file
diff --git a/Engine/Source/Runtime/GeometryCore/Public/Tessellation/Regularization.h b/Engine/Source/Runtime/GeometryCore/Public/Tessellation/Regularization.h
new file mode 100644
--- /dev/null
+++ b/Engine/Source/Runtime/GeometryCore/Public/Tessellation/Regularization.h
@@ -0,0 +1,70 @@
+// Copyright Epic Games, Inc. All Rights Reserved.
+
+#pragma once
+
+#include "Math/MathFwd.h"
+#include "Math/Vector.h"
+#include "Templates/Function.h"
+#include "Misc/Optional.h"
+#include "Containers/BitArray.h"
+#include "Containers/Array.h"
+
+namespace UE {
+namespace Geometry {
+
+class FDynamicMesh3;
+	
+using FDisplaceFunc = TFunction<FVector3d (const int32 VertexIndex)>;
+
+/** 
+ * Apply one iteration of mesh smoothing. If bPreserveVolume is set the new vertex positions will be adjusted to preserve
+ * the enclosing mesh volume, favoring only tangential shifts.
+ * 
+ * If ActiveVertices is provided, only vertices for which the bit is set will be processed.
+ */
+GEOMETRYCORE_API void 
+RegularizeVolumePreserving(FDynamicMesh3& Mesh, const bool bPreserveVolume, const TOptional<TBitArray<>>& ActiveVertices = {});
+
+/**
+ * Apply edge flips on non-Delaunay edges in order of importance.
+ * Intrinsic edge-flips would split the flipped edge and produce edges that are on the same surface.
+ * This function relaxes this and applies honest flips when the surface is planar enough (when the
+ * change in the oriented dihedral angle is smaller than FlipAngleThreshold)
+ * 
+ * See [Fisher, Springborn, Bobenko, Schroeder, "An Algorithm for the Construction of Intrinsic Delaunay Triangulations
+ * with Applications to Digital Geometry Processing", 2006]
+ *
+ *
+ * @param Mesh mesh to be modified
+ * @param Displacements array of per-vertex displacements (optional)
+ * @param DisplacementFunctor if Displacements are passed in, this functor needs to be provided to construct new displacements at split vertices
+ * @param FlipAngleThreshold controls honest vs intrinsic flips 
+ */
+GEOMETRYCORE_API void 
+IntrinsicDelaunayFlipEdge(FDynamicMesh3& Mesh, TArray<FVector3d>& Displacements, FDisplaceFunc DisplacementFunctor, const double FlipAngleThreshold = 0.1 * UE_PI);
+
+inline void IntrinsicDelaunayFlipEdge(FDynamicMesh3& Mesh, const double FlipAngleThreshold = 0.1 * UE_PI)
+{
+	TArray<FVector3d> NullDisplacements;
+	IntrinsicDelaunayFlipEdge(Mesh, NullDisplacements, FDisplaceFunc(), FlipAngleThreshold);
+}
+
+/**
+ * Walk through edges in priority order, splitting edges according to threshold criterion
+ *
+ * @param Mesh mesh to be modified
+ * @param Displacements array of per-vertex displacements (optional)
+ * @param DisplacementFunctor if Displacements are passed in, this functor needs to be provided to construct new displacements at split vertices
+ * @param SplitThreshold trigger split if ratio of longest edge/(sum of length of other edges) exceeds this threshold
+ */
+GEOMETRYCORE_API void 
+SplitLongEdges(FDynamicMesh3& Mesh, TArray<FVector3d>& Displacements, FDisplaceFunc DisplacementFunctor, const double SplitThreshold = 0.9);
+
+inline void SplitLongEdges(FDynamicMesh3& Mesh, const double SplitThreshold = 0.9)
+{
+	TArray<FVector3d> NullDisplacements;
+	SplitLongEdges(Mesh, NullDisplacements, FDisplaceFunc(), SplitThreshold);
+}
+	
+} // namespace Geometry
+} // namespace UE
