GeometryCore: Fix DisplacementMap build issue

#rnx

diff --git a/Engine/Source/Runtime/GeometryCore/GeometryCore.Build.cs b/Engine/Source/Runtime/GeometryCore/GeometryCore.Build.cs
--- a/Engine/Source/Runtime/GeometryCore/GeometryCore.Build.cs
+++ b/Engine/Source/Runtime/GeometryCore/GeometryCore.Build.cs
@@ -6,27 +6,25 @@ public class GeometryCore : ModuleRules
 {
 	public GeometryCore(ReadOnlyTargetRules Target) : base(Target)
 	{
-        PCHUsage = ModuleRules.PCHUsageMode.UseExplicitOrSharedPCHs;
+        PCHUsage = PCHUsageMode.UseExplicitOrSharedPCHs;
 
 		PublicIncludePathModuleNames.AddRange(
-			new string[] {
+			[
 				"AnimationCore",			// For the BoneWeights.h include
 				"Engine",					// For GPUSkinPublicDefs.h
-				"ImageCore"
-			}
+			]
 		);
 
 		PublicDependencyModuleNames.AddRange(
-			new string[] {
+			[
 				"Core"
-			}
+			]
 		);
 
 		PrivateDependencyModuleNames.AddRange(
-			new string[]
-			{
+			[
 				"AutoRTFM",
-			}
+			]
 		);
 
 		CppCompileWarningSettings.UnsafeTypeCastWarningLevel = WarningLevel.Error;
diff --git a/Engine/Source/Runtime/GeometryCore/Private/Image/DisplacementMap.cpp b/Engine/Source/Runtime/GeometryCore/Private/Image/DisplacementMap.cpp
--- a/Engine/Source/Runtime/GeometryCore/Private/Image/DisplacementMap.cpp
+++ b/Engine/Source/Runtime/GeometryCore/Private/Image/DisplacementMap.cpp
@@ -2,14 +2,8 @@
 
 #include "Image/DisplacementMap.h"
 
-#include "ImageCore.h"
-#include "ImageCoreUtils.h"
 #include "MathUtil.h"
 
-#include "Containers/Array.h"
-#include "Containers/ContainerAllocationPolicies.h"
-#include "Serialization/Archive.h"
-
 namespace UE
 {
 namespace Geometry
diff --git a/Engine/Source/Runtime/GeometryCore/Public/Image/DisplacementMap.h b/Engine/Source/Runtime/GeometryCore/Public/Image/DisplacementMap.h
--- a/Engine/Source/Runtime/GeometryCore/Public/Image/DisplacementMap.h
+++ b/Engine/Source/Runtime/GeometryCore/Public/Image/DisplacementMap.h
@@ -2,10 +2,7 @@
 
 #pragma once
 
-#include "Engine/Texture.h"
-#include "Serialization/Archive.h"
-
-class FArchive;
+#include "Engine/TextureDefines.h"
 
 namespace UE {
 namespace Geometry {
