fix use of the non-normalized line dir in quadric regularization

diff --git a/Engine/Source/Runtime/GeometryCore/Public/QuadricError.h b/Engine/Source/Runtime/GeometryCore/Public/QuadricError.h
--- a/Engine/Source/Runtime/GeometryCore/Public/QuadricError.h
+++ b/Engine/Source/Runtime/GeometryCore/Public/QuadricError.h
@@ -265,7 +265,7 @@ struct TQuadricError
 			return;
 		}
 		TVector<RealType> N1, N2;
-		VectorUtil::MakePerpVectors(LineDir, N1, N2);
+		VectorUtil::MakePerpVectors(NormalizedLineDir, N1, N2);
 		Add(Weight * .5, TQuadricError<RealType>(N1, Pos));
 		Add(Weight * .5, TQuadricError<RealType>(N2, Pos));
 	}
