Add an approximate contact area mesh query to the geometrycore meshqueries methods
+ minor bug fixes to other mesh queries (supporting non-compact meshes, and meshes w/ fvector3f vertex storage)

diff --git a/Engine/Source/Runtime/GeometryCore/Public/MeshQueries.h b/Engine/Source/Runtime/GeometryCore/Public/MeshQueries.h
--- a/Engine/Source/Runtime/GeometryCore/Public/MeshQueries.h
+++ b/Engine/Source/Runtime/GeometryCore/Public/MeshQueries.h
@@ -249,11 +249,11 @@ public:
 	static FAxisAlignedBox3d GetTriBounds(const TriangleMeshType& Mesh, int TID)
 	{
 		FIndex3i TriInds = Mesh.GetTriangle(TID);
-		FVector3d MinV, MaxV, V = Mesh.GetVertex(TriInds.A);
+		FVector3d MinV, MaxV, V = (FVector3d)Mesh.GetVertex(TriInds.A);
 		MinV = MaxV = V;
 		for (int i = 1; i < 3; ++i)
 		{
-			V = Mesh.GetVertex(TriInds[i]);
+			V = (FVector3d)Mesh.GetVertex(TriInds[i]);
 			if (V.X < MinV.X)				MinV.X = V.X;
 			else if (V.X > MaxV.X)			MaxV.X = V.X;
 			if (V.Y < MinV.Y)				MinV.Y = V.Y;
@@ -631,9 +631,9 @@ public:
 	static void VertexToSurfaceDistances(const TriangleMeshType& MeshA, const MeshSpatialType& SpatialB, TArray<double>& Distances)
 	{
 		check(SpatialB.SupportsNearestTriangle());
-		Distances.SetNumZeroed(MeshA.VertexCount());
+		Distances.SetNumZeroed(MeshA.MaxVertexID());
 
-		ParallelFor(MeshA.VertexCount(), [&MeshA, &SpatialB, &Distances](int VertexID)
+		ParallelFor(MeshA.MaxVertexID(), [&MeshA, &SpatialB, &Distances](int VertexID)
 		{
 			if (!MeshA.IsVertex(VertexID)) { return; }
 
@@ -672,15 +672,78 @@ public:
 		return FMath::Max(HausdorffDistance(MeshA, SpatialB), HausdorffDistance(MeshB, SpatialA));
 	}
 
+	/// Approximately compute the surface area of MeshA that is within the threshold distance of a target mesh
+	/// 
+	/// @param MeshA The mesh whose contact area should be computed.
+	/// @param SpatialB The target surface's acceleration structure.
+	/// @param ThresholdDistance The distance below which a point on MeshA is considered to be in contact with the target mesh
+	/// @param SamplesPerArea Requested minimum samples per unit of area. Note all triangles will at least be sampled once regardless.
+	/// @param MaxSamplesPerTri Optional maximum samples per triangle, if greater than 0.
+	template<typename MeshSpatialType, typename MeshRealType = double>
+	static double ApproximateContactArea(const TriangleMeshType& MeshA, const TTransform<MeshRealType>& TransformA, 
+		const MeshSpatialType& SpatialB, const TTransform<MeshRealType>& TransformB, double ThresholdDistance, double SamplesPerArea, int32 MaxSamplesPerTri = -1)
+	{
+		TArray<float> ContactAreas;
+		ContactAreas.SetNumZeroed(MeshA.MaxTriangleID());
+
+		ParallelFor(MeshA.MaxTriangleID(),
+			[&MeshA, &TransformA, &SpatialB, &TransformB, &ContactAreas, ThresholdDistanceSq = ThresholdDistance*ThresholdDistance, SamplesPerArea, MaxSamplesPerTri](int TriangleID)
+		{
+			if (!MeshA.IsTriangle(TriangleID))
+			{
+				return;
+			}
+
+			TVector<MeshRealType> A, B, C;
+			MeshA.GetTriVertices(TriangleID, A, B, C);
+			A = TransformA.TransformPosition(A);
+			B = TransformA.TransformPosition(B);
+			C = TransformA.TransformPosition(C);
+			float TriangleArea = (float)VectorUtil::Area(A, B, C);
+			if (TriangleArea <= 0) // no contact area from degenerate triangles
+			{
+				return;
+			}
+			FRandomStream RandStream(TriangleID);
+			int32 NumSamples = FMath::Max(1, FMath::CeilToInt32(SamplesPerArea * TriangleArea));
+			if (MaxSamplesPerTri > 0 && NumSamples > MaxSamplesPerTri)
+			{
+				NumSamples = MaxSamplesPerTri;
+			}
+			int32 ContactSamples = 0;
+			for (int32 SampleIdx = 0; SampleIdx < NumSamples; ++SampleIdx)
+			{
+				TVector<MeshRealType> P = VectorUtil::UniformSampleTrianglePoint<MeshRealType>((MeshRealType)RandStream.GetFraction(), (MeshRealType)RandStream.GetFraction(), A, B, C);
+				int32 TriB_Unused;
+				ContactSamples += (int32)SpatialB.IsWithinDistanceSquared((FVector3d)TransformB.InverseTransformPosition(P), ThresholdDistanceSq, TriB_Unused);
+			}
+			ContactAreas[TriangleID] = TriangleArea * ((float)ContactSamples / (float)NumSamples);
+		});
+
+		double TotalArea = 0;
+		for (float Area : ContactAreas)
+		{
+			TotalArea += (double)Area;
+		}
+		return TotalArea;
+	}
+
+	/// Version of ApproximateContactArea assuming MeshA and SpatialB are in a common coordinate space
+	template<typename MeshSpatialType, typename MeshRealType = double>
+	static double ApproximateContactArea(const TriangleMeshType& MeshA, const MeshSpatialType& SpatialB, double ThresholdDistance, double SamplesPerArea)
+	{
+		return ApproximateContactArea(MeshA, FTransform3d::Identity, SpatialB, FTransform3d::Identity, ThresholdDistance, SamplesPerArea);
+	}
+
 
 	/// Compute all vertex-to-surface distances in serial. Should only be used for debugging the parallel version above!
 	template<typename MeshSpatialType>
 	static void VertexToSurfaceDistancesSerial(const TriangleMeshType& MeshA, const MeshSpatialType& SpatialB, TArray<double>& Distances)
 	{
 		check(SpatialB.SupportsNearestTriangle());
-		Distances.SetNumZeroed(MeshA.VertexCount());
+		Distances.SetNumZeroed(MeshA.MaxVertexID());
 
-		for (int VertexID = 0; VertexID < MeshA.VertexCount(); ++VertexID)
+		for (int VertexID = 0; VertexID < MeshA.MaxVertexID(); ++VertexID)
 		{
 			if (!MeshA.IsVertex(VertexID)) { continue; }
 
