add a few casts to the aabb tree to let it work with meshes that use fvector3f

diff --git a/Engine/Source/Runtime/GeometryCore/Public/Spatial/MeshAABBTree3.h b/Engine/Source/Runtime/GeometryCore/Public/Spatial/MeshAABBTree3.h
--- a/Engine/Source/Runtime/GeometryCore/Public/Spatial/MeshAABBTree3.h
+++ b/Engine/Source/Runtime/GeometryCore/Public/Spatial/MeshAABBTree3.h
@@ -369,7 +369,7 @@ protected:
 				FIndex3i Triangle = Mesh->GetTriangle(ti);
 				for (int j = 0; j < 3; ++j) 
 				{
-					double VertexDistSqr = DistanceSquared(Mesh->GetVertex(Triangle[j]), P);
+					double VertexDistSqr = DistanceSquared((FVector3d)Mesh->GetVertex(Triangle[j]), P);
 					if (VertexDistSqr < NearestDistSqr) 
 					{
 						NearestDistSqr = VertexDistSqr;
@@ -977,9 +977,9 @@ public:
 			if (TransformF != nullptr)
 			{
 				FIndex3i Tri = TestMesh->GetTriangle(TID);
-				TestTri.V[0] = TransformF(TestMesh->GetVertex(Tri.A));
-				TestTri.V[1] = TransformF(TestMesh->GetVertex(Tri.B));
-				TestTri.V[2] = TransformF(TestMesh->GetVertex(Tri.C));
+				TestTri.V[0] = TransformF((FVector3d)TestMesh->GetVertex(Tri.A));
+				TestTri.V[1] = TransformF((FVector3d)TestMesh->GetVertex(Tri.B));
+				TestTri.V[2] = TransformF((FVector3d)TestMesh->GetVertex(Tri.C));
 			}
 			else
 			{
