Made the compute polygon plane function return the area as well.

#rnx

diff --git a/Engine/Source/Runtime/GeometryCore/Private/CompGeom/PolygonTriangulation.cpp b/Engine/Source/Runtime/GeometryCore/Private/CompGeom/PolygonTriangulation.cpp
--- a/Engine/Source/Runtime/GeometryCore/Private/CompGeom/PolygonTriangulation.cpp
+++ b/Engine/Source/Runtime/GeometryCore/Private/CompGeom/PolygonTriangulation.cpp
@@ -11,8 +11,8 @@ namespace PolygonTriangulation
 	template GEOMETRYCORE_API void TriangulateSimplePolygon<float>(const TArray<TVector2<float>>& VertexPositions, TArray<FIndex3i>& OutTriangles, bool bOrientAsHoleFill);
 	template GEOMETRYCORE_API void TriangulateSimplePolygon<double>(const TArray<TVector2<double>>& VertexPositions, TArray<FIndex3i>& OutTriangles, bool bOrientAsHoleFill);
 
-	template GEOMETRYCORE_API void ComputePolygonPlane<float>(const TArray<TVector<float>>& VertexPositions, TVector<float>& NormalOut, TVector<float>& PlanePointOut);
-	template GEOMETRYCORE_API void ComputePolygonPlane<double>(const TArray<TVector<double>>& VertexPositions, TVector<double>& NormalOut, TVector<double>& PlanePointOut);
+	template GEOMETRYCORE_API float ComputePolygonPlane<float>(const TArray<TVector<float>>& VertexPositions, TVector<float>& NormalOut, TVector<float>& PlanePointOut);
+	template GEOMETRYCORE_API double ComputePolygonPlane<double>(const TArray<TVector<double>>& VertexPositions, TVector<double>& NormalOut, TVector<double>& PlanePointOut);
 
 	template GEOMETRYCORE_API void TriangulateSimplePolygon<float>(const TArray<TVector<float>>& VertexPositions, TArray<FIndex3i>& OutTriangles, bool bOrientAsHoleFill);
 	template GEOMETRYCORE_API void TriangulateSimplePolygon<double>(const TArray<TVector<double>>& VertexPositions, TArray<FIndex3i>& OutTriangles, bool bOrientAsHoleFill);
@@ -172,7 +172,7 @@ void PolygonTriangulation::TriangulateSimplePolygon(const TArray<TVector2<T>>& V
 
 
 template<typename T>
-void PolygonTriangulation::ComputePolygonPlane(const TArray<TVector<T>>& VertexPositions, TVector<T>& PlaneNormalOut, TVector<T>& PlanePointOut)
+T PolygonTriangulation::ComputePolygonPlane(const TArray<TVector<T>>& VertexPositions, TVector<T>& PlaneNormalOut, TVector<T>& PlanePointOut)
 {
 	// NOTE: This polygon plane computation code is partially based on the implementation of "Newell's method" from Real-Time
 	//       Collision Detection by Christer Ericson, published by Morgan Kaufmann Publishers, (c) 2005 Elsevier Inc
@@ -191,8 +191,9 @@ void PolygonTriangulation::ComputePolygonPlane(const TArray<TVector<T>>& VertexP
 		PlaneNormalOut.Y += (PositionJ.Z - PositionI.Z) * (PositionI.X + PositionJ.X);
 		PlaneNormalOut.Z += (PositionJ.X - PositionI.X) * (PositionI.Y + PositionJ.Y);
 	}
-	Normalize(PlaneNormalOut);
+
 	PlanePointOut /= (T)NumVertices;
+	return static_cast<T>(0.5) * Normalize(PlaneNormalOut);
 }
 
 
diff --git a/Engine/Source/Runtime/GeometryCore/Public/CompGeom/PolygonTriangulation.h b/Engine/Source/Runtime/GeometryCore/Public/CompGeom/PolygonTriangulation.h
--- a/Engine/Source/Runtime/GeometryCore/Public/CompGeom/PolygonTriangulation.h
+++ b/Engine/Source/Runtime/GeometryCore/Public/CompGeom/PolygonTriangulation.h
@@ -26,9 +26,15 @@ namespace PolygonTriangulation
 	template<typename RealType>
 	void GEOMETRYCORE_API TriangulateSimplePolygon(const TArray<TVector2<RealType>>& VertexPositions, TArray<FIndex3i>& OutTriangles, bool bOrientAsHoleFill = true);
 
-
+	/**
+	* Computes a plane centroid and plane normal from an input 3D polygon.
+	* @param VertexPositions ordered vertices of 3D polygon.
+	* @param PlaneNormalOut normal of the computed plane. Will be normalized.
+	* @param PlanePointOut centroid of the computed plane.
+	* @return Area of the polygon.
+	*/
 	template<typename RealType>
-	void GEOMETRYCORE_API ComputePolygonPlane(const TArray<TVector<RealType>>& VertexPositions, TVector<RealType>& PlaneNormalOut, TVector<RealType>& PlanePointOut );
+	RealType GEOMETRYCORE_API ComputePolygonPlane(const TArray<TVector<RealType>>& VertexPositions, TVector<RealType>& PlaneNormalOut, TVector<RealType>& PlanePointOut );
 
 
 	/**
