TangentBezierSpline: Fix ensure when deleting final spline point

This regressed due to changes in the data we store for one point splines. The original change missed one invalidated assumption in an edge case when removing the final point which left bad data around.

diff --git a/Engine/Source/Runtime/GeometryCore/Public/Splines/TangentBezierSpline.h b/Engine/Source/Runtime/GeometryCore/Public/Splines/TangentBezierSpline.h
--- a/Engine/Source/Runtime/GeometryCore/Public/Splines/TangentBezierSpline.h
+++ b/Engine/Source/Runtime/GeometryCore/Public/Splines/TangentBezierSpline.h
@@ -1040,8 +1040,7 @@ public:
 		// special case for handling the last point and second last point
 		if (NumPoints == 1)
 		{
-			// Removing the only point
-			InternalSpline.RemoveValue(0);
+			InternalSpline.Clear();
 			TangentModes.Reset();
 			return;
 		}
