TangentBezierSpline: Fix evaluation of 1 point spline

Even using Constant out of bounds handling, evaluating an attribute channel on a spline with 1 value was not correctly returning that attribute's value. This is due to a regression associated with a recent change in how we store 1 point splines. We now detect and handle this case correctly again.

diff --git a/Engine/Source/Runtime/GeometryCore/Public/Splines/TangentBezierSpline.h b/Engine/Source/Runtime/GeometryCore/Public/Splines/TangentBezierSpline.h
--- a/Engine/Source/Runtime/GeometryCore/Public/Splines/TangentBezierSpline.h
+++ b/Engine/Source/Runtime/GeometryCore/Public/Splines/TangentBezierSpline.h
@@ -1178,12 +1178,12 @@ public:
     /** Gets the value of the specified point */
     ValueType GetValue(int32 Index) const
     {
-    	const int32 NumKeys = InternalSpline.NumKeys();
+    	const int32 NumPoints = GetNumPoints();
 
-    	if (NumKeys == 0 || Index < 0 || Index >= GetNumPoints())
+    	if (NumPoints == 0 || Index < 0 || Index >= GetNumPoints())
     		return ValueType();
     
-    	if (NumKeys == 1)
+    	if (NumPoints == 1)
     	{
     		// Special case when we only have one point
     		if (Index == 0)
@@ -1202,7 +1202,7 @@ public:
     	{
     		// Other points are P3 of previous segments
     		const int32 ControlPointIndex = (Index - 1) * 4 + 3;
-    		if (ControlPointIndex < NumKeys)
+    		if (ControlPointIndex < InternalSpline.NumKeys())
     			return InternalSpline.GetValue(ControlPointIndex);
     	}
 
@@ -1584,6 +1584,16 @@ public:
 
 	}
 
+	virtual ValueType EvaluateImpl(float Parameter) const override
+	{
+		if (GetNumPoints() == 1)
+		{
+			return GetValue(0);
+		}
+		
+		return InternalSpline.EvaluateImpl(Parameter);
+	}
+	
 	virtual float FindNearest(const ValueType& Point, float& OutSquaredDistance) const override
 	{
         // Delegate to internal spline
