NavHeightMap: implement tile indexer range validation

diff --git a/Engine/Source/Runtime/GeometryCore/Public/Util/GridIndexing2.h b/Engine/Source/Runtime/GeometryCore/Public/Util/GridIndexing2.h
--- a/Engine/Source/Runtime/GeometryCore/Public/Util/GridIndexing2.h
+++ b/Engine/Source/Runtime/GeometryCore/Public/Util/GridIndexing2.h
@@ -94,6 +94,14 @@ struct TShiftGridIndexer2
 		return FIntRect(Min.X, Min.Y, Max.X, Max.Y);
 	}
 
+	/** Convert real-valued point to real-valued grid coordinates */
+	inline TVector2<RealType> ToRealGrid(const TVector2<RealType>& Point) const
+	{
+		return TVector2<RealType>(
+			TMathUtil<RealType>::Floor((Point.X - Origin.X) / CellSize),
+			TMathUtil<RealType>::Floor((Point.Y - Origin.Y) / CellSize));
+	}
+
 	/** Convert integer grid coordinates to real-valued point */
 	inline TVector2<RealType> FromGrid(const FVector2i& GridPoint) const
 	{
@@ -117,6 +125,16 @@ struct TShiftGridIndexer2
 		const TVector2<RealType> Max = Min + TVector2<RealType>(CellSize);
 		return TBox2<RealType>(Min, Max);
 	}
+
+	/** Indicate if the real-valued point is within the valid range of the indexer. If it returns false, it means that the conversion to grid coordinates would cause an integer overflow. */
+	bool IsInValidRange(const TVector2<RealType>& Point) const
+	{
+		const TVector2<RealType> RealGridPoint = ToRealGrid(Point);
+		return (RealGridPoint.X >= static_cast<RealType>(std::numeric_limits<int32>::min())) &&
+			(RealGridPoint.X <= static_cast<RealType>(std::numeric_limits<int32>::max())) &&
+			(RealGridPoint.Y >= static_cast<RealType>(std::numeric_limits<int32>::min())) &&
+			(RealGridPoint.Y <= static_cast<RealType>(std::numeric_limits<int32>::max()));
+	}
 };
 typedef TShiftGridIndexer2<float> FShiftGridIndexer2f;
 typedef TShiftGridIndexer2<double> FShiftGridIndexer2d;
