Deprecated TIsMemberPointer.

diff --git a/Engine/Plugins/Experimental/StateGraph/Source/StateGraph/Public/StateGraph.h b/Engine/Plugins/Experimental/StateGraph/Source/StateGraph/Public/StateGraph.h
--- a/Engine/Plugins/Experimental/StateGraph/Source/StateGraph/Public/StateGraph.h
+++ b/Engine/Plugins/Experimental/StateGraph/Source/StateGraph/Public/StateGraph.h
@@ -24,10 +24,11 @@
 #include "StateGraphFwd.h"
 #include "Templates/IsClass.h"
 #include "Templates/IsInvocable.h"
-#include "Templates/IsMemberPointer.h"
 #include "Templates/UnrealTemplate.h"
 #include "UObject/NameTypes.h"
 
+#include <type_traits>
+
 #define UE_API STATEGRAPH_API
 
 STATEGRAPH_API DECLARE_LOG_CATEGORY_EXTERN(LogStateGraph, Log, All);
@@ -344,7 +345,7 @@ public:
 		{
 			return CreateNode<FStateGraphNodeFunction>(NodeName, FStateGraphNodeFunctionStart::CreateUObject(Object, Function, Forward<ArgsTypes>(Args)...));
 		}
-		else if constexpr (TIsMemberPointer<FunctionType>::Value)
+		else if constexpr (std::is_member_function_pointer_v<FunctionType>)
 		{
 			return CreateNode<FStateGraphNodeFunction>(NodeName, FStateGraphNodeFunctionStart::CreateRaw(Object, Function, Forward<ArgsTypes>(Args)...));
 		}
@@ -462,3 +463,7 @@ FStateGraphNodeFunctionPtr FStateGraphNode::Next(ArgsTypes&&... Args) const
 } // UE
 
 #undef UE_API
+
+#if UE_ENABLE_INCLUDE_ORDER_DEPRECATED_IN_5_8
+#include "Templates/IsMemberPointer.h"
+#endif
diff --git a/Engine/Source/Runtime/Core/Public/Templates/IsMemberPointer.h b/Engine/Source/Runtime/Core/Public/Templates/IsMemberPointer.h
--- a/Engine/Source/Runtime/Core/Public/Templates/IsMemberPointer.h
+++ b/Engine/Source/Runtime/Core/Public/Templates/IsMemberPointer.h
@@ -5,17 +5,19 @@
 #include "CoreTypes.h"
 #include "Templates/IsPointer.h"
 
+UE_DEPRECATED_HEADER(5.8, "Templates/IsMemberPointer.h has been deprecated, please use <type_traits> and std::is_member_pointer instead.")
+
 /**
  * Traits class which tests if a type is a pointer to member (data member or member function).
  */
 template <typename T>
-struct TIsMemberPointer
+struct UE_DEPRECATED(5.8, "TIsMemberPointer has been deprecated, please use std::is_member_pointer instead.") TIsMemberPointer
 {
 	enum { Value = false };
 };
 
-template <typename T, typename U> struct TIsMemberPointer<T U::*> { enum { Value = true }; };
+template <typename T, typename U> struct UE_DEPRECATED(5.8, "TIsMemberPointer has been deprecated, please use std::is_member_pointer instead.") TIsMemberPointer<T U::*> { enum { Value = true }; };
 
-template <typename T> struct TIsMemberPointer<const          T> { enum { Value = TIsPointer<T>::Value }; };
-template <typename T> struct TIsMemberPointer<      volatile T> { enum { Value = TIsPointer<T>::Value }; };
-template <typename T> struct TIsMemberPointer<const volatile T> { enum { Value = TIsPointer<T>::Value }; };
+template <typename T> struct UE_DEPRECATED(5.8, "TIsMemberPointer has been deprecated, please use std::is_member_pointer instead.") TIsMemberPointer<const          T> { enum { Value = TIsPointer<T>::Value }; };
+template <typename T> struct UE_DEPRECATED(5.8, "TIsMemberPointer has been deprecated, please use std::is_member_pointer instead.") TIsMemberPointer<      volatile T> { enum { Value = TIsPointer<T>::Value }; };
+template <typename T> struct UE_DEPRECATED(5.8, "TIsMemberPointer has been deprecated, please use std::is_member_pointer instead.") TIsMemberPointer<const volatile T> { enum { Value = TIsPointer<T>::Value }; };
diff --git a/Engine/Source/Runtime/GeometryCore/Public/Tessellation/AdaptiveTessellator.h b/Engine/Source/Runtime/GeometryCore/Public/Tessellation/AdaptiveTessellator.h
--- a/Engine/Source/Runtime/GeometryCore/Public/Tessellation/AdaptiveTessellator.h
+++ b/Engine/Source/Runtime/GeometryCore/Public/Tessellation/AdaptiveTessellator.h
@@ -9,7 +9,6 @@
 #include "Async/ParallelFor.h"
 #include "Templates/Function.h"
 #include "Templates/Tuple.h"
-#include "Templates/IsMemberPointer.h"
 #include "IndexTypes.h"
 #include "ProfilingDebugging/CpuProfilerTrace.h"
 
@@ -906,3 +905,7 @@ inline void TAdaptiveTessellator<MeshT, DisplacementPolicyT>::TryDelaunayFlip( i
 
 } // namespace Geometry
 } // namespace UE
+
+#if UE_ENABLE_INCLUDE_ORDER_DEPRECATED_IN_5_8
+#include "Templates/IsMemberPointer.h"
+#endif
diff --git a/Engine/Source/Runtime/GeometryCore/Public/Tessellation/Affine.h b/Engine/Source/Runtime/GeometryCore/Public/Tessellation/Affine.h
--- a/Engine/Source/Runtime/GeometryCore/Public/Tessellation/Affine.h
+++ b/Engine/Source/Runtime/GeometryCore/Public/Tessellation/Affine.h
@@ -3,12 +3,9 @@
 #pragma once
 
 #include "CoreMinimal.h"
+#include "Concepts/FloatingPoint.h"
 #include "Math/Vector.h"
 
-#include "Templates/EnableIf.h"
-#include "Templates/IsFloatingPoint.h"
-#include "Templates/IsMemberPointer.h"
-
 namespace UE {
 namespace Geometry {
 
@@ -116,9 +113,9 @@ struct TAffine
 
 	// Smaller than (v|v)
 	// require that T has SizeSquared and declares floating point type FReal 
-	template <typename Q = T> FORCEINLINE
-	typename TEnableIf<TIsFloatingPoint<typename Q::FReal>::Value && TIsMemberPointer<decltype(&Q::SizeSquared)>::Value, TAffine<typename Q::FReal, Num>>::type
-	SizeSquared() const
+	template <typename Q = T>
+		requires (UE::CFloatingPoint<typename Q::FReal> && std::is_member_function_pointer_v<decltype(&Q::SizeSquared)>)
+	FORCEINLINE TAffine<typename Q::FReal, Num> SizeSquared() const
 	{
 		using FReal = typename Q::FReal;
 
@@ -184,9 +181,8 @@ FORCEINLINE TAffine< float, Num > operator|( const TAffine< T, Num >& A, const T
 	return Result;
 }
 
-template <typename T, uint32 Num> FORCEINLINE
-typename TEnableIf<TIsFloatingPoint<T>::Value, TAffine<T, Num>>::type
-Clamp( const TAffine< T, Num >& x, T Min, T Max )
+template <UE::CFloatingPoint T, uint32 Num>
+FORCEINLINE TAffine<T, Num> Clamp( const TAffine< T, Num >& x, T Min, T Max )
 {
 	// Using Chebyshev approximation
 	T xMin = x.GetMin();
@@ -214,9 +210,8 @@ Clamp( const TAffine< T, Num >& x, T Min, T Max )
 	return Result;
 }
 
-template <typename T, uint32 Num> FORCEINLINE
-typename TEnableIf<TIsFloatingPoint<T>::Value, TAffine<T, Num>>::type
-InvSqrt( const TAffine< T, Num >& x )
+template <UE::CFloatingPoint T, uint32 Num>
+FORCEINLINE TAffine<T, Num> InvSqrt( const TAffine< T, Num >& x )
 {
 	// Using min range approximation
 	T xMin = FMath::Max( T(1.e-4), x.GetMin() );
@@ -246,4 +241,10 @@ FORCEINLINE TAffine< T, Num > Normalize( const TAffine< T, Num >& x )
 }
 
 } // namespace Geometry
-} // namespace UE
\ No newline at end of file
+} // namespace UE
+
+#if UE_ENABLE_INCLUDE_ORDER_DEPRECATED_IN_5_8
+#include "Templates/EnableIf.h"
+#include "Templates/IsFloatingPoint.h"
+#include "Templates/IsMemberPointer.h"
+#endif
