Fix `FSplineTypeRegistry` to be transactionally safe.

#rnx

diff --git a/Engine/Source/Runtime/GeometryCore/Public/Splines/SplineTypeRegistry.h b/Engine/Source/Runtime/GeometryCore/Public/Splines/SplineTypeRegistry.h
--- a/Engine/Source/Runtime/GeometryCore/Public/Splines/SplineTypeRegistry.h
+++ b/Engine/Source/Runtime/GeometryCore/Public/Splines/SplineTypeRegistry.h
@@ -5,6 +5,7 @@
 #include "CoreMinimal.h"
 #include "SplineTypeId.h"
 #include "Containers/LockFreeList.h"
+#include "Misc/TransactionallySafeCriticalSection.h"
 #include "SplineInterfaces.h" 
 #include "Templates/Function.h"
 
@@ -47,7 +48,7 @@ public:
     {
         // Check if this type is already registered
         TPair<FString, FString> NamePair(ImplName, ValueTypeName);
-        FScopeLock Lock(&GetRegistryLock());
+		UE::TScopeLock Lock(GetRegistryLock());
         
         // Check for existing registration with same ID but different names (hash collision)
         if (const TPair<FString, FString>* ExistingNames = GetIdToNames().Find(TypeId))
@@ -89,7 +90,7 @@ public:
      */
     static inline TUniquePtr<ISplineInterface> CreateSpline(TypeId TypeId)
     {
-        FScopeLock Lock(&GetRegistryLock());
+		UE::TScopeLock Lock(GetRegistryLock());
         
         const FactoryFunction* Factory = GetFactories().Find(TypeId);
         if (Factory)
@@ -125,7 +126,7 @@ public:
     {
         TPair<FString, FString> NamePair(ImplName, ValueTypeName);
         
-        FScopeLock Lock(&GetRegistryLock());
+		UE::TScopeLock Lock(GetRegistryLock());
         
         const TypeId* TypeId = GetNamesToId().Find(NamePair);
         if (TypeId && *TypeId != 0)
@@ -151,7 +152,7 @@ public:
         FString& OutImplName,
         FString& OutValueTypeName)
     {
-        FScopeLock Lock(&GetRegistryLock());
+        UE::TScopeLock Lock(GetRegistryLock());
         
         const TPair<FString, FString>* Names = GetIdToNames().Find(TypeId);
         if (Names)
@@ -176,7 +177,7 @@ public:
     {
         TPair<FString, FString> NamePair(ImplName, ValueTypeName);
         
-        FScopeLock Lock(&GetRegistryLock());
+		UE::TScopeLock Lock(GetRegistryLock());
         
         const TypeId* Id = GetNamesToId().Find(NamePair);
         if (Id)
@@ -190,9 +191,9 @@ private:
     /**
      * Gets the registry lock (thread-safe singleton access)
      */
-    static inline FCriticalSection& GetRegistryLock()
+    static inline FTransactionallySafeCriticalSection& GetRegistryLock()
     {
-        static FCriticalSection RegistryLock;
+        static FTransactionallySafeCriticalSection RegistryLock;
         return RegistryLock;
     }
     
