Minor fixups.

#rnx

diff --git a/Engine/Source/Runtime/GeometryCore/Public/Splines/SplineTypeRegistry.h b/Engine/Source/Runtime/GeometryCore/Public/Splines/SplineTypeRegistry.h
--- a/Engine/Source/Runtime/GeometryCore/Public/Splines/SplineTypeRegistry.h
+++ b/Engine/Source/Runtime/GeometryCore/Public/Splines/SplineTypeRegistry.h
@@ -4,7 +4,6 @@
 
 #include "CoreMinimal.h"
 #include "SplineTypeId.h"
-#include "Containers/LockFreeList.h"
 #include "Misc/TransactionallySafeCriticalSection.h"
 #include "SplineInterfaces.h" 
 #include "Templates/Function.h"
