SplineCore: Add getters for child mapping ranges

diff --git a/Engine/Source/Runtime/GeometryCore/Public/Splines/MultiSpline.h b/Engine/Source/Runtime/GeometryCore/Public/Splines/MultiSpline.h
--- a/Engine/Source/Runtime/GeometryCore/Public/Splines/MultiSpline.h
+++ b/Engine/Source/Runtime/GeometryCore/Public/Splines/MultiSpline.h
@@ -668,6 +668,38 @@ public:
         Child->MappingRangeSpace = RangeSpace;
         return true;
     }
+
+	/**
+	 * Returns the mapping range for an attribute channel
+	 * @param Name Name of the attribute channel
+	 * @return The mapping range
+	 */
+	FInterval1f GetAttributeChannelRange(FName Name)
+	{
+		FChildSpline* Child = Children.Find(Name);
+		if (!Child)
+		{
+			return FInterval1f::Empty();
+		}
+
+		return Child->MappingRange;
+	}
+
+	/**
+	* Returns the mapping range space for an attribute channel
+	* @param Name Name of the attribute channel
+	* @return The mapping range space. Always returns normalized if the child spline is not found.
+	*/
+	EMappingRangeSpace GetAttributeChannelRangeSpace(FName Name)
+	{
+		FChildSpline* Child = Children.Find(Name);
+		if (!Child)
+		{
+			return EMappingRangeSpace::Normalized;
+		}
+
+		return Child->MappingRangeSpace;
+	}
     
 private:
 
