Optimize PolyBezierSpline - FindNearest.

#rnx

diff --git a/Engine/Source/Runtime/GeometryCore/Public/Splines/PolyBezierSpline.h b/Engine/Source/Runtime/GeometryCore/Public/Splines/PolyBezierSpline.h
--- a/Engine/Source/Runtime/GeometryCore/Public/Splines/PolyBezierSpline.h
+++ b/Engine/Source/Runtime/GeometryCore/Public/Splines/PolyBezierSpline.h
@@ -11,6 +11,15 @@ namespace Geometry
 {
 namespace Spline
 {
+namespace SplineCVars
+{
+	static bool bFindNearestUseSplineCulling = true;
+	static FAutoConsoleVariableRef CVarFindNearestUseSplineCulling (
+		TEXT("Spline.PolyBezier.FindNearestUseSplineCulling"),
+		bFindNearestUseSplineCulling,
+		TEXT("If enabled, FindNearest uses spline culling to optimize."),
+		ECVF_Default);
+}
 
 template<typename VALUETYPE> class UE_EXPERIMENTAL(5.7, "New spline APIs are experimental.") TPolyBezierSpline;
 	
@@ -305,24 +314,76 @@ public:
     	return MapLocalSegmentParameterToGlobal(SegmentIndex, LocalT);
     }
 	
-	virtual float FindNearest(const ValueType& Point, float& OutSquaredDistance) const override
+	float PointDistanceSqToSegmentBBox(const ValueType &Point, int32 SegmentIndex) const
 	{
-    	float BestDistSq = TNumericLimits<float>::Max();
-    	float BestParam = 0.0f;
+		// Get the control points
+		int32 BaseIndex = SegmentIndex * 4;
+		const ValueType& P0 = Base::GetValue(BaseIndex);
+		const ValueType& P1 = Base::GetValue(BaseIndex + 1);
+		const ValueType& P2 = Base::GetValue(BaseIndex + 2);
+		const ValueType& P3 = Base::GetValue(BaseIndex + 3);
+
+
+		//	For cubics, convex hull is defined by control points, so use them to compute BBox. 
+		ValueType BBoxMin = P0.ComponentMin(P1);
+		ValueType BBoxMax = P0.ComponentMax(P1);
+
+		BBoxMin = BBoxMin.ComponentMin(P2);
+		BBoxMax = BBoxMax.ComponentMax(P2);
 
+		BBoxMin = BBoxMin.ComponentMin(P3);
+		BBoxMax = BBoxMax.ComponentMax(P3);
+
+		ValueType PtOnBox = BBoxMin.ComponentMax(BBoxMax.ComponentMin(Point));
+
+		return float(Math::SizeSquared(PtOnBox - Point));
+	}
+
+	virtual float FindNearest(const ValueType& Point, float& OutSquaredDistance) const override
+	{
     	const int32 NumSegments = Base::NumKeys() / 4;
-    	for (int32 i = 0; i < NumSegments; ++i)
-    	{
-    		float SegmentDistSq;
-    		const float SplineParam = FindNearestOnSegment(Point, i, SegmentDistSq);
-    		
-    		if (SegmentDistSq < BestDistSq)
-    		{
-    			BestDistSq = SegmentDistSq;
-    			BestParam = SplineParam;
-    		}
-    	}
+		float BestDistSq;
+		float BestParam = FindNearestOnSegment(Point, 0, BestDistSq);
+		
+		float SegmentDistSq;
+
+		if (SplineCVars::bFindNearestUseSplineCulling)
+		{
+			for (int32 i = 1; i < NumSegments; ++i)
+			{
+				// For cubic spline: compute AABB and distance to point. 
+				//	If distance > best result distance then reject segment. Saves a lot of computation (see FindNearestOnSegment)
+				// 
+				if (PointDistanceSqToSegmentBBox(Point, i) >= BestDistSq)
+				{
+					continue;
+				}
 
+				const float SplineParam = FindNearestOnSegment(Point, i, SegmentDistSq);
+				if (SegmentDistSq < BestDistSq)
+				{
+					BestDistSq = SegmentDistSq;
+					BestParam = SplineParam;
+				}
+
+				if (FMath::IsNearlyZero(BestDistSq))
+				{
+					break;
+				}
+			}
+		}
+		else
+		{
+			for (int32 i = 1; i < NumSegments; ++i)
+			{
+				const float SplineParam = FindNearestOnSegment(Point, i, SegmentDistSq);
+				if (SegmentDistSq < BestDistSq)
+				{
+					BestDistSq = SegmentDistSq;
+					BestParam = SplineParam;
+				}
+			}
+		}
     	OutSquaredDistance = BestDistSq;
     	
     	return BestParam;
@@ -1939,7 +2000,18 @@ private:
 };
 
 
-	
+template<> 
+inline float TPolyBezierSpline<FQuat>::PointDistanceSqToSegmentBBox(const ValueType& Point, int32 SegmentIndex) const
+{
+	return 0;
+}
+
+template<> 
+inline float TPolyBezierSpline<float>::PointDistanceSqToSegmentBBox(const ValueType& Point, int32 SegmentIndex) const
+{
+	return 0;
+}
+
 // Common type definitions
 using FPolyBezierSpline2f = TPolyBezierSpline<FVector2f>;
 using FPolyBezierSpline3f = TPolyBezierSpline<FVector3f>;
