Spline rotation fixes

#Synced-CL 48296767

diff --git a/Engine/Source/Runtime/Engine/Private/Curves/Spline.cpp b/Engine/Source/Runtime/Engine/Private/Curves/Spline.cpp
--- a/Engine/Source/Runtime/Engine/Private/Curves/Spline.cpp
+++ b/Engine/Source/Runtime/Engine/Private/Curves/Spline.cpp
@@ -41,7 +41,7 @@ FAutoConsoleVariableRef CVarSplineImplementation(
 #endif
 );
 
-int32 GParameterizationPolicy = 2;
+int32 GParameterizationPolicy = 0;
 FAutoConsoleVariableRef CVarParameterizationPolicy(
 	TEXT("Spline.ParameterizationPolicy"),
 	GParameterizationPolicy,
@@ -69,7 +69,7 @@ FAutoConsoleVariableRef CVarUseLegacyPositionEvaluation(
 	TEXT("If true, evaluating the position channel always routes through an interp curve.")
 );
 
-bool GUseLegacyRotationEvaluation = true;
+bool GUseLegacyRotationEvaluation = false;
 FAutoConsoleVariableRef CVarUseLegacyRotationEvaluation(
 	TEXT("Spline.UseLegacyRotationEvaluation"),
 	GUseLegacyRotationEvaluation,
@@ -133,9 +133,9 @@ inline UE::Geometry::Spline::ETangentMode ConvertInterpCurveModeToTangentMode(EI
 	switch (Mode)
 	{
 	case CIM_Linear:
-		return UE::Geometry::Spline::ETangentMode::Linear;
+		return UE::Geometry::Spline::ETangentMode::LegacyLinear; // Default to legacy for compatibility
 	case CIM_CurveAuto:
-		return UE::Geometry::Spline::ETangentMode::Auto;
+		return UE::Geometry::Spline::ETangentMode::LegacyAuto; // Default to legacy for compatibility
 	case CIM_Constant:
 		return UE::Geometry::Spline::ETangentMode::Constant;
 	case CIM_CurveUser:
@@ -143,7 +143,7 @@ inline UE::Geometry::Spline::ETangentMode ConvertInterpCurveModeToTangentMode(EI
 	case CIM_CurveBreak:
 		return UE::Geometry::Spline::ETangentMode::Broken;
 	case CIM_CurveAutoClamped:
-		return UE::Geometry::Spline::ETangentMode::AutoClamped;
+		return UE::Geometry::Spline::ETangentMode::LegacyAutoClamped; // Default to legacy for compatibility
 	case CIM_Unknown:
 	default:
 		return UE::Geometry::Spline::ETangentMode::Unknown;
@@ -157,9 +157,11 @@ inline EInterpCurveMode ConvertTangentModeToInterpCurveMode(UE::Geometry::Spline
 {
 	switch (Mode)
 	{
+	case UE::Geometry::Spline::ETangentMode::LegacyLinear:
 	case UE::Geometry::Spline::ETangentMode::Linear:
 		return CIM_Linear;
 	case UE::Geometry::Spline::ETangentMode::Auto:
+	case UE::Geometry::Spline::ETangentMode::LegacyAuto:
 		return CIM_CurveAuto;
 	case UE::Geometry::Spline::ETangentMode::Constant:
 		return CIM_Constant;
@@ -168,6 +170,7 @@ inline EInterpCurveMode ConvertTangentModeToInterpCurveMode(UE::Geometry::Spline
 	case UE::Geometry::Spline::ETangentMode::Broken:
 		return CIM_CurveBreak;
 	case UE::Geometry::Spline::ETangentMode::AutoClamped:
+	case UE::Geometry::Spline::ETangentMode::LegacyAutoClamped:
 		return CIM_CurveAutoClamped;
 	case UE::Geometry::Spline::ETangentMode::Unknown:
 	default:
@@ -1990,7 +1993,8 @@ float FNewSpline::GetSegmentLength(const int32 Index, const float Param, const F
 	const auto& T1 = GetSpline().GetTangentIn(Index == LastPoint ? 0 : Index + 1);
 
 	// Special cases for linear
-	if (GetSpline().GetTangentModes()[Index] == UE::Geometry::Spline::ETangentMode::Linear)
+	if (GetSpline().GetTangentModes()[Index] == UE::Geometry::Spline::ETangentMode::Linear ||
+		GetSpline().GetTangentModes()[Index] == UE::Geometry::Spline::ETangentMode::LegacyLinear)
 	{
 		return ((P1 - P0) * Scale3D).Size() * Param;
 	}
@@ -2121,13 +2125,13 @@ void FNewSpline::AddPoint(const FSplinePoint& Point)
 
 		if (UE::Geometry::Spline::TTangentBezierSpline<FVector>* ScaleChild = GetTypedAttributeChannel<UE::Geometry::Spline::TTangentBezierSpline<FVector>>(ScaleAttrName))
 		{
-			ScaleChild->AppendPoint(UE::Geometry::Spline::TTangentBezierSpline<FVector>::FTangentBezierControlPoint(FVector::OneVector));
+			ScaleChild->AppendPoint(UE::Geometry::Spline::TTangentBezierSpline<FVector>::FTangentBezierControlPoint(Point.Scale));
 			ScaleChild->SetKnotVector(GetSpline().GetKnotVector());
 		}
 		
 		if (UE::Geometry::Spline::TTangentBezierSpline<FQuat>* RotChild = GetTypedAttributeChannel<UE::Geometry::Spline::TTangentBezierSpline<FQuat>>(RotationAttrName))
 		{
-			RotChild->AppendPoint(UE::Geometry::Spline::TTangentBezierSpline<FQuat>::FTangentBezierControlPoint(FQuat::Identity));
+			RotChild->AppendPoint(UE::Geometry::Spline::TTangentBezierSpline<FQuat>::FTangentBezierControlPoint(Point.Rotation.Quaternion()));
 			RotChild->SetKnotVector(GetSpline().GetKnotVector());
 		}
 		
@@ -2140,13 +2144,13 @@ void FNewSpline::AddPoint(const FSplinePoint& Point)
     	
 		if (UE::Geometry::Spline::TTangentBezierSpline<FVector>* ScaleChild = GetTypedAttributeChannel<UE::Geometry::Spline::TTangentBezierSpline<FVector>>(ScaleAttrName))
 		{
-			ScaleChild->PrependPoint(UE::Geometry::Spline::TTangentBezierSpline<FVector>::FTangentBezierControlPoint(FVector::OneVector));
+    		ScaleChild->PrependPoint(UE::Geometry::Spline::TTangentBezierSpline<FVector>::FTangentBezierControlPoint(Point.Scale));
 			ScaleChild->SetKnotVector(GetSpline().GetKnotVector());
 		}
     
 		if (UE::Geometry::Spline::TTangentBezierSpline<FQuat>* RotChild = GetTypedAttributeChannel<UE::Geometry::Spline::TTangentBezierSpline<FQuat>>(RotationAttrName))
 		{
-			RotChild->PrependPoint(UE::Geometry::Spline::TTangentBezierSpline<FQuat>::FTangentBezierControlPoint(FQuat::Identity));
+    		RotChild->PrependPoint(UE::Geometry::Spline::TTangentBezierSpline<FQuat>::FTangentBezierControlPoint(Point.Rotation.Quaternion()));
 			RotChild->SetKnotVector(GetSpline().GetKnotVector());
 		}
     	
@@ -2161,14 +2165,14 @@ void FNewSpline::AddPoint(const FSplinePoint& Point)
 	float ScaleParam = MapParameterToChildSpace(ScaleAttrName, NewPointInternalParam);
 	if (UE::Geometry::Spline::TTangentBezierSpline<FVector>* ScaleChild = GetTypedAttributeChannel<UE::Geometry::Spline::TTangentBezierSpline<FVector>>(ScaleAttrName))
 	{
-		ScaleChild->InsertPointAtGlobalParam(ScaleParam, UE::Geometry::Spline::TTangentBezierSpline<FVector>::FTangentBezierControlPoint(FVector::OneVector));
+		ScaleChild->InsertPointAtGlobalParam(ScaleParam, UE::Geometry::Spline::TTangentBezierSpline<FVector>::FTangentBezierControlPoint(Point.Scale));
 		ScaleChild->SetKnotVector(GetSpline().GetKnotVector());
 	}
 	
 	float RotParam = MapParameterToChildSpace(RotationAttrName, NewPointInternalParam);
 	if (UE::Geometry::Spline::TTangentBezierSpline<FQuat>* RotChild = GetTypedAttributeChannel<UE::Geometry::Spline::TTangentBezierSpline<FQuat>>(RotationAttrName))
 	{
-		RotChild->InsertPointAtGlobalParam(RotParam, UE::Geometry::Spline::TTangentBezierSpline<FQuat>::FTangentBezierControlPoint(FQuat::Identity));
+		RotChild->InsertPointAtGlobalParam(RotParam, UE::Geometry::Spline::TTangentBezierSpline<FQuat>::FTangentBezierControlPoint(Point.Rotation.Quaternion()));
 		RotChild->SetKnotVector(GetSpline().GetKnotVector());
 	}
 	
@@ -2195,13 +2199,13 @@ void FNewSpline::InsertPoint(const FSplinePoint& Point, int32 Index)
 		
 		if (UE::Geometry::Spline::TTangentBezierSpline<FVector>* ScaleChild = GetTypedAttributeChannel<UE::Geometry::Spline::TTangentBezierSpline<FVector>>(ScaleAttrName))
 		{
-			ScaleChild->AppendPoint(UE::Geometry::Spline::TTangentBezierSpline<FVector>::FTangentBezierControlPoint(FVector::OneVector));
+			ScaleChild->AppendPoint(UE::Geometry::Spline::TTangentBezierSpline<FVector>::FTangentBezierControlPoint(Point.Scale));
 			ScaleChild->SetKnotVector(GetSpline().GetKnotVector());
 		}
 		
 		if (UE::Geometry::Spline::TTangentBezierSpline<FQuat>* RotChild = GetTypedAttributeChannel<UE::Geometry::Spline::TTangentBezierSpline<FQuat>>(RotationAttrName))
 		{
-			RotChild->AppendPoint(UE::Geometry::Spline::TTangentBezierSpline<FQuat>::FTangentBezierControlPoint(FQuat::Identity));
+			RotChild->AppendPoint(UE::Geometry::Spline::TTangentBezierSpline<FQuat>::FTangentBezierControlPoint(Point.Rotation.Quaternion()));
 			RotChild->SetKnotVector(GetSpline().GetKnotVector());
 		}
 		
@@ -2214,13 +2218,13 @@ void FNewSpline::InsertPoint(const FSplinePoint& Point, int32 Index)
 		
 		if (UE::Geometry::Spline::TTangentBezierSpline<FVector>* ScaleChild = GetTypedAttributeChannel<UE::Geometry::Spline::TTangentBezierSpline<FVector>>(ScaleAttrName))
 		{
-			ScaleChild->PrependPoint(UE::Geometry::Spline::TTangentBezierSpline<FVector>::FTangentBezierControlPoint(FVector::OneVector));
+			ScaleChild->PrependPoint(UE::Geometry::Spline::TTangentBezierSpline<FVector>::FTangentBezierControlPoint(Point.Scale));
 			ScaleChild->SetKnotVector(GetSpline().GetKnotVector());
 		}
 	
 		if (UE::Geometry::Spline::TTangentBezierSpline<FQuat>* RotChild = GetTypedAttributeChannel<UE::Geometry::Spline::TTangentBezierSpline<FQuat>>(RotationAttrName))
 		{
-			RotChild->PrependPoint(UE::Geometry::Spline::TTangentBezierSpline<FQuat>::FTangentBezierControlPoint(FQuat::Identity));
+			RotChild->PrependPoint(UE::Geometry::Spline::TTangentBezierSpline<FQuat>::FTangentBezierControlPoint(Point.Rotation.Quaternion()));
 			RotChild->SetKnotVector(GetSpline().GetKnotVector());
 		}
 		
@@ -2236,14 +2240,14 @@ void FNewSpline::InsertPoint(const FSplinePoint& Point, int32 Index)
 	float ScaleParam = MapParameterToChildSpace(ScaleAttrName, NewPointInternalParam);
 	if (UE::Geometry::Spline::TTangentBezierSpline<FVector>* ScaleChild = GetTypedAttributeChannel<UE::Geometry::Spline::TTangentBezierSpline<FVector>>(ScaleAttrName))
 	{
-		ScaleChild->InsertPointAtGlobalParam(ScaleParam, UE::Geometry::Spline::TTangentBezierSpline<FVector>::FTangentBezierControlPoint(FVector::OneVector));
+		ScaleChild->InsertPointAtGlobalParam(ScaleParam, UE::Geometry::Spline::TTangentBezierSpline<FVector>::FTangentBezierControlPoint(Point.Scale));
 		ScaleChild->SetKnotVector(GetSpline().GetKnotVector());
 	}
 	
 	float RotParam = MapParameterToChildSpace(RotationAttrName, NewPointInternalParam);
 	if (UE::Geometry::Spline::TTangentBezierSpline<FQuat>* RotChild = GetTypedAttributeChannel<UE::Geometry::Spline::TTangentBezierSpline<FQuat>>(RotationAttrName))
 	{
-		RotChild->InsertPointAtGlobalParam(RotParam, UE::Geometry::Spline::TTangentBezierSpline<FQuat>::FTangentBezierControlPoint(FQuat::Identity));
+		RotChild->InsertPointAtGlobalParam(RotParam, UE::Geometry::Spline::TTangentBezierSpline<FQuat>::FTangentBezierControlPoint(Point.Rotation.Quaternion()));
 		RotChild->SetKnotVector(GetSpline().GetKnotVector());
 	}
 	
diff --git a/Engine/Source/Runtime/GeometryCore/Public/Splines/BSpline.h b/Engine/Source/Runtime/GeometryCore/Public/Splines/BSpline.h
--- a/Engine/Source/Runtime/GeometryCore/Public/Splines/BSpline.h
+++ b/Engine/Source/Runtime/GeometryCore/Public/Splines/BSpline.h
@@ -130,7 +130,8 @@ public:
 
 	static constexpr int32 WindowSize = Degree + 1;
 	using FWindow = TStaticArray<const ValueType*, WindowSize>;
-
+	using FWindowStore = TStaticArray<ValueType, WindowSize>;
+	
 	DECLARE_SPLINE_TYPE_ID(
 		BSplineNameSelector<DEGREE>::Name,
 		*TSplineValueTypeTraits<VALUETYPE>::Name
@@ -279,10 +280,16 @@ public:
 		if (NumValues == 0) return ValueType();
 		if (NumValues == 1) return GetValue(0);
 
-		FWindow Window = FindWindow(Parameter);
-
-		// If FindWindow fails, it will return an array of nullptr.
-		if (Window[0] == nullptr) return ValueType();
+		FWindow Window;
+		FWindowStore Scratch;
+		FindWindow(Parameter, Window, Scratch);
+		
+		// If FindWindow fails, it will return an array of nullptr. InterpolateWindow assumes validity of elements.
+		if (!Math::FSplineValidation::IsValidWindow<ValueType, WindowSize>(Window) ||
+				!Math::FSplineValidation::IsValidParameter(Parameter))
+		{
+			return Math::FSplineValidation::GetDefaultValue<ValueType>();
+		}
 
 		return InterpolateWindow(Window, Parameter);
 	}
@@ -1199,12 +1206,12 @@ private:
 		return Degree;
 	}
 
-	virtual FWindow FindWindow(float Parameter) const
+	virtual void FindWindow(float Parameter, FWindow& OutWindow, FWindowStore& Scratch) const
 	{
-		FWindow Window = {};
+		OutWindow = {};
 		if (NumKeys() < Degree + 1)
 		{
-			return Window;
+			return;
 		}
 
 		int32 Span = FindKnotSpan(Parameter);
@@ -1219,7 +1226,7 @@ private:
 			{
 				int32 Index = (Span - Degree + i) % NumPoints;
 				if (Index < 0) Index += NumPoints;
-				Window[i] = &GetValue(Index);
+				OutWindow[i] = &GetValue(Index);
 			}
 		}
 		else
@@ -1230,12 +1237,10 @@ private:
 				const int32 Index = Span - Degree + i;
 				if (Index >= 0 && Index < NumKeys())
 				{
-					Window[i] = &GetValue(Index);
+					OutWindow[i] = &GetValue(Index);
 				}
 			}
 		}
-
-		return Window;
 	}
 
 	/**
diff --git a/Engine/Source/Runtime/GeometryCore/Public/Splines/InterpolationPolicies.h b/Engine/Source/Runtime/GeometryCore/Public/Splines/InterpolationPolicies.h
--- a/Engine/Source/Runtime/GeometryCore/Public/Splines/InterpolationPolicies.h
+++ b/Engine/Source/Runtime/GeometryCore/Public/Splines/InterpolationPolicies.h
@@ -57,6 +57,39 @@ public:
 		}
 	}
 
+	static T InterpolateCubic(TArrayView<const T* const> Window, float Parameter)
+	{
+		if (Window.Num() < 4)
+		{
+			// Fallback to linear interpolation if not enough points
+			return Interpolate(Window, Parameter);
+		}
+		if constexpr (TIsMathInterpolable<T>::value)
+		{
+			// Calculate Bernstein basis functions
+			const float t = Parameter;
+			const float OneMinusT = 1.0f - t;
+			const float t2 = t * t;
+			const float OneMinusT2 = OneMinusT * OneMinusT;
+		
+			// Cubic Bernstein polynomials
+			const float Basis[4] =
+			{
+				OneMinusT * OneMinusT2,     // (1-t)³
+				3.0f * t * OneMinusT2,      // 3t(1-t)²
+				3.0f * t2 * OneMinusT,      // 3t²(1-t)
+				t * t2                      // t³
+			};
+		
+			return TSplineInterpolationPolicy<T>::InterpolateWithBasis(Window, Basis);
+		}
+		else
+		{
+			// Fallback for non-math types:
+			return Parameter < 0.5f ? *Window[1] : *Window[2];
+		}
+	}
+
 	// Add window-based interpolation for spline implementations
 	static T InterpolateWithBasis(TArrayView<const T* const> Window, TArrayView<const float> Basis)
 	{
@@ -212,136 +245,206 @@ public:
         return TQuat::Slerp(*Window[0], *Window[1], static_cast<RealType>(Parameter));
 	}
 
-	// Weighted quaternion interpolation 
-	// based on eigenvector approach from Markley et al.,
-	// which handles antipodal quaternions and degenerate cases effectively.
+	/**
+	 * Cubic interpolation using SQUAD (Spherical Quadrangle).
+	 * Provides C1 continuous rotation interpolation matching FInterpCurve behavior.
+	 * 
+	 * Window layout matches cubic Bezier:
+	 * - P0: Start rotation
+	 * - C0: Out control (tangent at start)
+	 * - C1: In control (tangent at end)
+	 * - P1: End rotation
+	 * - Alpha parameter is in [0,1]
+	 */
+	static TQuat InterpolateCubic(TArrayView<const TQuat* const> Window, float Alpha)
+	{
+		if (Window.Num() < 4)
+		{
+			// Fallback if not enough points
+			return Interpolate(Window, Alpha);
+		}
+		
+		return TQuat::Squad(
+				*Window[0],  // P0 - start point
+				*Window[1],  // C0 - out control
+				*Window[3],  // P1 - end point
+				*Window[2],  // C1 - in control
+				Alpha);
+	}
+	
+	/**
+	 * Performs a weighted blend of multiple quaternions using arbitrary basis weights.
+	 *
+	 * This implements a *reference-centric* log–exp mean on S³:
+	 *   1) Choose a reference quaternion Qref (here: the sample with max |weight|).
+	 *   2) For each sample Qi, hemisphere-align to Qref, form the relative rotation
+	 *      Qrel = Qref^{-1} * Qi, map to tangent via V = Log(Qrel).
+	 *   3) Average the tangent vectors with the given weights and map back:
+	 *      Result = Qref * Exp( average(V) ).
+	 *
+	 * Properties:
+	 *  - Smooth, normalized, and works for arbitrary (even non-convex) weights.
+	 *  - With two inputs and convex weights [1−t, t], this reduces exactly to SLERP.
+	 *  - This is a *mean*, not a spherical Bezier/SQUAD; use de Casteljau/SQUAD for true cubic curves.
+	 *  - If all weights are nonnegative, we normalize by the signed sum (convex case, allows cancellation).
+	 *  - For mixed-sign weights we normalize by the sum of absolutes to avoid degeneracy.
+	 */
 	static TQuat InterpolateWithBasis(TArrayView<const TQuat* const> Window, TArrayView<const float> Basis)
 	{
-		const int32 Count = FMath::Min(Window.Num(), Basis.Num());
-		if (Count == 0)
+		const int32 N = FMath::Min(Window.Num(), Basis.Num());
+		if (N <= 0)
 		{
 			return TQuat::Identity;
 		}
-
-		// Use first quaternion as reference for hemisphere checks
-		const TQuat& ReferenceQuat = *Window[0];
-
-        // Build 4x4 correlation matrix
-		RealType M[4][4] = {};
-        RealType TotalWeight = 0;
-
-		for (int32 Index = 0; Index < Count; ++Index)
+		if (N == 1)
+		{
+			TQuat q = *Window[0];
+			q.Normalize();
+			return q;
+		}
+		// Optimization: detect linear case (2 points with [1-t, t] weights)
+		if (N == 2)
 		{
-            const RealType Weight = static_cast<RealType>(Basis[Index]);
-            if (Weight > static_cast<RealType>(UE_KINDA_SMALL_NUMBER))
+			const RealType W0 = static_cast<RealType>(Basis[0]);
+			const RealType W1 = static_cast<RealType>(Basis[1]);
+			const RealType Sum = W0 + W1;
+			
+			if (FMath::IsNearlyEqual(Sum, static_cast<RealType>(1), static_cast<RealType>(1e-4)) && 
+				W0 >= static_cast<RealType>(0) && W1 >= static_cast<RealType>(0))
 			{
-				TQuat Quat = *Window[Index];
-
-                // Ensure consistent hemisphere to avoid cancellation
-				if ((Quat | ReferenceQuat) < static_cast<RealType>(0))
+				// Linear blend - use efficient Slerp
+				TQuat Q0 = (*Window[0]); Q0.Normalize();
+				TQuat Q1 = (*Window[1]); Q1.Normalize();
+				
+				if ((Q0 | Q1) < static_cast<RealType>(0))
 				{
-					Quat *= static_cast<RealType>(-1);
+					Q1 = Q1 * static_cast<RealType>(-1);
 				}
-
-				// Accumulate outer product contribution
-				const RealType Q[4] = {Quat.X, Quat.Y, Quat.Z, Quat.W};
-				for (int32 Row = 0; Row < 4; ++Row)
+				
+				return TQuat::Slerp(Q0, Q1, W1 / Sum);
+			}
+		}
+		constexpr RealType DotAntipodalCut  = static_cast<RealType>(-0.95); // treat dot ≤ -0.95 as "nearly antipodal"
+		constexpr RealType WeightBalanceTol = static_cast<RealType>(0.05);  // within 5% of each other -> "balanced"
+
+		// ---------------------------------------------------------------------------
+		// Antipodal pair detection: if two high-weight samples are nearly opposite
+		// and their weights are nearly equal, we want cancellation about Identity.
+		// ---------------------------------------------------------------------------
+		bool bHasAntipodalBalancedPair = false;
+		for (int32 i = 0; i < N && !bHasAntipodalBalancedPair; ++i)
+		{
+			for (int32 j = i + 1; j < N; ++j)
+			{
+				const RealType wi = static_cast<RealType>(Basis[i]);
+				const RealType wj = static_cast<RealType>(Basis[j]);
+				if (wi >= static_cast<RealType>(0) && wj >= static_cast<RealType>(0)) // convex-only check
 				{
-					for (int32 Col = 0; Col < 4; ++Col)
+					const RealType sum = wi + wj;
+					if (sum > static_cast<RealType>(UE_KINDA_SMALL_NUMBER))
 					{
-						M[Row][Col] += Weight * Q[Row] * Q[Col];
+						const RealType balance = FMath::Abs(wi - wj) / sum;
+						TQuat qi = *Window[i]; qi.Normalize();
+						TQuat qj = *Window[j]; qj.Normalize();
+						const RealType dot = (qi | qj);
+
+						if (dot <= DotAntipodalCut && balance <= WeightBalanceTol)
+						{
+							bHasAntipodalBalancedPair = true;
+							break;
+						}
 					}
 				}
-
-				TotalWeight += Weight;
 			}
 		}
-
-		if (FMath::IsNearlyZero(TotalWeight))
+		
+		// ---------------------------------------------------------------------------
+		// Choose reference quaternion
+		// - If antipodal balanced pair detected: use Identity so the pair cancels.
+		// - Else: use the max-|weight| sample (your current strategy).
+		// ---------------------------------------------------------------------------
+		TQuat Qref;
+		int32 RefIdx = 0;
+
+		if (bHasAntipodalBalancedPair)
 		{
-			return TQuat::Identity;
+			Qref = TQuat::Identity;
 		}
-
-		// Precompute inverse for efficiency
-		double InvTotalWeight = 1.0 / TotalWeight;
-
-        // Normalize matrix for numerical stability
-		for (int32 Row = 0; Row < 4; ++Row)
-		{
-			for (int32 Col = 0; Col < 4; ++Col)
-			{
-				M[Row][Col] *= InvTotalWeight;
-			}
-		}
-
-		// Find dominant eigenvector using power iteration
-        RealType EigenVec[4] = {0, 0, 0, 1};
-        RealType PrevEigenValue = 0;
-
-		constexpr int32 MaxIterations = 32;
-		for (int32 Iter = 0; Iter < MaxIterations; ++Iter)
+		else
 		{
-			// Matrix-vector multiplication: v' = M*v
-            RealType NewVec[4] = {0};
-			for (int32 Row = 0; Row < 4; ++Row)
+			RealType MaxAbsW = static_cast<RealType>(0);
+			for (int32 i = 0; i < N; ++i)
 			{
-				for (int32 Col = 0; Col < 4; ++Col)
+				const RealType aw = FMath::Abs(static_cast<RealType>(Basis[i]));
+				if (aw > MaxAbsW)
 				{
-					NewVec[Row] += M[Row][Col] * EigenVec[Col];
+					MaxAbsW = aw;
+					RefIdx = i;
 				}
 			}
+			Qref = *Window[RefIdx];
+			Qref.Normalize();
 
-            // Normalize vector
-			RealType LengthSq = 
-				NewVec[0] * NewVec[0] + 
-				NewVec[1] * NewVec[1] +
-				NewVec[2] * NewVec[2] + 
-				NewVec[3] * NewVec[3];
-
-			// Possible with certain opposing quaternion configurations
-			if (FMath::IsNearlyZero(LengthSq))
+			if (MaxAbsW <= static_cast<RealType>(UE_KINDA_SMALL_NUMBER))
 			{
-				return TQuat::Identity;
+				return Qref; // nothing meaningful to average
 			}
+		}
 
-			RealType InvLength = FMath::InvSqrt(LengthSq);
-			for (int32 i = 0; i < 4; ++i)
+		const TQuat QrefInv = Qref.Inverse();
+		
+		// Accumulate weighted logarithms in tangent space
+		RealType SumX = 0, SumY = 0, SumZ = 0;
+		RealType WeightSum = 0;          // signed sum for nonnegative bases
+		RealType WeightAbsSum = 0;       // backup for mixed-sign cases
+		bool bAllNonNegative = true;
+		
+		for (int32 i = 0; i < N; ++i)
+		{
+			const RealType Weight = static_cast<RealType>(Basis[i]);
+			if (FMath::IsNearlyZero(Weight))
 			{
-				EigenVec[i] = NewVec[i] * InvLength;
+				continue;
 			}
+			bAllNonNegative &= (Weight >= static_cast<RealType>(0));
+			
+			TQuat Qi = (*Window[i]);
+			Qi.Normalize();
 
-			// Calculate Rayleigh quotient for convergence check
-            RealType EigenValue = 0;
-			for (int32 Row = 0; Row < 4; ++Row)
+			// Hemisphere alignment relative to reference
+			if ((Qref | Qi) < static_cast<RealType>(0))
 			{
-				for (int32 Col = 0; Col < 4; ++Col)
-				{
-					EigenValue += EigenVec[Row] * M[Row][Col] * EigenVec[Col];
-				}
+				Qi = Qi * static_cast<RealType>(-1);
 			}
 
-			// Check for sufficient convergence
-			if (FMath::IsNearlyEqual(EigenValue, PrevEigenValue))
-			{
-				break;
-			}
+			// Relative rotation: Qrel = Qref^{-1} * Qi
+			const TQuat Qrel = QrefInv * Qi;
+			const TQuat V = Qrel.Log(); // tangent vector (pure imaginary)
 
-			PrevEigenValue = EigenValue;
+			SumX += Weight * V.X;
+			SumY += Weight * V.Y;
+			SumZ += Weight * V.Z;
+			WeightSum += Weight;
+			WeightAbsSum += FMath::Abs(Weight);
 		}
 
-        // Convert eigenvector to TQuat
-		TQuat ResultQuat(
-            EigenVec[0],
-            EigenVec[1],
-            EigenVec[2],
-            EigenVec[3]
-		);
+		// If all weights are nonnegative, use their (signed) sum so cancellations happen naturally.
+		// If mixed signs, fall back to the absolute sum to avoid division by ~0 and keep continuity.
+		const RealType Norm = bAllNonNegative ? WeightSum : WeightAbsSum;
+		if (Norm <= static_cast<RealType>(UE_KINDA_SMALL_NUMBER))
+		{
+			return Qref;
+		}
 
-		ResultQuat.Normalize();
-		return ResultQuat;
+		// Map back to manifold: Qref * Exp(avg tangent)
+		const TQuat D = TQuat(SumX / Norm, SumY / Norm, SumZ / Norm, static_cast<RealType>(0)).Exp();
+		TQuat Result = Qref * D;
+		Result.Normalize();
+		return Result;
 	}
 
 
+
 	// Template-based n-th derivative calculation
 	template<int32 Order>
 	static TQuat EvaluateDerivative(TArrayView<const TQuat* const> Window, float Parameter)
diff --git a/Engine/Source/Runtime/GeometryCore/Public/Splines/PolyBezierSpline.h b/Engine/Source/Runtime/GeometryCore/Public/Splines/PolyBezierSpline.h
--- a/Engine/Source/Runtime/GeometryCore/Public/Splines/PolyBezierSpline.h
+++ b/Engine/Source/Runtime/GeometryCore/Public/Splines/PolyBezierSpline.h
@@ -29,7 +29,12 @@ public:
 	using Base = TBSpline<VALUETYPE, 3>;
 	typedef typename Base::ValueType ValueType;
 	using FWindow = typename Base::FWindow;
-
+	using FWindowStore = typename Base::FWindowStore;
+	using FFindWindowOverride = void(*)(FWindow& /*InOut*/,
+										int32 /*SegmentIndex*/,
+										const FInterval1f& /*SegmentRange*/,
+										FWindowStore& /*Scratch*/);
+	
 	// Generate compile-time type ID for PolyBezier
 	DECLARE_SPLINE_TYPE_ID(
 		TEXT("PolyBezier"),
@@ -395,10 +400,15 @@ public:
     template<int32 Order>
     ValueType EvaluateDerivative(float Parameter) const
     {
-		FWindow Window = FindWindow(Parameter);
-
+		FWindow Window;
+		FWindowStore Scratch;
+		FindWindow(Parameter, Window, Scratch);
 		// If FindWindow fails, it will return an array of nullptr. InterpolateWindow assumes validity of elements.
-		if (Window[0] == nullptr) { return ValueType(); }
+		if (!Math::FSplineValidation::IsValidWindow<ValueType, Base::WindowSize>(Window) || // Cubic Bezier needs 4 points
+				!Math::FSplineValidation::IsValidParameter(Parameter))
+		{
+			return Math::FSplineValidation::GetDefaultValue<ValueType>();
+		}
 		
         // Order 0 is just regular evaluation
 	    if constexpr (Order == 0)
@@ -411,11 +421,10 @@ public:
     	float LocalT;
     	MapGlobalParameterToLocalSegment(Parameter, SegmentIndex, LocalT);
 		
-	    // For PolyBezier, each segment is normalized to unit parameter space
-        constexpr float SegmentScale = 1.0f;
+		const float SegmentScale = GetSegmentParameterRange(SegmentIndex).Length();
 	    
-	    // Compute the derivative with proper scaling
-	    return Math::TBezierDerivativeCalculator<ValueType, Order>::Compute(
+		// Compute the derivative with proper scaling
+		return Math::TBezierDerivativeCalculator<ValueType, Order>::Compute(
 	        Window, LocalT, SegmentScale);
     }
 
@@ -1565,7 +1574,7 @@ public:
 	{
 		Base::SetCustomKnots(NewKnots);
 	}
-	
+	void SetFindWindowOverride(FFindWindowOverride Fn) { FindWindowOverrideFn = Fn; }
 private:
 	
 	/**
@@ -1661,12 +1670,12 @@ private:
 		this->MarkFlatKnotsCacheDirty();
 	}
 
-    virtual FWindow FindWindow(float Parameter) const override
+    virtual void FindWindow(float Parameter, FWindow& OutWindow, FWindowStore& Scratch) const override
     {
-        FWindow Window = {};
+        OutWindow = {};
         if (Base::NumKeys() < 4)  // Need at least one complete Bezier segment
         {
-            return Window;
+            return;
         }
 		const TArray<ValueType>& Entries = this->Values;
         // Find which segment we're in
@@ -1697,7 +1706,7 @@ private:
 
         	for (int Idx = 0; Idx < Base::WindowSize; ++Idx)
         	{
-        		Window[Idx] = &Entries[WrapIndex(BaseIndex + Idx)];
+        		OutWindow[Idx] = &Entries[WrapIndex(BaseIndex + Idx)];
         	}
         }
         else
@@ -1705,17 +1714,21 @@ private:
         	// Return null window if BaseIndex is too large to fully populate the window.
             if (BaseIndex + 3 >= Entries.Num())
             {
-                return Window;
+                return;
             }
 
         	// Standard case - no wrapping
         	for (int Idx = 0; Idx < Base::WindowSize; ++Idx)
         	{
-        		Window[Idx] = &Entries[BaseIndex + Idx];
+        		OutWindow[Idx] = &Entries[BaseIndex + Idx];
         	}
         }
-
-        return Window;
+		
+		if (FindWindowOverrideFn && Math::FSplineValidation::IsValidWindow<ValueType, Base::WindowSize>(OutWindow))
+		{
+			const FInterval1f SegmentRange = GetSegmentParameterRange(SegmentIndex);
+			FindWindowOverrideFn(OutWindow, SegmentIndex, SegmentRange, Scratch);
+		}
     }
 
     virtual ValueType InterpolateWindow(TArrayView<const ValueType* const> Window, float Parameter) const override 
@@ -1724,27 +1737,11 @@ private:
         {
             return ValueType();
         }
-  		int32 NumSegments = GetNumberOfSegments();
     	int32 SegmentIndex;
     	float LocalT;
     	MapGlobalParameterToLocalSegment(Parameter, SegmentIndex, LocalT);
 
-        // Calculate Bernstein basis functions
-        const float t = LocalT;
-        const float OneMinusT = 1.0f - t;
-        const float t2 = t * t;
-        const float OneMinusT2 = OneMinusT * OneMinusT;
-
-        // Cubic Bernstein polynomials
-        const float Basis[4] =
-        {
-            OneMinusT * OneMinusT2,     // (1-t)³
-            3.0f * t * OneMinusT2,      // 3t(1-t)²
-            3.0f * t2 * OneMinusT,      // 3t²(1-t)
-            t * t2                      // t³
-        };
-
-        return TSplineInterpolationPolicy<ValueType>::InterpolateWithBasis(Window, Basis);
+        return TSplineInterpolationPolicy<ValueType>::InterpolateCubic(Window, LocalT);
     }
 
 	/**
@@ -2005,6 +2002,9 @@ private:
 
 	/** This is the last segment fetched by parameter. Used to accelerate predictable sampling patterns. */
 	mutable int32 LastEvaluatedSegment = INDEX_NONE;
+	
+	/** Optional override for FindWindow*/
+	FFindWindowOverride FindWindowOverrideFn = nullptr;
 };
 
 
diff --git a/Engine/Source/Runtime/GeometryCore/Public/Splines/SplineMath.h b/Engine/Source/Runtime/GeometryCore/Public/Splines/SplineMath.h
--- a/Engine/Source/Runtime/GeometryCore/Public/Splines/SplineMath.h
+++ b/Engine/Source/Runtime/GeometryCore/Public/Splines/SplineMath.h
@@ -270,23 +270,16 @@ namespace Math
 			}
 		}
 
-		template <typename T>
-		static bool IsValidWindow(TArrayView<const T* const> Window, int32 RequiredSize)
+		template <typename T, int32 N>
+		static bool IsValidWindow(const TArrayView<const T* const>& Window)
 		{
-			if (Window.Num() < RequiredSize)
+			for (int32 i = 0; i < N; ++i)
 			{
-				return false;
-			}
-
-			// Check for null pointers
-			for (const T* Ptr : Window)
-			{
-				if (!Ptr)
+				if (!Window.IsValidIndex(i) || Window[i] == nullptr)
 				{
 					return false;
 				}
 			}
-
 			return true;
 		}
 
@@ -461,7 +454,7 @@ namespace Math
 		static T Compute(TArrayView<const T* const> Window, float Parameter)
 		{
 			// Basic validation
-			if (!FSplineValidation::IsValidWindow(Window, Order + 1))
+			if (!FSplineValidation::IsValidWindow<T, Order+1>(Window))
 			{
 				return FSplineValidation::GetDefaultValue<T>();
 			}
@@ -696,7 +689,13 @@ namespace Math
 	};
 
 
-	// Bezier curve specific calculations
+	/**
+	 * Bezier derivative calculator.
+	 * Computes derivatives and applies segment scaling.
+	 * 
+	 * @tparam T Value type (FVector, FQuat, FTransform, etc.)
+	 * @tparam Order Derivative order (0=position, 1=velocity, 2=acceleration, etc.)
+	 */
 	template <typename T, int32 Order>
 	struct TBezierDerivativeCalculator
 	{
@@ -707,7 +706,7 @@ namespace Math
 				return FSplineValidation::GetDefaultValue<T>();
 			}
 
-			if (!FSplineValidation::IsValidWindow(Window, 4) || // Cubic Bezier needs 4 points
+			if (!FSplineValidation::IsValidWindow<T, 4>(Window) || // Cubic Bezier needs 4 points
 				!FSplineValidation::IsValidParameter(Parameter))
 			{
 				return FSplineValidation::GetDefaultValue<T>();
@@ -729,12 +728,15 @@ namespace Math
 			// derivatives scale inversely with parameter range
 			if constexpr (Order > 0)
 			{
-				// Calculate scaling factor: 1 / SegmentScale^Order
-				float ScalingFactor = 1.0f;
-				for (int32 i = 0; i < Order; ++i)
+				// If L == 0, the segment is degenerate in parameter space
+				// return the zero derivative (well-defined limit)
+				if (SegmentScale == 0.0f)
 				{
-					ScalingFactor /= SegmentScale;
+					return FSplineValidation::GetDefaultValue<T>();
 				}
+				
+				// Calculate scaling factor: 1 / SegmentScale^Order
+				float ScalingFactor = FMath::Pow(1.0f / SegmentScale, static_cast<float>(Order));
 				return Result * ScalingFactor;
 			}
 			else
diff --git a/Engine/Source/Runtime/GeometryCore/Public/Splines/TangentBezierSpline.h b/Engine/Source/Runtime/GeometryCore/Public/Splines/TangentBezierSpline.h
--- a/Engine/Source/Runtime/GeometryCore/Public/Splines/TangentBezierSpline.h
+++ b/Engine/Source/Runtime/GeometryCore/Public/Splines/TangentBezierSpline.h
@@ -2,6 +2,7 @@
 
 #pragma once
 
+#include "TangentsPolicy.h"
 #include "Splines/PolyBezierSpline.h"
 namespace UE
 {
@@ -18,24 +19,33 @@ template<typename VALUETYPE> class UE_EXPERIMENTAL(5.7, "New spline APIs are exp
  */
 enum class UE_EXPERIMENTAL(5.7, "New spline APIs are experimental.") ETangentMode : uint8
 {
+	/** FInterpCurve: Copies prev tangents if prev is not curve key */
+	LegacyAuto,
+	
+	/** FInterpCurve: Clamped version with copying behavior */
+	LegacyAutoClamped,
+	
+	/** User-specified tangents */
+	User,
+	
+	/** Broken tangents - in and out can be different */
+	Broken,
+	
+	/** FInterpCurve behavior (continuity-preserving) */
+	LegacyLinear,
+		
+	/** Constant - no interpolation between points */
+	Constant,
+	
 	/** Automatically compute unclamped tangents based on surrounding points */
 	Auto,
     
 	/** Automatically compute clamped tangents based on surrounding points */
 	AutoClamped,
     
-	/** User-specified tangents */
-	User,
-    
-	/** Broken tangents - in and out can be different */
-	Broken,
-    
 	/** Linear tangents aligned with adjacent points */
 	Linear,
     
-	/** Constant - no interpolation between points */
-	Constant,
-    
 	/** Unknown or invalid tangent mode */
 	Unknown
 };
@@ -132,24 +142,34 @@ class TTangentBezierSpline :
 {
 public:
 	typedef typename TSplineWrapper<TPolyBezierSpline<VALUETYPE>>::ValueType ValueType;
-	using TSplineWrapper<TPolyBezierSpline<VALUETYPE>>::InternalSpline;
+	
 	using FTangentBezierControlPoint = TTangentBezierControlPoint<ValueType>;
-
+	using FPolyBezierSpline = TPolyBezierSpline<ValueType>;
+	using FWindow      = typename FPolyBezierSpline::FWindow;
+	using FWindowStore = typename FPolyBezierSpline::FWindowStore;
+	using TangentOps = Math::TTangentOps<ValueType>;
+	
 	// Generate compile-time type ID for TangentBezier
 	DECLARE_SPLINE_TYPE_ID(
 		TEXT("TangentBezier"),
 		*TSplineValueTypeTraits<VALUETYPE>::Name
 	);
 	
-	TTangentBezierSpline() = default;
+	TTangentBezierSpline()
+	{
+		InstallFindWindowOverride();
+	}
+	
 	virtual ~TTangentBezierSpline() override = default;
 	
 	/** Default constructor with at least one segment */
 	TTangentBezierSpline(const ValueType& StartPoint, const ValueType& EndPoint)
-		: Tension(0.0f)
-		, TangentModes({ETangentMode::Auto, ETangentMode::Auto})
+		: TTangentBezierSpline()
+ 
 	{
-		InternalSpline = FPolyBezierSpline3d::CreateLine(StartPoint, EndPoint);
+		Tension = 0.f;
+		TangentModes = {ETangentMode::Auto, ETangentMode::Auto};
+		AssignInternalSpline(FPolyBezierSpline3d::CreateLine(StartPoint, EndPoint));
 	}
 
 	/**
@@ -161,23 +181,25 @@ public:
 		const ValueType& StartTangent,
 		const ValueType& EndTangent,
 		bool bAutoTangents = false)
-		: Tension(0.0f)
-		, TangentModes({bAutoTangents ? ETangentMode::Auto : ETangentMode::User,
-							bAutoTangents ? ETangentMode::Auto : ETangentMode::User})
+			: TTangentBezierSpline()
 	{
-		InternalSpline = FPolyBezierSpline3d(StartPoint,
-											StartPoint + StartTangent / 3.0f,
-											EndPoint - (EndTangent / 3.0f),
-											EndPoint);
+		Tension = 0.f;
+		TangentModes = {bAutoTangents ? ETangentMode::Auto : ETangentMode::User,
+						bAutoTangents ? ETangentMode::Auto : ETangentMode::User};
+		AssignInternalSpline( FPolyBezierSpline3d(StartPoint,
+											TangentOps::TangentOutToP1(StartPoint, StartTangent),
+											TangentOps::TangentInToP2(EndPoint, EndTangent),
+											EndPoint));
 	}
    
     /**  Copy constructor */
     TTangentBezierSpline(const TTangentBezierSpline& Other)
-	    : Tension(Other.Tension)
-	    , TangentModes(Other.TangentModes)
-    {
-		InternalSpline = Other.InternalSpline;
-    }
+	: TTangentBezierSpline()
+	{
+		Tension = Other.Tension;
+		TangentModes = Other.TangentModes;
+		AssignInternalSpline(Other.InternalSpline);
+	}
 
     /**  Copy assignment */
     TTangentBezierSpline& operator=(const TTangentBezierSpline& Other)
@@ -187,11 +209,33 @@ public:
             Tension = Other.Tension;
             TangentModes = Other.TangentModes;
 			bStationaryEndpoint = Other.bStationaryEndpoint;
-            InternalSpline = Other.InternalSpline;
+        	AssignInternalSpline(Other.InternalSpline);
         }
         return *this;
     }
 
+	/** Move constructor */
+	TTangentBezierSpline(TTangentBezierSpline&& Other) noexcept
+		: TTangentBezierSpline() 
+	{
+		Tension = Other.Tension;
+		TangentModes = MoveTemp(Other.TangentModes);
+		bStationaryEndpoint = Other.bStationaryEndpoint;
+		AssignInternalSpline(MoveTemp(Other.InternalSpline));
+	}
+	/** Move assignment */
+	TTangentBezierSpline& operator=(TTangentBezierSpline&& Other) noexcept
+	{
+		if (this != &Other)
+		{
+			Tension = Other.Tension;
+			TangentModes = MoveTemp(Other.TangentModes);
+			bStationaryEndpoint = Other.bStationaryEndpoint;
+			AssignInternalSpline(MoveTemp(Other.InternalSpline));
+		}
+		return *this;
+	}
+
 	virtual bool IsEqual(const ISplineInterface* OtherSpline) const override
 	{
 		if (OtherSpline->GetTypeId() == GetTypeId())
@@ -219,6 +263,8 @@ public:
 		if (Ar.IsLoading())
 		{
 			TangentModes.SetNum(NumTangentModes);
+
+			InstallFindWindowOverride();
 			
 			// It was once valid for a spline to have a single bezier point, but we now expect 3 for this case.
 			// It is never valid to have 1 or 2 points.
@@ -275,7 +321,7 @@ public:
 		TTangentBezierSpline Result;
         
 		// Clear the default initialization
-		Result.InternalSpline = FPolyBezierSpline3d::CreateLine(StartPoint, EndPoint);
+		Result.AssignInternalSpline(FPolyBezierSpline3d::CreateLine(StartPoint, EndPoint));
 		
         Result.Reparameterize(EParameterizationPolicy::Uniform);
         return Result;
@@ -301,15 +347,15 @@ public:
         TTangentBezierSpline Result;
         
         // Clear the default initialization
-        Result.InternalSpline = FPolyBezierSpline3d::CreateCircleArc(
-            Center, Radius, StartAngle, EndAngle, NumSegments);
+        Result.AssignInternalSpline(FPolyBezierSpline3d::CreateCircleArc(
+            Center, Radius, StartAngle, EndAngle, NumSegments));
         
 		// Set up tangent status for all points
 		Result.TangentModes.Init(ETangentMode::Auto, Result.GetNumPoints());
         
         // Apply consistent parameterization - even though the internal PolyBezier already calls
         // Reparameterize, we need to do it again here to ensure consistency
-        Result.Reparameterize(EParameterizationPolicy::Centripetal);
+        Result.Reparameterize(EParameterizationPolicy::Uniform);
         
         return Result;
     }
@@ -328,14 +374,14 @@ public:
     {
         // Create the circle
         TTangentBezierSpline Result;
-        Result.InternalSpline = FPolyBezierSpline3d::CreateCircle(
-            ValueType(Center), Radius, NumSegments);
+        Result.AssignInternalSpline(FPolyBezierSpline3d::CreateCircle(
+            ValueType(Center), Radius, NumSegments));
         
         // Set up tangent status for all points
 		Result.TangentModes.Init(ETangentMode::Auto, Result.GetNumPoints());
 		
 		// Apply consistent parameterization
-		Result.Reparameterize(EParameterizationPolicy::Centripetal);
+		Result.Reparameterize(EParameterizationPolicy::Uniform);
         
         return Result;
     }
@@ -356,14 +402,14 @@ public:
     {
         // Create the ellipse
         TTangentBezierSpline Result;
-        Result.InternalSpline = FPolyBezierSpline3d::CreateEllipse(
-            Center, RadiusX, RadiusY, NumSegments);
+        Result.AssignInternalSpline(FPolyBezierSpline3d::CreateEllipse(
+            Center, RadiusX, RadiusY, NumSegments));
         
 		// Set up tangent status for all points
 		Result.TangentModes.Init(ETangentMode::Auto, Result.GetNumPoints());
         
         // Apply consistent parameterization
-        Result.Reparameterize(EParameterizationPolicy::Centripetal);
+        Result.Reparameterize(EParameterizationPolicy::Uniform);
         
         return Result;
     }
@@ -372,7 +418,7 @@ public:
 		TUniquePtr<TTangentBezierSpline<VALUETYPE>> Clone = MakeUnique<TTangentBezierSpline<VALUETYPE>>();
     
 		// Copy internal spline
-		Clone->InternalSpline = this->InternalSpline;
+		Clone->AssignInternalSpline(this->InternalSpline);
     
 		// Copy tangent settings
 		Clone->Tension = this->Tension;
@@ -426,11 +472,11 @@ public:
 	            // Calculate control points for the closing segment
 	            ValueType P0 = LastPos;
 	            ValueType P3 = FirstPos;
-	            ValueType P1 = P0 + (LastOutTangent / 3.0f);
-	            ValueType P2 = P3 - (FirstInTangent / 3.0f);
-	            
+	        	ValueType P1 = TangentOps::TangentOutToP1(P0, LastOutTangent);
+	        	ValueType P2 = TangentOps::TangentInToP2(P3, FirstInTangent);
+	        	
 	            // Add the closing segment to the internal spline
-	            InternalSpline.AppendBezierSegment(P1, P2, P3);
+	            InternalSpline.AppendBezierSegment(P1, P2, P3, EParameterizationPolicy::Uniform);
 	        	
 	        	InternalSpline.SetClosedLoopFlag(bInClosedLoop);
 	        }
@@ -462,7 +508,8 @@ public:
 	 * @param Points - Array of control points to set
 	 * @return True if the points were successfully set
 	 */
-	bool SetControlPoints(const TArray<FTangentBezierControlPoint>& Points)
+	bool SetControlPoints(const TArray<FTangentBezierControlPoint>& Points,
+		EParameterizationPolicy ParameterizationPolicy = EParameterizationPolicy::Uniform)
 	{
 	    const int32 NumPoints = Points.Num();
 		
@@ -495,121 +542,13 @@ public:
 	        
 	        ValueType P0 = StartPoint.Position;
 	        ValueType P3 = EndPoint.Position;
-	        
-	        // Calculate tangents based on the points' tangent modes
-	        ValueType OutTangent, InTangent;
-	        
-	        // Process start point's outgoing tangent
-	        if (StartPoint.TangentMode == ETangentMode::Auto || 
-	            StartPoint.TangentMode == ETangentMode::AutoClamped)
-	        {
-	            // Compute auto tangent
-	            if (i == 0 && NumPoints > 2)
-	            {
-	                // First point - use forward difference
-	                OutTangent = (1.0f - Tension) * (Points[i + 1].Position - Points[i].Position);
-	            }
-	            else if (i == NumPoints - 2 && NumPoints > 2)
-	            {
-	                // Point before last - use central difference
-	                OutTangent = (1.0f - Tension) * (Points[i + 1].Position - Points[i - 1].Position);
-	            }
-	            else if (NumPoints > 2)
-	            {
-	                // Normal case - use central difference
-	                OutTangent = (1.0f - Tension) * (Points[i + 1].Position - Points[i - 1].Position);
-	            }
-	            else
-	            {
-	                // Only two points - use forward difference
-	                OutTangent = (1.0f - Tension) * (Points[i + 1].Position - Points[i].Position);
-	            }
-	            
-	            // Apply clamping for AutoClamped
-	            if (StartPoint.TangentMode == ETangentMode::AutoClamped)
-	            {
-	                const float SegmentLength = Math::Distance(P3, P0);
-	                const float MaxTangentSize = SegmentLength * 0.33f;
-	                
-	                if (Math::Size(OutTangent) > MaxTangentSize && !FMath::IsNearlyZero(Math::Size(OutTangent)))
-	                {
-	                    OutTangent = Math::GetSafeNormal(OutTangent) * MaxTangentSize;
-	                }
-	            }
-	        }
-	        else if (StartPoint.TangentMode == ETangentMode::Linear)
-	        {
-	            // Linear tangent
-	            OutTangent = (P3 - P0) * 0.33f;
-	        }
-	        else if (StartPoint.TangentMode == ETangentMode::Constant)
-	        {
-	            // Zero tangent for constant interpolation
-	            OutTangent = ValueType();
-	        }
-	        else
-	        {
-	            // User or Broken mode - use provided tangent
-	            OutTangent = StartPoint.TangentOut;
-	        }
-	        
-	        // Process end point's incoming tangent
-	        if (EndPoint.TangentMode == ETangentMode::Auto || 
-	            EndPoint.TangentMode == ETangentMode::AutoClamped)
-	        {
-	            // Compute auto tangent
-	            if (i + 1 == 0 && NumPoints > 2)
-	            {
-	                // First point - use central difference
-	                InTangent = (1.0f - Tension) * (Points[i + 2].Position - Points[i].Position);
-	            }
-	            else if (i + 1 == NumPoints - 1 && NumPoints > 2)
-	            {
-	                // Last point - use backward difference
-	                InTangent = (1.0f - Tension) * (Points[i + 1].Position - Points[i].Position);
-	            }
-	            else if (NumPoints > 2)
-	            {
-	                // Normal case - use central difference
-	                InTangent = (1.0f - Tension) * (Points[i + 2].Position - Points[i].Position);
-	            }
-	            else
-	            {
-	                // Only two points - use backward difference
-	                InTangent = (1.0f - Tension) * (Points[i + 1].Position - Points[i].Position);
-	            }
-	            
-	            // Apply clamping for AutoClamped
-	            if (EndPoint.TangentMode == ETangentMode::AutoClamped)
-	            {
-	                const float SegmentLength = Math::Distance(P3, P0);
-	                const float MaxTangentSize = SegmentLength * 0.33f;
-	                
-	                if (Math::Size(InTangent) > MaxTangentSize && !FMath::IsNearlyZero(Math::Size(InTangent)))
-	                {
-	                    InTangent = Math::GetSafeNormal(InTangent) * MaxTangentSize;
-	                }
-	            }
-	        }
-	        else if (EndPoint.TangentMode == ETangentMode::Linear)
-	        {
-	            // Linear tangent
-	            InTangent = (P3 - P0) * 0.33f;
-	        }
-	        else if (EndPoint.TangentMode == ETangentMode::Constant)
-	        {
-	            // Zero tangent for constant interpolation
-	            InTangent = ValueType();
-	        }
-	        else
-	        {
-	            // User or Broken mode - use provided tangent
-	            InTangent = EndPoint.TangentIn;
-	        }
-	        
+	    	
+	    	ValueType OutTangent = StartPoint.TangentOut;
+	    	ValueType InTangent = EndPoint.TangentIn;
+	       	        
 	        // Calculate bezier control points
-	        ValueType P1 = P0 + OutTangent / 3.0f;
-	        ValueType P2 = P3 - InTangent / 3.0f;
+	    	ValueType P1 = TangentOps::TangentOutToP1(P0, OutTangent);
+	    	ValueType P2 = TangentOps::TangentInToP2(P3, InTangent);
 	        
 	        // Add all 4 control points for this segment
 	        BezierPoints.Add(P0);
@@ -619,7 +558,7 @@ public:
 	    }
 	    
 	    // Set the control points in the internal spline
-	    InternalSpline.SetControlPoints(BezierPoints, EParameterizationPolicy::Centripetal);
+	    InternalSpline.SetControlPoints(BezierPoints, ParameterizationPolicy);
 	    
 	    // Store tangent modes for later use
 	    TangentModes.Reset(NumPoints);
@@ -674,10 +613,23 @@ public:
                     const ValueType LastExistingPoint = GetValue(ExistingNumPoints - 1);
                     
                     // Check if connection points are reasonably close
-                    constexpr double ConnectionTolerance = 1e-4;
-                    const double DistanceSquared = Math::SizeSquared(LastExistingPoint - Points[0].Position);
+                	constexpr double ConnectionTolerance = 1e-4;
+                	bool bSameEndpoint;
+
+                	if constexpr (Math::TIsQuaternionType<ValueType>::value)
+                	{
+                		// angular distance
+                		const double dot = Math::Dot(LastExistingPoint, Points[0].Position);
+						const double ang = FMath::Acos(FMath::Clamp(dot, -1.0, 1.0)); // radians
+						bSameEndpoint = (ang <= ConnectionTolerance);
+                	}
+                	else
+                	{
+                		const double DistanceSquared = Math::SizeSquared(LastExistingPoint - Points[0].Position);
+                		bSameEndpoint = (DistanceSquared <= ConnectionTolerance);
+                	}
                     
-                    if (DistanceSquared <= ConnectionTolerance)
+                    if (bSameEndpoint)
                     {
                         // Points are close - the first existing point is the same as our first new point
                         // So update the last point's tangent mode if needed
@@ -705,9 +657,22 @@ public:
                     
                     // Check if connection points are reasonably close
                     constexpr double ConnectionTolerance = 1e-4;
-                    const double DistanceSquared = Math::SizeSquared(FirstExistingPoint - Points[i].Position);
+                	bool bSameEndpoint;
+
+                	if constexpr (Math::TIsQuaternionType<ValueType>::value)
+                	{
+                		// angular distance
+                		const double dot = Math::Dot(FirstExistingPoint, Points[i].Position);
+                		const double ang = FMath::Acos(FMath::Clamp(dot, -1.0, 1.0)); // radians
+                		bSameEndpoint = (ang <= ConnectionTolerance); // choose a small angular eps, e.g. ~1e-4 rad
+                	}
+	                else
+	                {
+		                const double DistanceSquared = Math::SizeSquared(FirstExistingPoint - Points[i].Position);
+	                	bSameEndpoint = (DistanceSquared <= ConnectionTolerance);
+	                }
                     
-                    if (DistanceSquared <= ConnectionTolerance)
+                    if (bSameEndpoint)
                     {
                         // Points are close - update first point's tangent mode if needed
                         if (Points[i].TangentMode != ETangentMode::User && Points[i].TangentMode != ETangentMode::Broken)
@@ -763,7 +728,7 @@ public:
 		int32 SegmentIndex,
 		float LocalT,
 		const FTangentBezierControlPoint& ControlPoint,
-		EParameterizationPolicy ParameterizationPolicy = EParameterizationPolicy::Centripetal)
+		EParameterizationPolicy ParameterizationPolicy = EParameterizationPolicy::Uniform)
 	{
 		int32 OutNewIndex;
 		// Special case: handle first two points
@@ -774,7 +739,7 @@ public:
 		
 		// Insert at that parameter with the exact position
 		// Have the internal spline handle the insertion
-		int32 InsertedPointIndex = InternalSpline.InsertPointAtSegmentParam(SegmentIndex, LocalT, ControlPoint.Position);
+		int32 InsertedPointIndex = InternalSpline.InsertPointAtSegmentParam(SegmentIndex, LocalT, ControlPoint.Position, ParameterizationPolicy);
 		
 		// If insertion was successful, update tangents if not auto-computed
 		if (InsertedPointIndex >= 0 &&
@@ -789,16 +754,12 @@ public:
 		SetPointTangentMode(InsertedPointIndex, ControlPoint.TangentMode);
 				
 		// Update adjacent points if they have automatic tangents
-		if (InsertedPointIndex > 0 && 
-			(TangentModes[InsertedPointIndex-1] == ETangentMode::Auto || 
-			 TangentModes[InsertedPointIndex-1] == ETangentMode::AutoClamped))
+		if (InsertedPointIndex > 0 && IsAutoTangent(InsertedPointIndex-1))
 		{
 			UpdatePointTangents(InsertedPointIndex-1, TangentModes[InsertedPointIndex-1]);
 		}
         
-		if (InsertedPointIndex+1 < GetNumPoints() && 
-			(TangentModes[InsertedPointIndex+1] == ETangentMode::Auto || 
-			 TangentModes[InsertedPointIndex+1] == ETangentMode::AutoClamped))
+		if (InsertedPointIndex+1 < GetNumPoints() && IsAutoTangent(InsertedPointIndex+1))
 		{
 			UpdatePointTangents(InsertedPointIndex+1, TangentModes[InsertedPointIndex+1]);
 		}
@@ -839,16 +800,12 @@ public:
 		SetPointTangentMode(InsertedPointIndex, ControlPoint.TangentMode);
 		
 		// Update adjacent points if they have automatic tangents
-		if (InsertedPointIndex > 0 && 
-			(TangentModes[InsertedPointIndex-1] == ETangentMode::Auto || 
-			 TangentModes[InsertedPointIndex-1] == ETangentMode::AutoClamped))
+		if (InsertedPointIndex > 0 && IsAutoTangent(InsertedPointIndex-1))
 		{
 			UpdatePointTangents(InsertedPointIndex-1, TangentModes[InsertedPointIndex-1]);
 		}
         
-		if (InsertedPointIndex + 1 < GetNumPoints() && 
-			(TangentModes[InsertedPointIndex+1] == ETangentMode::Auto || 
-			 TangentModes[InsertedPointIndex+1] == ETangentMode::AutoClamped))
+		if (InsertedPointIndex + 1 < GetNumPoints() && IsAutoTangent(InsertedPointIndex+1))
 		{
 			UpdatePointTangents(InsertedPointIndex+1, TangentModes[InsertedPointIndex+1]);
 		}
@@ -899,18 +856,18 @@ public:
 			OutTangent = GetValue(0) - ControlPoint.Position;
 		}
         
-		ValueType P1 = P0 + (OutTangent / 3.0f);
-		ValueType P2 = GetValue(0) - (GetTangentIn(0) / 3.0f);
-        
+		ValueType P1, P2;
+		P2 = TangentOps::TangentInToP2(GetValue(0), GetTangentIn(0));
+		P1 = TangentOps::TangentOutToP1(P0, OutTangent);
+		
 		// Let InternalSpline handle the prepending efficiently
-		InternalSpline.PrependBezierSegment(P0, P1, P2);
+		InternalSpline.PrependBezierSegment(P0, P1, P2, EParameterizationPolicy::Uniform);
         
 		TangentModes.InsertDefaulted(0);
 		SetPointTangentMode(0, ControlPoint.TangentMode);
 		
 		// Update tangents for the next point if it's auto
-		if (GetNumPoints() > 0 && 
-			(TangentModes[1] == ETangentMode::Auto || TangentModes[1] == ETangentMode::AutoClamped))
+		if (GetNumPoints() > 0 && IsAutoTangent(1))
 		{
 			UpdatePointTangents(1, TangentModes[1]);
 		}
@@ -947,8 +904,7 @@ public:
 		ValueType P3 = ControlPoint.Position;
     
 		// Update the previous point's tangents if needed
-		if (TangentModes[NewIndex - 1] == ETangentMode::Auto || 
-			TangentModes[NewIndex - 1] == ETangentMode::AutoClamped)
+		if (IsAutoTangent(NewIndex - 1))
 		{
 			UpdatePointTangents(NewIndex - 1, TangentModes[NewIndex - 1]);
 		}
@@ -957,38 +913,26 @@ public:
 		ValueType OutTangentPrev = GetTangentOut(NewIndex - 1);
     
 		// Calculate control points for the new segment
-		ValueType P1 = P0 + (OutTangentPrev / 3.0f);
-    
 		// For the incoming tangent of the new point
-		ValueType P2;
-    
-		if (ControlPoint.TangentMode == ETangentMode::User || 
-			ControlPoint.TangentMode == ETangentMode::Broken)
-		{
-			// For user-defined tangents, use the provided tangent
-			P2 = P3 - (ControlPoint.TangentIn / 3.0f);
-		}
-		else
-		{
-			// For other modes, we'll set up a reasonable default
-			// The UpdatePointTangents call below will update it properly
-			ValueType DefaultTangent = (P3 - P0) * (1.0f - Tension);
-			P2 = P3 - (DefaultTangent / 3.0f);
-		}
+		ValueType P1, P2;
+		ValueType InTangent = ControlPoint.TangentIn;
+		P2 = TangentOps::TangentInToP2(P3, InTangent);
+		P1 = TangentOps::TangentOutToP1(P0, OutTangentPrev);
     
 		// Add the segment
-		InternalSpline.AppendBezierSegment(P1, P2, P3);
+		InternalSpline.AppendBezierSegment(P1, P2, P3, EParameterizationPolicy::Uniform);
 		
 		TangentModes.InsertDefaulted(NewIndex);
 		SetPointTangentMode(NewIndex, ControlPoint.TangentMode);
 
         // If this is auto mode, update tangents for adjacent points
         if (ControlPoint.TangentMode == ETangentMode::Auto || 
-            ControlPoint.TangentMode == ETangentMode::AutoClamped)
+            ControlPoint.TangentMode == ETangentMode::AutoClamped ||
+            ControlPoint.TangentMode == ETangentMode::LegacyAuto || 
+			ControlPoint.TangentMode == ETangentMode::LegacyAutoClamped)
         {
             // Update previous point's tangents again to ensure continuity
-            if (TangentModes[NewIndex - 1] == ETangentMode::Auto || 
-                TangentModes[NewIndex - 1] == ETangentMode::AutoClamped)
+            if (IsAutoTangent(NewIndex - 1))
             {
                 UpdatePointTangents(NewIndex - 1, TangentModes[NewIndex - 1]);
             }
@@ -1102,7 +1046,8 @@ public:
 
         const ValueType OldPosition = GetValue(Index);
         bool bPositionChanged = (OldPosition != ControlPoint.Position);
-        
+		const bool bTangentModeChanged = (TangentModes[Index] != ControlPoint.TangentMode);
+
         SetPointTangentMode(Index, ControlPoint.TangentMode);
 
         // Update position
@@ -1134,8 +1079,11 @@ public:
         
         // If not auto tangents, apply the specified tangents
         if (ControlPoint.TangentMode != ETangentMode::Auto && 
-            ControlPoint.TangentMode != ETangentMode::AutoClamped && 
+            ControlPoint.TangentMode != ETangentMode::AutoClamped &&
+            ControlPoint.TangentMode != ETangentMode::LegacyAuto && 
+            ControlPoint.TangentMode != ETangentMode::LegacyAutoClamped && 
             ControlPoint.TangentMode != ETangentMode::Linear && 
+            ControlPoint.TangentMode != ETangentMode::LegacyLinear && 
             ControlPoint.TangentMode != ETangentMode::Constant)
         {
             SetTangentIn(Index, ControlPoint.TangentIn);
@@ -1143,19 +1091,21 @@ public:
         }
         
         // Update tangents based on tangent mode
-        UpdatePointTangents(Index, ControlPoint.TangentMode);
+		if (bTangentModeChanged)
+		{
+			UpdatePointTangents(Index, ControlPoint.TangentMode);
+		}
 
 		// Update surrounding points if position changed and they're auto tangents
+		// Only update neighbors for non-quaternion splines
         if (bPositionChanged)
         {
-            if (Index > 0 && (TangentModes[Index-1] == ETangentMode::Auto || 
-                              TangentModes[Index-1] == ETangentMode::AutoClamped))
+            if (Index > 0 && IsAutoTangent(Index-1))
             {
                 UpdatePointTangents(Index-1, TangentModes[Index-1]);
             }
             
-            if (Index < NumPoints-1 && (TangentModes[Index+1] == ETangentMode::Auto ||
-                                       TangentModes[Index+1] == ETangentMode::AutoClamped))
+            if (Index < NumPoints-1 && IsAutoTangent(Index+1))
             {
                 UpdatePointTangents(Index+1, TangentModes[Index+1]);
             }
@@ -1180,7 +1130,7 @@ public:
     {
     	const int32 NumPoints = GetNumPoints();
 
-    	if (NumPoints == 0 || Index < 0 || Index >= GetNumPoints())
+    	if (NumPoints == 0 || Index < 0 || Index > GetNumPoints())
     		return ValueType();
     
     	if (NumPoints == 1)
@@ -1238,43 +1188,32 @@ public:
 				return GetTangentOut(Index);
             }
         }
-        
-        // Normal case - get from previous segment's control points
-        const int32 BezierIndex = (Index - 1) * 4;
-        if (BezierIndex + 3 < InternalSpline.NumKeys())
-        {
-            ValueType P3 = InternalSpline.GetValue(BezierIndex + 3);
-            ValueType P2 = InternalSpline.GetValue(BezierIndex + 2);
-            return (P3 - P2) * 3.0f;
-        }
-        return ValueType();
+		const int32 p2Idx = P2IndexForIncomingOfPoint(Index);
+		if (p2Idx < 0)
+			return ValueType();
+
+		const ValueType P3 = GetValue(Index);
+		const ValueType P2 = InternalSpline.GetValue(p2Idx);
+		return TangentOps::P2ToTangentIn(P3, P2);
     }
     
 	/** Gets outgoing tangent for the specified point */
 	ValueType GetTangentOut(int32 Index) const
 	{
 		const int32 NumPoints = GetNumPoints();
-
-		if (NumPoints <= 1)
-		{
-			return ValueType();
-		}
-		
 		if (Index == NumPoints - 1 && !IsClosedLoop())
 		{
 			// If spline is unclosed, out tangent == in tangent at final point.
 			return GetTangentIn(Index);
 		}
 		
-		// Normal case - get from next segment's control points
-		const int32 BezierIndex = Index * 4;
-		if (BezierIndex + 1 < InternalSpline.NumKeys())
-		{
-			ValueType P0 = InternalSpline.GetValue(BezierIndex);
-			ValueType P1 = InternalSpline.GetValue(BezierIndex + 1);
-			return (P1 - P0) * 3.0f;
-		}
-		return ValueType();
+		const int32 p1Idx = P1IndexForOutgoingOfPoint(Index);
+		if (p1Idx < 0) return ValueType();
+
+		const ValueType P0 = GetValue(Index);
+		const ValueType P1 = InternalSpline.GetValue(p1Idx);
+
+		return TangentOps::P1ToTangentOut(P0, P1);
 	}
 
     /** Sets tangent in for the specified point */
@@ -1284,47 +1223,37 @@ public:
 		{
 			return;
 		}
+		// No incoming segment? (first point on open spline)
+		const int32 P2Index = P2IndexForIncomingOfPoint(Index);
+		if (P2Index < 0) return;
+
+		ValueType P3 = GetValue(Index);
+		ValueType Tin = NewTangent;
 		
-        if (Index == 0)
-        {
-            // Special case for first point in closed loop
-            if (InternalSpline.IsClosedLoop() && GetNumberOfSegments() > 0)
-            {
-                const int32 LastSegment = GetNumberOfSegments() - 1;
-                ValueType P3 = InternalSpline.GetValue(LastSegment * 4 + 3);
-                ValueType P2 = P3 - (NewTangent / 3.0f);
-                UpdateBezierControlPoint(LastSegment * 4 + 2, P2);
-            }
-            return;
-        }
+        ValueType P2 = TangentOps::TangentInToP2(P3, Tin);
+        UpdateBezierControlPoint(P2Index, P2);
         
-        // Normal case - update previous segment's P2 control point
-        const int32 BezierIndex = (Index - 1) * 4 + 2;
-        if (BezierIndex < InternalSpline.NumKeys())
-        {
-            ValueType P3 = InternalSpline.GetValue(BezierIndex + 1);
-            ValueType P2 = P3 - (NewTangent / 3.0f);
-            UpdateBezierControlPoint(BezierIndex, P2);
-        }
     }
     
     /** Sets tangent out for the specified point */
     void SetTangentOut(int32 Index, const ValueType& NewTangent)
     {
-        const int32 NumPoints = GetNumPoints();
-        if (Index == NumPoints - 1 && !IsClosedLoop())
-        {
-            return;
-        }
-        
-        // Normal case - update next segment's P1 control point
-        const int32 BezierIndex = Index * 4 + 1;
-        if (BezierIndex < InternalSpline.NumKeys())
-        {
-            ValueType P0 = InternalSpline.GetValue(BezierIndex - 1);
-            ValueType P1 = P0 + (NewTangent / 3.0f);
-            UpdateBezierControlPoint(BezierIndex, P1);
-        }
+		const int32 NumPoints = GetNumPoints();
+		if (Index < 0 || (Index >= NumPoints && !IsClosedLoop()))
+		{
+		    return;
+		}
+
+		// No outgoing segment? (last point on open spline)
+		const int32 P1Index = P1IndexForOutgoingOfPoint(Index);
+		if (P1Index < 0) return;
+
+		ValueType P0 = GetValue(Index);
+		ValueType Tout = NewTangent;
+		
+		ValueType P1 = TangentOps::TangentOutToP1(P0, Tout);
+		UpdateBezierControlPoint(P1Index, P1);
+      
     }
 
     /** Gets the number of points in the spline */
@@ -1346,6 +1275,7 @@ public:
 			return IsClosedLoop() ? GetNumberOfSegments() : GetNumberOfSegments() + 1;
     	}
     }
+	
 
 	virtual int32 GetNumberOfSegments() const override
 	{
@@ -1360,7 +1290,10 @@ public:
     /** Checks if a point's tangents are automatically computed */
     bool IsAutoTangent(int32 Index) const
     {
-		return TangentModes.IsValidIndex(Index) && (TangentModes[Index] == ETangentMode::Auto || TangentModes[Index] == ETangentMode::AutoClamped);
+		return TangentModes.IsValidIndex(Index) && (TangentModes[Index] == ETangentMode::Auto ||
+													TangentModes[Index] == ETangentMode::AutoClamped ||
+													TangentModes[Index] == ETangentMode::LegacyAuto ||
+													TangentModes[Index] == ETangentMode::LegacyAutoClamped);
     }
 
 	/**
@@ -1451,40 +1384,60 @@ public:
 	void UpdatePointTangents(int32 Index, ETangentMode Mode)
 	{
 		const int32 NumPoints = GetNumPoints();
+		const int32 NumSegments = GetNumberOfSegments();
 		if (Index < 0 || Index >= NumPoints || Mode == ETangentMode::User || Mode == ETangentMode::Broken)
 		{
 			// Invalid index or user-defined mode, no action needed
 			return;
 		}
-		
+
+		if constexpr (Math::TIsQuaternionType<ValueType>::value)
+		{
+			if (Mode == ETangentMode::Linear || Mode == ETangentMode::LegacyLinear)
+			{
+				const ValueType Q = GetValue(Index).GetNormalized();
+				SetTangentIn(Index,  Q);
+				SetTangentOut(Index, Q);
+				return; 
+			}
+		}
 
 		// Get the position of the current point
 		const ValueType& Current = GetValue(Index);
 
-		// Get previous and next points (handling closed loop case)
-		const ValueType& Prev = (Index > 0) ? GetValue(Index - 1) : (IsClosedLoop() ? GetValue(NumPoints - 1) : Current);
-
-		const ValueType& Next = (Index < NumPoints - 1) ? GetValue(Index + 1) : (IsClosedLoop() ? GetValue(0) : Current);
+		// Calculate indices handling closed loops
+		const int32 PrevIndex = (Index > 0) ? (Index - 1) : (IsClosedLoop() ? NumPoints - 1 : Index);
+		const int32 NextIndex = (Index < NumPoints - 1) ? (Index + 1) : (IsClosedLoop() ? 0 : Index);
+		const ValueType& Prev = GetValue(PrevIndex);
+		const ValueType& Next = GetValue(NextIndex);
 
-		// Get actual parameter values for points
-		float PrevParam = (Index > 0)
-			                  ? static_cast<float>(Index - 1)
-			                  : (IsClosedLoop() ? static_cast<float>(NumPoints - 1) : static_cast<float>(Index));
-
-		float CurrentParam = static_cast<float>(Index);
-
-		float NextParam = (Index < NumPoints - 1) ? static_cast<float>(Index + 1) : (IsClosedLoop() ? 0.0f : static_cast<float>(Index));
-
-		// Adjust params for closed loop to ensure proper weighting
-		if (IsClosedLoop() && Index == 0)
-		{
-			PrevParam -= static_cast<float>(NumPoints); // Make previous parameter negative for proper delta
-		}
-		else if (IsClosedLoop() && Index == NumPoints - 1)
+		// Get the actual parameter values for each point
+		float PrevParam = GetParameter(PrevIndex);
+		float CurrentParam = GetParameter(Index);
+		float NextParam = GetParameter(NextIndex);
+		
+		// If we are closed and updating the first or last point, we need to do some special work 
+		// to get virtual PrevParam or NextParam. The 'virtual' parameter is not the actual parameter 
+		// of the previous or next point, but the current parameter +- the knot delta (closing segment length).
+		if (IsClosedLoop())
 		{
-			NextParam += static_cast<float>(NumPoints); // Make next parameter greater than NumPoints for proper delta
+			if (Index == 0 && PrevIndex == NumPoints - 1)
+			{
+				// First point, prev wraps from end - subtract the closing segment length
+				// The closing segment is the last segment in a closed loop
+				PrevParam = CurrentParam - GetSegmentParameterRange(NumSegments - 1).Length();
+			}
+			else if (Index == NumPoints - 1 && NextIndex == 0)
+			{
+				// Last point, next wraps to start - add the closing segment length
+				NextParam = CurrentParam + GetSegmentParameterRange(NumSegments - 1).Length();
+			}
 		}
 
+		const float dTL = FMath::Max(UE_KINDA_SMALL_NUMBER, CurrentParam  - PrevParam); // span to the *left*
+		const float dTR = FMath::Max(UE_KINDA_SMALL_NUMBER, NextParam - CurrentParam ); // span to the *right*
+		
+		
 		// Calculate tangents based on mode
 		ValueType InTangent = GetTangentIn(Index); // Keep existing if not changing
 		ValueType OutTangent = GetTangentOut(Index);
@@ -1494,64 +1447,76 @@ public:
 		case ETangentMode::Auto:
 		case ETangentMode::AutoClamped:
 			{
-				// Legacy-style auto tangent calculation with parameter weighting
-				float PrevToNextParamDiff = FMath::Max<float>(UE_KINDA_SMALL_NUMBER, NextParam - PrevParam);
-
-				// Calculate basic tangent - matches legacy formula
-				ValueType AutoTangent = ((Next - Prev) / PrevToNextParamDiff) * (1.0f - Tension);
-
-				// For AutoClamped mode, apply clamping like the legacy implementation
-				if (Mode == ETangentMode::AutoClamped)
+				const bool bIsStationaryEndpoint = bStationaryEndpoint && 
+										  (Index == 0 || Index == NumPoints - 1) && 
+										  !IsClosedLoop();
+        
+				if (bIsStationaryEndpoint)
 				{
-					float PrevToCurParamDiff = FMath::Max<float>(UE_KINDA_SMALL_NUMBER, CurrentParam - PrevParam);
-					float NextToCurParamDiff = FMath::Max<float>(UE_KINDA_SMALL_NUMBER, NextParam - CurrentParam);
-
-					ValueType PrevToCurTangent = (Current - Prev) / PrevToCurParamDiff;
-					ValueType NextToCurTangent = (Next - Current) / NextToCurParamDiff;
-
-					float PrevToCurLength = static_cast<float>(Math::Size(PrevToCurTangent));
-					float NextToCurLength = static_cast<float>(Math::Size(NextToCurTangent));
-
-					// Clamp tangent magnitude
-					float MaxScale = FMath::Min(PrevToCurLength, NextToCurLength);
-
-					if (Math::Size(AutoTangent) > MaxScale && !FMath::IsNearlyZero(Math::Size(AutoTangent)))
+					if constexpr (Math::TIsQuaternionType<ValueType>::value)
+					{
+						InTangent  = Current; // control = key rotation
+						OutTangent = Current;
+					}
+					else
 					{
-						AutoTangent = Math::GetSafeNormal(AutoTangent) * MaxScale;
+						// Stationary endpoints get zero tangent
+						InTangent  = ValueType(0);
+						OutTangent = ValueType(0);
 					}
 				}
-
-				// Handle stationary endpoints
-				if (bStationaryEndpoint && (Index == 0 || Index == NumPoints - 1) && !IsClosedLoop())
+				else
 				{
-					AutoTangent = ValueType();
+					const bool bWantClamping = (Mode == ETangentMode::AutoClamped);
+					TangentOps::ComputeAutoTangent(
+						PrevParam, Prev, CurrentParam, Current, NextParam, Next,
+						Tension, bWantClamping, InTangent);
+					OutTangent = InTangent;
 				}
-
-				InTangent = AutoTangent;
-				OutTangent = AutoTangent;
 			}
 			break;
 
 		case ETangentMode::Linear:
 			{
-				// Previous segment tangent (incoming)
-				if (Index > 0 || IsClosedLoop())
-				{
-					InTangent = Current - Prev;
-				}
-				else
+				InTangent  = Current - Prev;
+				OutTangent = Next - Current;
+			}
+			break;
+			
+		case ETangentMode::LegacyLinear:
+		{
+			// FInterpCurve-compatible behavior
+			OutTangent = Next - Current;
+			InTangent = IsCurveKey(TangentModes[PrevIndex]) ? InTangent = OutTangent : (Current - Prev);
+		}
+			break;
+		case ETangentMode::LegacyAuto:
+		case ETangentMode::LegacyAutoClamped:
+			{
+				if (bStationaryEndpoint && (Index == 0 || Index == NumPoints - 1) && !IsClosedLoop())
 				{
-					InTangent = ValueType();
+					// Stationary endpoints get zero tangent
+					InTangent  = ValueType(0);
+					OutTangent = ValueType(0);
 				}
-
-				// Next segment tangent (outgoing)
-				if (Index < NumPoints - 1 || IsClosedLoop())
+				else if (IsCurveKey(TangentModes[PrevIndex]))
 				{
-					OutTangent = Next - Current;
+					// Use TTangentOps which handles clamping correctly
+					const bool bWantClamping = (Mode == ETangentMode::LegacyAutoClamped);
+					TangentOps::ComputeAutoTangent(
+						PrevParam, Prev,
+						CurrentParam, Current,
+						NextParam, Next,
+						Tension,
+						bWantClamping,
+						InTangent);
+					OutTangent = InTangent;
 				}
 				else
 				{
-					OutTangent = ValueType();
+					// Previous is Linear or Constant - copy its tangents (legacy behavior)
+					InTangent = GetTangentIn(PrevIndex);
+					OutTangent = GetTangentOut(PrevIndex);
 				}
 			}
 			break;
@@ -1571,13 +1536,19 @@ public:
 		default:
 			// Default to Auto mode if unknown
 			{
-				ValueType Tangent = (Next - Prev) * (1.0f - Tension);
+				ValueType Tangent;
+				TangentOps::ComputeAutoTangent(
+						PrevParam, Prev,
+						CurrentParam, Current,
+						NextParam, Next,
+						Tension,
+						false,
+						Tangent);
 				InTangent = Tangent;
 				OutTangent = Tangent;
 			}
 			break;
-		}
-
+		}		
 		// Apply the calculated tangents
 		SetTangentIn(Index, InTangent);
 		SetTangentOut(Index, OutTangent);
@@ -1600,7 +1571,7 @@ public:
         return InternalSpline.FindNearest(Point, OutSquaredDistance);
     }
 
-	void Reparameterize(EParameterizationPolicy Policy = EParameterizationPolicy::Centripetal)
+	void Reparameterize(EParameterizationPolicy Policy = EParameterizationPolicy::Uniform)
     {
 	    InternalSpline.Reparameterize(Policy);
 		
@@ -1679,11 +1650,70 @@ protected:
 		}
 	}
 private:
+	void InstallFindWindowOverride()
+	{
+		InternalSpline.SetFindWindowOverride(&RescaleNormalizedHandlesForSegment);
+	}
+	
+	/** ONLY adjusts P1/P2 in caller-provided scratch; generic over WindowSize */
+	static void RescaleNormalizedHandlesForSegment(FWindow& Window, int32 SegmentIndex,  const FInterval1f& SegRange, FWindowStore& Scratch)
+	{
+		// bail if window not complete
+		if (!Math::FSplineValidation::IsValidWindow<ValueType, 4>(Window))
+		{
+			return;
+		}
+		const float L = SegRange.Length();
+		const ValueType& P0 = *Window[0];
+		const ValueType& P1 = *Window[1];
+		const ValueType& P2 = *Window[2];
+		const ValueType& P3 = *Window[3];
+		
+		ValueType TOut = TangentOps::P1ToTangentOut(P0, P1);
+		ValueType TIn  = TangentOps::P2ToTangentIn(P3, P2);
+		
+		// Only write adjusted P1/P2 into scratch; keep P0,P3 pointers intact
+		Scratch[1] = TangentOps::TangentOutToP1(P0, TOut * L);
+		Scratch[2] = TangentOps::TangentInToP2(P3, TIn * L);
+		Window[1] = &Scratch[1];
+		Window[2] = &Scratch[2];
+	}
 
 	static bool IsCurveKey(ETangentMode Mode)
 	{
-		return ((Mode == ETangentMode::Auto) || (Mode == ETangentMode::AutoClamped) || (Mode == ETangentMode::User) || (Mode == ETangentMode::Broken));
+		return Mode == ETangentMode::Auto || 
+			   Mode == ETangentMode::AutoClamped ||
+			   Mode == ETangentMode::LegacyAuto ||
+			   Mode == ETangentMode::LegacyAutoClamped ||
+			   Mode == ETangentMode::User ||
+			   Mode == ETangentMode::Broken;
 	}
+
+	/** Returns base Bezier index for the given segment */
+	static int32 SegmentBaseIndex(int32 Segment) { return 4 * Segment; }
+
+	/** Returns Bezier index for the outgoing P1 control of point i; -1 if none */
+	int32 P1IndexForOutgoingOfPoint(int32 i) const
+	{
+		const int32 N = GetNumPoints();
+		const int32 S = i;
+		if (GetNumberOfSegments() == 0) return -1;
+		if (!IsClosedLoop() && i >= N - 1) return -1; // no outgoing at last point
+		const int32 base = SegmentBaseIndex(S);
+		return base + 1; // P1
+	}
+
+	/** Returns Bezier index for the incoming P2 control of point i; -1 if none */
+	int32 P2IndexForIncomingOfPoint(int32 i) const
+	{
+		const int32 N = GetNumPoints();
+		if (GetNumberOfSegments() == 0) return -1;
+		if (!IsClosedLoop() && i == 0) return -1; // no incoming at first point
+		const int32 S = (IsClosedLoop() && i == 0) ? (GetNumberOfSegments() - 1) : (i - 1);
+		const int32 base = SegmentBaseIndex(S);
+		return base + 2; // P2
+	}
+	
 	/**
 	 * Helper method to add the first segment to the spline
 	 * @param ControlPoint - Control point to insert
@@ -1698,8 +1728,8 @@ private:
     	{
     		// First point - just store it directly with both tangents, these are not used until the first segment is completed
     		InternalSpline.AddValue(ControlPoint.Position);
-    		InternalSpline.AddValue(ControlPoint.Position + (ControlPoint.TangentOut / 3.0f));
-			InternalSpline.AddValue(ControlPoint.Position - (ControlPoint.TangentIn  / 3.0f));
+    		InternalSpline.AddValue(TangentOps::TangentOutToP1(ControlPoint.Position, ControlPoint.TangentOut));
+    		InternalSpline.AddValue(TangentOps::TangentInToP2(ControlPoint.Position, ControlPoint.TangentIn));
     		
     		// update tangent mode
             if (TangentModes.IsEmpty())
@@ -1723,35 +1753,16 @@ private:
 				// Get the first point, then clear to start building segments
 				
 				// Calculate control points
+    			ValueType OutTangent = ControlPoint.TangentOut;
+    			ValueType InTangent = ControlPoint.TangentIn;
+    			
 				ValueType P0 = InternalSpline.GetValue(0);
-				ValueType P1, P2;
-				ValueType P3 = ControlPoint.Position;
-				
-				ValueType AutoTangent = (P3 - P0) * (1.0f - Tension);
-				
-				if (ControlPoint.TangentMode == ETangentMode::Auto ||
-					ControlPoint.TangentMode == ETangentMode::AutoClamped)
-				{
-					P2 = P3 - AutoTangent / 3.0f;
-				}
-				else
-				{
-					// Use provided tangents
-					P2 = P3 - ControlPoint.TangentIn / 3.0f;
-				}
-				
-				if (IsAutoTangent(0))
-				{
-					P1 = P0 + AutoTangent / 3.0f;
-				}
-				else
-				{
-					// Values[1] is already the proper out tangent value (Values[2] is the in tangent we'd use if prepending instead of appending here)
-					P1 = InternalSpline.GetValue(1);
-				}
-				
+    			ValueType P1 = InternalSpline.GetValue(1);  // Already stored placeholder P1
+    			ValueType P3 = ControlPoint.Position;
+    			ValueType P2 = TangentOps::TangentInToP2(P3, ControlPoint.TangentIn);  // Second point's IN tangent
+
 				// Add the first segment
-				InternalSpline.SetControlPoints( {P0, P1, P2, P3}, EParameterizationPolicy::Centripetal);
+				InternalSpline.SetControlPoints( {P0, P1, P2, P3}, EParameterizationPolicy::Uniform);
 				
 				// update tangent mode
 				TangentModes.AddDefaulted();
@@ -1768,7 +1779,16 @@ private:
     	
     	return false;
     }
-
+	void AssignInternalSpline(const FPolyBezierSpline& NewSpline)
+	{
+		InternalSpline = NewSpline;
+		InstallFindWindowOverride();
+	}
+	void AssignInternalSpline(FPolyBezierSpline&& NewSpline)
+	{
+		InternalSpline = MoveTemp(NewSpline);
+		InstallFindWindowOverride();
+	}
 protected:
 	/* Tension parameter for auto-computed tangents [0,1] */
 	float Tension = 0.0f;
@@ -1777,6 +1797,10 @@ protected:
 	TArray<ETangentMode> TangentModes;
 
 	bool bStationaryEndpoint = false;
+	
+	TStaticArray<ValueType, TPolyBezierSpline<VALUETYPE>::Base::WindowSize> IntermediateWindow;
+private:
+	using TSplineWrapper<TPolyBezierSpline<VALUETYPE>>::InternalSpline;
 };
 
 using FTangentBezierSpline3f = TTangentBezierSpline<FVector3f>; 
diff --git a/Engine/Source/Runtime/GeometryCore/Public/Splines/TangentsPolicy.h b/Engine/Source/Runtime/GeometryCore/Public/Splines/TangentsPolicy.h
new file mode 100644
--- /dev/null
+++ b/Engine/Source/Runtime/GeometryCore/Public/Splines/TangentsPolicy.h
@@ -0,0 +1,203 @@
+﻿// Copyright Epic Games, Inc. All Rights Reserved.
+
+#pragma once
+
+#include "CoreMinimal.h"
+#include "SplineInterfaces.h"
+#include "ParameterizedTypes.h"
+#include <cstdint>
+#include <limits>
+
+namespace UE
+{
+namespace Geometry
+{
+namespace Spline
+{
+namespace Math
+{
+	/** Computes Tangent for a curve segment */
+	template< class T, class U > 
+	inline void AutoCalcTangent( const T& PrevP, const T& P, const T& NextP, const U& Tension, T& OutTan )
+	{
+		OutTan = (1.f - Tension) * ( (P - PrevP) + (NextP - P) );
+	}
+
+
+	/**
+	 * This actually returns the control point not a tangent. This is expected by the CubicInterp function for Quaternions
+	 */
+	template< class U > 
+	inline void AutoCalcTangent( const FQuat& PrevP, const FQuat& P, const FQuat& NextP, const U& Tension, FQuat& OutTan  )
+	{
+		FQuat::CalcTangents(PrevP, P, NextP, Tension, OutTan);
+	}
+	
+	/**
+	 * Opt-in trait for types that are contiguous floats
+	 * Only types with this trait can use component-wise clamping
+	 */
+	template<typename T>
+	struct TIsContiguousFloat : std::false_type {};
+
+	template<> struct TIsContiguousFloat<float>       : std::true_type {};
+	template<> struct TIsContiguousFloat<double>      : std::true_type {};
+	template<> struct TIsContiguousFloat<FVector2D>   : std::true_type {};
+	template<> struct TIsContiguousFloat<FVector>     : std::true_type {};
+	template<> struct TIsContiguousFloat<FVector4>    : std::true_type {};
+	template<> struct TIsContiguousFloat<FTwoVectors> : std::true_type {};
+	// Add more as needed: FVector4f, FLinearColor, FPlane, etc.
+
+	/**
+	 * Opt-in trait for quaternion types
+	 */
+	template<typename T>
+	struct TIsQuaternionType : std::false_type {};
+
+	template<> struct TIsQuaternionType<TQuat<float>>  : std::true_type {};
+	template<> struct TIsQuaternionType<TQuat<double>> : std::true_type {};
+
+	/**
+	 * Per-component clamp path (matches InterpCurvePoint's ComputeClampableFloatVectorCurveTangent)
+	 * NOTE: clamped branch does NOT divide by (NextTime - PrevTime); unclamped does.
+	 */
+	template<typename T>
+	static std::enable_if_t<TIsContiguousFloat<T>::value, void>
+	ComputeClampableTangent(
+		float PrevTime, const T& PrevPoint,
+		float CurTime,  const T& CurPoint,
+		float NextTime, const T& NextPoint,
+		float Tension,
+		bool  bWantClamping,
+		T&    OutTangent)
+		{
+			if (bWantClamping)
+			{
+				static_assert(std::is_trivially_copyable_v<T>, "Type must be trivially copyable.");
+				static_assert(sizeof(T) % sizeof(float) == 0, "Type size must be a multiple of float.");
+
+				const float* Prev = reinterpret_cast<const float*>(&PrevPoint);
+				const float* Cur  = reinterpret_cast<const float*>(&CurPoint);
+				const float* Next = reinterpret_cast<const float*>(&NextPoint);
+				float*       Out  = reinterpret_cast<float*>(&OutTangent);
+
+				constexpr int32 N = int32(sizeof(T) / sizeof(float));
+				for (int32 i = 0; i < N; ++i)
+				{
+					const float Clamped = ClampFloatTangent(
+						Prev[i], PrevTime,
+						Cur[i],  CurTime,
+						Next[i], NextTime);
+
+					Out[i] = (1.0f - Tension) * Clamped;
+				}
+			}
+			else
+			{
+				AutoCalcTangent(PrevPoint, CurPoint, NextPoint, Tension, OutTangent);
+				const float PrevToNext = FMath::Max(UE_KINDA_SMALL_NUMBER, NextTime - PrevTime);
+				OutTangent /= PrevToNext;
+			}
+		}
+	template<typename T> struct TTangentOps;
+
+	/**
+	 * Policy for turning tangents into Bezier P1/P2 and computing auto tangents
+	 */
+	template<typename T>
+	struct TTangentOps
+	{
+	    static constexpr bool bContig = TIsContiguousFloat<T>::value;
+	    static constexpr bool bQuat   = TIsQuaternionType<T>::value;
+
+	    /** Compute auto/auto-clamped tangent */
+	    static void ComputeAutoTangent(float PrevTime, const T& Prev,
+	                                   float CurTime,  const T& Cur,
+	                                   float NextTime, const T& Next,
+	                                   float Tension, bool bClamp, T& Out)
+	    {
+	        static_assert(bContig && !bQuat, "Vector path only");
+	        ComputeClampableTangent(PrevTime, Prev, CurTime, Cur, NextTime, Next, Tension, bClamp, Out);
+	    }
+
+	    /** Turn out tangents into the cubic’s P1 control (Bezier) */
+	    static T TangentOutToP1(const T& P0, const T& Tout)
+	    {
+	        static_assert(bContig && !bQuat, "Vector path only");
+	        return P0 + (Tout / 3.0f);
+	    }
+
+		/** Turn in tangents into the cubic’s P2 control (Bezier) */
+		static T TangentInToP2(const T& P3, const T& Tin)
+	    {
+	    	static_assert(bContig && !bQuat, "Vector path only");
+	    	return P3 - (Tin  / 3.0f);
+	    }
+
+		/** P1 to tangent out */
+		static T P1ToTangentOut(const T& P0, const T& P1)
+	    {
+	    	static_assert(bContig && !bQuat, "Vector path only");
+	    	return (P1 - P0) * 3.0f;
+	    }
+		
+		/** P2 to tangent in */
+		static T P2ToTangentIn(const T& P3, const T& P2)
+	    {
+	    	static_assert(bContig && !bQuat, "Vector path only");
+	    	return (P3 - P2) * 3.0f;
+	    }
+	};
+
+	/**
+	 * Quaternion specialization: "tangents" are control quats (SQUAD controls)
+	 */
+	template<typename Scalar>
+	struct TTangentOps<UE::Math::TQuat<Scalar>>
+	{
+		using Quat = UE::Math::TQuat<Scalar>;
+	    static void ComputeAutoTangent(float PrevTime, const Quat& Prev,
+	                                   float CurTime,  const Quat& Cur,
+	                                   float NextTime, const Quat& Next,
+	                                   float Tension, bool /*bClamp*/, Quat& OutControl)
+	    {
+	        Quat P = Prev, C = Cur, N = Next;
+	    	
+	    	AutoCalcTangent(P, C, N, Tension, OutControl);
+	    	const float PrevToNextTimeDiff = FMath::Max< float >( UE_KINDA_SMALL_NUMBER, NextTime - PrevTime );
+		    
+	    	OutControl /= PrevToNextTimeDiff;
+	    }
+
+		/** Turn out tangents into the cubic’s P1 control (Bezier) */
+		static Quat TangentOutToP1(const Quat& P0, const Quat& Tout)
+	    {
+	    	// For quats, the "tangents" are already the control quats used by cubic interp
+	    	return Tout;
+	    }
+
+		/** Turn in tangents into the cubic’s P2 control (Bezier) */
+		static Quat TangentInToP2(const Quat& P3, const Quat& Tin)
+	    {
+	    	// For quats, the "tangents" are already the control quats used by cubic interp
+	    	return Tin;
+	    }
+
+		/** P1 to tangent out */
+		static Quat P1ToTangentOut(const Quat& P0, const Quat& P1)
+	    {
+	    	return P1;
+	    }
+
+		/** P2 to tangent in */
+		static Quat P2ToTangentIn(const Quat& P3, const Quat& P2)
+	    {
+	    	return P2;
+	    }
+	};
+	
+	
+} // end namespace UE::Geometry::Spline::Math
+} // end namespace UE::Geometry::Spline
+} // end namespace UE::Geometry
+} // end namespace UE
