Extract roey's deterministic parallel vertex normal method to geometrycore, as a templated MeshType / MeshAdapter

#Synced-CL 48266510

diff --git a/Engine/Source/Runtime/GeometryCore/Private/DynamicMesh/DynamicMesh3_Queries.cpp b/Engine/Source/Runtime/GeometryCore/Private/DynamicMesh/DynamicMesh3_Queries.cpp
--- a/Engine/Source/Runtime/GeometryCore/Private/DynamicMesh/DynamicMesh3_Queries.cpp
+++ b/Engine/Source/Runtime/GeometryCore/Private/DynamicMesh/DynamicMesh3_Queries.cpp
@@ -895,15 +895,7 @@ double FDynamicMesh3::GetTriInternalAngleR(int tID, int i) const
 FVector3d FDynamicMesh3::GetTriInternalAnglesR(int tID) const
 {
 	const FIndex3i& Triangle = Triangles[tID];
-	const FVector3d TV[3] = {Vertices[Triangle[0]], Vertices[Triangle[1]], Vertices[Triangle[2]]};
-	FVector3d ABhat = Normalized(TV[1] - TV[0]);
-	FVector3d BChat = Normalized(TV[2] - TV[1]);
-	FVector3d AChat = Normalized(TV[2] - TV[0]);
-	return FVector3d(
-		AngleR(ABhat, AChat),	// AB vs AC angle
-		AngleR(-ABhat, BChat), // BA vs BC angle
-		AngleR(AChat, BChat) // AC vs BC angle (same as CA vs CB, so don't need to negate both)
-		);
+	return VectorUtil::TriangleInternalAngles(Vertices[Triangle[0]], Vertices[Triangle[1]], Vertices[Triangle[2]]);
 }
 
 
diff --git a/Engine/Source/Runtime/GeometryCore/Public/MeshAdapter/MeshVertexNormals.h b/Engine/Source/Runtime/GeometryCore/Public/MeshAdapter/MeshVertexNormals.h
new file mode 100644
--- /dev/null
+++ b/Engine/Source/Runtime/GeometryCore/Public/MeshAdapter/MeshVertexNormals.h
@@ -0,0 +1,228 @@
+// Copyright Epic Games, Inc. All Rights Reserved.
+
+#pragma once
+
+
+#include <atomic>
+
+#include "Algo/Accumulate.h"
+#include "Async/ParallelFor.h"
+#include "Containers/Array.h"
+#include "Containers/ConsumeAllMpmcQueue.h"
+#include "Experimental/ConcurrentLinearAllocator.h"
+#include "IndexTypes.h"
+#include "ProfilingDebugging/ScopedTimers.h"
+#include "Templates/Tuple.h"
+#include "VectorUtil.h"
+
+namespace UE
+{
+namespace Geometry
+{
+
+namespace MeshNormals
+{
+
+	// helper methods for the below parallel normal compute methods
+	namespace ParallelHelpers
+	{
+		// This is a workaround for some platforms not supporting C++20's fetch_add for floats
+		// UE platforms should be on C++20, so hopefully we can remove this method in the future and just use .fetch_add directly
+		template<typename AtomicFloatType = std::atomic<float>>
+		inline void AtomicFloatFetchAdd(AtomicFloatType& Value, float ToAdd)
+		{
+			constexpr bool bHasAtomicFloatFetchAdd = requires(AtomicFloatType & AtomicFloat, float Param)
+			{
+				AtomicFloat.fetch_add(Param, std::memory_order_relaxed);
+			};
+			if constexpr (bHasAtomicFloatFetchAdd)
+			{
+				Value.fetch_add(ToAdd, std::memory_order_relaxed);
+			}
+			else
+			{
+				// Manual implementation of fetch_add
+				float Old = Value.load();
+				while (!Value.compare_exchange_weak(Old, Old + ToAdd, std::memory_order_relaxed)) {}
+			}
+		}
+
+		// This is a workaround for some platforms not supporting C++20's std::atomic_ref
+		// UE platforms should be on C++20, so hopefully we can remove this method in the future and just use std::atomic_ref directly
+		// Atomic_Ref is a slight performance improvement for platforms with weak memory order such as ARM as we can avoid a full barrier with
+		// relaxed ordering here.
+		template<typename IntType = int>
+		inline void AtomicIncrement(IntType& Value, std::memory_order MemoryOrder)
+		{
+#if defined(__cpp_lib_atomic_ref)
+			std::atomic_ref<IntType>(Value).fetch_add(1, MemoryOrder);
+#else
+			FPlatformAtomics::InterlockedIncrement(&Value);
+#endif
+		}
+
+		// This is a workaround for some platforms not supporting C++20's std::atomic_ref
+		// UE platforms should be on C++20, so hopefully we can remove this method in the future and just use std::atomic_ref directly
+		// Atomic_Ref is a slight performance improvement for platforms with weak memory order such as ARM as we can avoid a full barrier with
+		// relaxed ordering here.
+		template<typename IntType = int>
+		inline void AtomicDecrement(IntType& Value, std::memory_order MemoryOrder)
+		{
+#if defined(__cpp_lib_atomic_ref)
+			std::atomic_ref<IntType>(Value).fetch_sub(1, MemoryOrder);
+#else
+			FPlatformAtomics::InterlockedDecrement(&Value);
+#endif
+		}
+
+		struct FVertexNormalDeviceBlockAllocationTag : FDefaultBlockAllocationTag
+		{
+			static constexpr const char* TagName = "GeometryVertexNormalBlock";
+
+			using Allocator = FAlignedAllocator;
+		};
+
+		struct FVertexNormalLinearAllocator
+		{
+			FORCEINLINE static void* Malloc(SIZE_T Size, uint32 Alignment)
+			{
+				return TConcurrentLinearAllocator<FVertexNormalDeviceBlockAllocationTag>::Malloc(Size, Alignment);
+			}
+
+			FORCEINLINE static void Free(void* Pointer)
+			{
+				TConcurrentLinearAllocator<FVertexNormalDeviceBlockAllocationTag>::Free(Pointer);
+			}
+		};
+	}
+
+	/**
+	 * Compute per-vertex smooth normals for a generic triangle mesh type, writing to a *pre-allocated* output normals array
+	 * 
+	 * @param Mesh Input mesh
+	 * @param OutNormals Output normals array; must have Mesh.MaxVertexID() elements
+	 * @param bRequireDeterministicNormals If false, uses a non-deterministic fast path with unordered float operations that can produce slightly different roundings
+	 */
+	template<typename MeshType, typename NormalsArrayType>
+	void ComputeVertexNormals(const MeshType& Mesh, NormalsArrayType& OutNormals, bool bRequireDeterministicNormals = true)
+	{
+		TRACE_CPUPROFILER_EVENT_SCOPE(GeometryCore::MeshNormals::RecomputeNormals);
+		check(Mesh.MaxVertexID() == OutNormals.Num());
+
+		if (bRequireDeterministicNormals)
+		{
+			TArray<UE::TConsumeAllMpmcQueue<TPair<int, FVector3f>, ParallelHelpers::FVertexNormalLinearAllocator>> VertexNormalsToSum;
+			VertexNormalsToSum.SetNum(Mesh.MaxVertexID());
+
+			{
+				TRACE_CPUPROFILER_EVENT_SCOPE(GeometryCore::MeshNormals::Tri_Normal_Parallel);
+				ParallelFor(Mesh.MaxTriangleID(), [&](int TriangleID)
+				{
+					if (!Mesh.IsTriangle(TriangleID))
+					{
+						return;
+					}
+
+					const FIndex3i Triangle = Mesh.GetTriangle(TriangleID);
+					FVector3d A, B, C;
+					Mesh.GetTriVertices(TriangleID, A, B, C);
+					double TriArea;
+					const FVector3f TriNormal = FVector3f(VectorUtil::NormalArea(A, B, C, TriArea));
+					FVector3f TriWeights = FVector3f(VectorUtil::TriangleInternalAngles(A, B, C) * TriArea);
+
+					for (int VertIndex = 0; VertIndex < 3; ++VertIndex)
+					{
+						const int VertexID = Triangle[VertIndex];
+						const int ElementID = 3 * VertexID;
+
+						const float VertexWeight = TriWeights[VertIndex];
+						const FVector3f Normal = TriNormal * VertexWeight;
+
+						VertexNormalsToSum[VertexID].ProduceItem(TPair<int, FVector3f>(TriangleID, Normal));
+					}
+				});
+			}
+			{
+				TRACE_CPUPROFILER_EVENT_SCOPE(GeometryCore::MeshNormals::Average);
+				ParallelFor(Mesh.MaxVertexID(), [&](int VertexID)
+				{
+					if (!Mesh.IsVertex(VertexID))
+					{
+						return;
+					}
+
+					TArray<TPair<int, FVector3f>, TInlineAllocator<8>> VertexNormals;
+					VertexNormalsToSum[VertexID].ConsumeAllLifo([&](TPair<int, FVector3f>&& Item) { VertexNormals.Add(Item); });
+					VertexNormals.Sort([](const TPair<int, FVector3f>& A, const TPair<int, FVector3f>& B)
+					{
+						return A.Key < B.Key;
+					});
+
+					const FVector3f Normal = Algo::Accumulate(VertexNormals, FVector3f::ZeroVector,
+						[](FVector3f Result, const TPair<int, FVector3f>& Normal) -> FVector3f
+						{
+							return Result + Normal.Value;
+						}
+					);
+
+					OutNormals[VertexID] = FVector3f(Normalized(Normal));
+				});
+			}
+		}
+		else // Non-deterministic fast path which uses unordered float operations which can produce slightly different roundings
+		{
+			TArray<std::atomic<float>> VertexNormals;
+			VertexNormals.SetNumZeroed(Mesh.MaxVertexID() * 3);
+
+			{
+				TRACE_CPUPROFILER_EVENT_SCOPE(GeometryCore::MeshNormals::Tri_Normal_Parallel);
+				ParallelFor(Mesh.MaxTriangleID(), [&](int TriangleID)
+				{
+					if (!Mesh.IsTriangle(TriangleID))
+					{
+						return;
+					}
+
+					const FIndex3i Triangle = Mesh.GetTriangle(TriangleID);
+					FVector3d A, B, C;
+					Mesh.GetTriVertices(TriangleID, A, B, C);
+					double TriArea;
+					const FVector3f TriNormal = FVector3f(VectorUtil::NormalArea(A, B, C, TriArea));
+
+					FVector3f TriWeights = FVector3f(VectorUtil::TriangleInternalAngles(A, B, C) * TriArea);
+
+					for (int VertIndex = 0; VertIndex < 3; ++VertIndex)
+					{
+						const int VertexID = Triangle[VertIndex];
+						const int ElementID = 3 * VertexID;
+
+						const float VertexWeight = TriWeights[VertIndex];
+						const FVector3f Normal = TriNormal * VertexWeight;
+
+						for (int FloatIndex = 0; FloatIndex < 3; ++FloatIndex)
+						{
+							ParallelHelpers::AtomicFloatFetchAdd(VertexNormals[ElementID + FloatIndex], Normal[FloatIndex]);
+						}
+					}
+				});
+			}
+			{
+				TRACE_CPUPROFILER_EVENT_SCOPE(GeometryCore::MeshNormals::Average);
+				ParallelFor(Mesh.MaxVertexID(), [&](int VertexID)
+				{
+					if (!Mesh.IsVertex(VertexID))
+					{
+						return;
+					}
+
+					const FVector3f Normal(VertexNormals[VertexID * 3], VertexNormals[VertexID * 3 + 1], VertexNormals[VertexID * 3 + 2]);
+
+					OutNormals[VertexID] = Normalized(Normal);
+				});
+			}
+		}
+	}
+}
+
+} // end namespace UE::Geometry
+} // end namespace UE
diff --git a/Engine/Source/Runtime/GeometryCore/Public/VectorUtil.h b/Engine/Source/Runtime/GeometryCore/Public/VectorUtil.h
--- a/Engine/Source/Runtime/GeometryCore/Public/VectorUtil.h
+++ b/Engine/Source/Runtime/GeometryCore/Public/VectorUtil.h
@@ -520,6 +520,22 @@ namespace VectorUtil
 		return AngleD(A, B);
 	}
 
+	/**
+	 * @return the internal angles of triangle ABC
+	 */
+	template<typename RealType>
+	inline TVector<RealType> TriangleInternalAngles(const TVector<RealType> A, const TVector<RealType> B, const TVector<RealType>& C)
+	{
+		const TVector<RealType> ABhat = Normalized(B - A);
+		const TVector<RealType> BChat = Normalized(C - B);
+		const TVector<RealType> AChat = Normalized(C - A);
+		return FVector3d(
+			AngleR(ABhat, AChat),
+			AngleR(-ABhat, BChat),
+			AngleR(AChat, BChat)
+		);
+	}
+
 	/** 
 	 * @return oriented dihedral angle of two triangles sharing an oriented edge
 	 * 
