fix DynamicMesh CompactCopy checking the wrong mesh's attribute set for compactness
+ fix const correctness of IsCompact on attribute set

#Synced-CL 48544215

diff --git a/Engine/Source/Runtime/GeometryCore/Private/DynamicMesh/DynamicMesh3.cpp b/Engine/Source/Runtime/GeometryCore/Private/DynamicMesh/DynamicMesh3.cpp
--- a/Engine/Source/Runtime/GeometryCore/Private/DynamicMesh/DynamicMesh3.cpp
+++ b/Engine/Source/Runtime/GeometryCore/Private/DynamicMesh/DynamicMesh3.cpp
@@ -359,7 +359,7 @@ void FDynamicMesh3::AppendWithOffsets(const FDynamicMesh3& ToAppend, FAppendInfo
 
 void FDynamicMesh3::CompactCopy(const FDynamicMesh3& copy, bool bNormals, bool bColors, bool bUVs, bool bAttributes, FCompactMaps* CompactInfo)
 {
-	if (copy.IsCompact() && ((!bAttributes || !HasAttributes()) || AttributeSet->IsCompact())) {
+	if (copy.IsCompact() && ((!bAttributes || !copy.HasAttributes()) || copy.Attributes()->IsCompact())) {
 		Copy(copy, bNormals, bColors, bUVs, bAttributes);
 		if (CompactInfo)
 		{
diff --git a/Engine/Source/Runtime/GeometryCore/Private/DynamicMesh/DynamicMeshAttributeSet.cpp b/Engine/Source/Runtime/GeometryCore/Private/DynamicMesh/DynamicMeshAttributeSet.cpp
--- a/Engine/Source/Runtime/GeometryCore/Private/DynamicMesh/DynamicMeshAttributeSet.cpp
+++ b/Engine/Source/Runtime/GeometryCore/Private/DynamicMesh/DynamicMeshAttributeSet.cpp
@@ -105,7 +105,7 @@ void FDynamicMeshAttributeSet::Copy(const FDynamicMeshAttributeSet& Copy)
 }
 
 
-bool FDynamicMeshAttributeSet::IsCompact()
+bool FDynamicMeshAttributeSet::IsCompact() const
 {
 	for (int UVIdx = 0; UVIdx < NumUVLayers(); UVIdx++)
 	{
diff --git a/Engine/Source/Runtime/GeometryCore/Public/DynamicMesh/DynamicMeshAttributeSet.h b/Engine/Source/Runtime/GeometryCore/Public/DynamicMesh/DynamicMeshAttributeSet.h
--- a/Engine/Source/Runtime/GeometryCore/Public/DynamicMesh/DynamicMeshAttributeSet.h
+++ b/Engine/Source/Runtime/GeometryCore/Public/DynamicMesh/DynamicMeshAttributeSet.h
@@ -92,7 +92,7 @@ public:
 	GEOMETRYCORE_API void Copy(const FDynamicMeshAttributeSet& Copy);
 
 	/** returns true if the attached overlays/attributes are compact */
-	GEOMETRYCORE_API bool IsCompact();
+	GEOMETRYCORE_API bool IsCompact() const;
 
 	/**
 	 * Performs a CompactCopy of the attached overlays/attributes.
