Add more snapping to the mesh-mesh cut to avoid 'explosion' in number of vertices in difficult input cases:
- Snap embed-path edge points to vertices
- If an embed-path edge split would create a degenerate edge, try to edge flip instead (equivalent to splitting then immediately collapsing)

+ fix an issue where the Mesh Surface Point edge encoding was inconsistent -- deprecate the mis-used constructor, and add new functions w/ clearer usage
(and fix other usages)

Add unit test w/ an input that was causing the 'explosion' (generating a >60m tri mesh from ~10k tri mesh inputs, mostly of degenerate triangles)

+ other minor related cleanup of meshmeshcut

[RN] Improved handling of degenerate cases in mesh cutting algorithm used by [REDACTED_PERSON] by more aggressivly snapping to close geometry when embedding cut paths. Avoids exponential increase in triangle count in some difficult cases.
#Synced-CL 48544215

diff --git a/Engine/Plugins/Runtime/GeometryProcessing/Source/DynamicMesh/Private/Operations/GroupEdgeInserter.cpp b/Engine/Plugins/Runtime/GeometryProcessing/Source/DynamicMesh/Private/Operations/GroupEdgeInserter.cpp
--- a/Engine/Plugins/Runtime/GeometryProcessing/Source/DynamicMesh/Private/Operations/GroupEdgeInserter.cpp
+++ b/Engine/Plugins/Runtime/GeometryProcessing/Source/DynamicMesh/Private/Operations/GroupEdgeInserter.cpp
@@ -1374,7 +1374,7 @@ bool GetPlaneCutPath(const FDynamicMesh3& Mesh, int32 GroupID,
 		CurrentEdgeVertPlaneDistances[0] = (float)FVector::PointPlaneDist((FVector)Mesh.GetVertex(EdgeVids.A), CutPlaneOrigin, CutPlaneNormal);
 		CurrentEdgeVertPlaneDistances[1] = (float)FVector::PointPlaneDist((FVector)Mesh.GetVertex(EdgeVids.B), CutPlaneOrigin, CutPlaneNormal);
 
-			OutputPath.Emplace(FMeshSurfacePoint(StartPoint.ElementID, StartPoint.EdgeTValue), FDynamicMesh3::InvalidID);
+			OutputPath.Emplace(FMeshSurfacePoint::MakeEdgePoint(StartPoint.ElementID, StartPoint.EdgeTValue), FDynamicMesh3::InvalidID);
 		}
 	check(OutputPath.Num() == 1);
 
@@ -1440,7 +1440,7 @@ bool GetPlaneCutPath(const FDynamicMesh3& Mesh, int32 GroupID,
 						if (EndPoint.ElementID == TriangleEids[i])
 						{
 							// Got to the end
-							OutputPath.Emplace(FMeshSurfacePoint(EndPoint.ElementID, EndPoint.EdgeTValue), FDynamicMesh3::InvalidID);
+							OutputPath.Emplace(FMeshSurfacePoint::MakeEdgePoint(EndPoint.ElementID, EndPoint.EdgeTValue), FDynamicMesh3::InvalidID);
 							return true;
 						}
 					}
@@ -1508,7 +1508,7 @@ bool GetPlaneCutPath(const FDynamicMesh3& Mesh, int32 GroupID,
 						EdgeTValue = 1 - EdgeTValue;
 					}
 
-					if (UpdateNextPoint(FMeshSurfacePoint(Eid, EdgeTValue)))
+					if (UpdateNextPoint(FMeshSurfacePoint::MakeEdgePoint(Eid, EdgeTValue)))
 					{
 						CurrentEdgeVertPlaneDistances[0] = PlaneDistanceA;
 						CurrentEdgeVertPlaneDistances[1] = PlaneDistanceB;
@@ -1577,7 +1577,7 @@ bool GetPlaneCutPath(const FDynamicMesh3& Mesh, int32 GroupID,
 					if (EndPoint.ElementID == TriangleEids[i])
 					{
 						// Got to the end
-						OutputPath.Emplace(FMeshSurfacePoint(EndPoint.ElementID, EndPoint.EdgeTValue), FDynamicMesh3::InvalidID);
+						OutputPath.Emplace(FMeshSurfacePoint::MakeEdgePoint(EndPoint.ElementID, EndPoint.EdgeTValue), FDynamicMesh3::InvalidID);
 						return true;
 					}
 				}
@@ -1632,7 +1632,7 @@ bool GetPlaneCutPath(const FDynamicMesh3& Mesh, int32 GroupID,
 					Swap(CurrentEdgeVertPlaneDistances[0], CurrentEdgeVertPlaneDistances[1]);
 				}
 
-				OutputPath.Emplace(FMeshSurfacePoint(Eid, EdgeTValue), FDynamicMesh3::InvalidID);
+				OutputPath.Emplace(FMeshSurfacePoint::MakeEdgePoint(Eid, EdgeTValue), FDynamicMesh3::InvalidID);
 			}//end cutting through edge
 		}//end if last point was edge
 
diff --git a/Engine/Source/Runtime/GeometryCore/Private/Operations/EmbedSurfacePath.cpp b/Engine/Source/Runtime/GeometryCore/Private/Operations/EmbedSurfacePath.cpp
--- a/Engine/Source/Runtime/GeometryCore/Private/Operations/EmbedSurfacePath.cpp
+++ b/Engine/Source/Runtime/GeometryCore/Private/Operations/EmbedSurfacePath.cpp
@@ -20,6 +20,7 @@ FVector3d FMeshSurfacePoint::Pos(const FDynamicMesh3 *Mesh) const
 	{
 		FVector3d EA, EB;
 		Mesh->GetEdgeV(ElementID, EA, EB);
+		// Note this is equivalent to Lerp(EA, EB, BaryCoord[1])
 		return BaryCoord[0] * EA + BaryCoord[1] * EB;
 	}
 	else // PointType == ESurfacePointType::Triangle
@@ -30,10 +31,13 @@ FVector3d FMeshSurfacePoint::Pos(const FDynamicMesh3 *Mesh) const
 	}
 }
 
+namespace UE::EmbedSurfacePathLocals
+{
+
 /**
  * Helper function to snap a triangle surface point to the triangle vertices or edges if it's close enough.  Input SurfacePt must be a triangle.
  */
-void RefineSurfacePtFromTriangleToSubElement(const FDynamicMesh3* Mesh, FVector3d Pos, FMeshSurfacePoint& SurfacePt, double SnapElementThresholdSq)
+static void RefineSurfacePtFromTriangleToSubElement(const FDynamicMesh3* Mesh, FVector3d Pos, FMeshSurfacePoint& SurfacePt, double SnapElementThresholdSq)
 {
 	// expect this to only be called on SurfacePoints with PointType == Triangle; otherwise indicative of incorrect usage
 	if (!ensure(SurfacePt.PointType == ESurfacePointType::Triangle))
@@ -84,9 +88,7 @@ void RefineSurfacePtFromTriangleToSubElement(const FDynamicMesh3* Mesh, FVector3
 
 	if (BestSubIdx > -1)
 	{
-		SurfacePt.ElementID = TriEdgeIDs[BestSubIdx];
-		SurfacePt.PointType = ESurfacePointType::Edge;
-		SurfacePt.BaryCoord = FVector3d(BestEdgeParam, 1 - BestEdgeParam, 0);
+		SurfacePt = FMeshSurfacePoint::MakeEdgePoint(TriEdgeIDs[BestSubIdx], BestEdgeParam);
 		return;
 	}
 
@@ -95,7 +97,8 @@ void RefineSurfacePtFromTriangleToSubElement(const FDynamicMesh3* Mesh, FVector3
 
 // For when a triangle is replaced by multiple triangles, create a new surface point for the point's new location among the smaller triangles.
 //  Note input position is in the coordinate space of the mesh vertices, not e.g. a transformed space (unlike the positions passed to WalkMeshPlanar below!)
-FMeshSurfacePoint RelocateTrianglePointAfterRefinement(const FDynamicMesh3* Mesh, const FVector3d& PosInVertexCoordSpace, TArray<int> TriIDs, double SnapElementThresholdSq)
+template<typename TIterableTrisType>
+static FMeshSurfacePoint RelocateTrianglePointAfterRefinement(const FDynamicMesh3* Mesh, const FVector3d& PosInVertexCoordSpace, const TIterableTrisType& TriIDs, double SnapElementThresholdSq)
 {
 	double BestTriDistSq = 0;
 	FVector3d BestBaryCoords;
@@ -121,6 +124,24 @@ FMeshSurfacePoint RelocateTrianglePointAfterRefinement(const FDynamicMesh3* Mesh
 	return SurfacePt;
 }
 
+// Helper to snap a point to the nearest point in a C-style array
+template<int NumPoints>
+static bool SelectSnapPoint(FVector3d(&PointArray)[NumPoints], const FVector3d& Point, int32& OutIdx, double SnapThresholdSq)
+{
+	double BestDistSq = SnapThresholdSq;
+	OutIdx = INDEX_NONE;
+	for (int32 Idx = 0; Idx < NumPoints; ++Idx)
+	{
+		double DSq = FVector3d::DistSquared(Point, PointArray[Idx]);
+		if (DSq < BestDistSq)
+		{
+			OutIdx = Idx;
+			BestDistSq = DSq;
+		}
+	}
+	return OutIdx != INDEX_NONE;
+}
+
 struct FIndexDistance
 {
 	int Index;
@@ -132,7 +153,7 @@ struct FIndexDistance
 	}
 };
 
-bool WalkMeshPlanar(const FDynamicMesh3* Mesh, int StartTri, int StartVID, FVector3d StartPt, int EndTri, int EndVertID, FVector3d EndPt, FVector3d WalkPlaneNormal, TFunction<FVector3d(const FDynamicMesh3*, int)> VertexToPosnFn,
+static bool WalkMeshPlanar(const FDynamicMesh3* Mesh, int StartTri, int StartVID, FVector3d StartPt, int EndTri, int EndVertID, FVector3d EndPt, FVector3d WalkPlaneNormal, TFunction<FVector3d(const FDynamicMesh3*, int)> VertexToPosnFn,
 	bool bAllowBackwardsSearch, double AcceptEndPtOutsideDist, double PtOnPlaneThresholdSq, TArray<TPair<FMeshSurfacePoint, int>>& WalkedPath, double BackwardsTolerance)
 {
 	// Even when bAllowBackwardsSearch is false, there may be multiple paths from start point to end point. 
@@ -176,7 +197,7 @@ bool WalkMeshPlanar(const FDynamicMesh3* Mesh, int StartTri, int StartVID, FVect
 	// The actual payloads are indices into ComputedPointsAndSources.
 	TArray<FIndexDistance> UnexploredEnds;
 
-	TSet<int> ExploredTriangles, CrossedVertices;
+	TSet<int> ExploredTriangles, CrossedVertices, CrossedEdges;
 
 	// When we've found a path, the index into ComputedPointsAndSources of the best path end.
 	int BestKnownEnd = -1;
@@ -392,13 +413,13 @@ bool WalkMeshPlanar(const FDynamicMesh3* Mesh, int StartTri, int StartVID, FVect
 						// TODO: check if this is close enough to the EndPt, and if so just stop the walk here
 						continue;
 					}
+					// not allowed to go in a direction that gets us further from the destination than our initial point if backwards search not allowed
 					bool bIsForward = ForwardsDirection.Dot(CrossingP - StartPt) >= -BackwardsTolerance;
-					if (!bAllowBackwardsSearch && !bIsForward)
+					if ((bAllowBackwardsSearch || bIsForward) && !CrossedEdges.Contains(CandidateEdgeID))
 					{
-						// not allowed to go in a direction that gets us further from the destination than our initial point if backwards search not allowed
-						continue;
+						CrossedEdges.Add(CandidateEdgeID);
+						ComputedPointsAndSources.Emplace(FMeshSurfacePoint::MakeEdgePoint(CandidateEdgeID, CrossingT), FWalkIndices(CrossingP, CurrentEnd, CrossToTriID));
 					}
-					ComputedPointsAndSources.Emplace(FMeshSurfacePoint(CandidateEdgeID, CrossingT), FWalkIndices(CrossingP, CurrentEnd, CrossToTriID));
 				}
 			}
 		}
@@ -490,6 +511,9 @@ bool WalkMeshPlanar(const FDynamicMesh3* Mesh, int StartTri, int StartVID, FVect
 }
 
 
+} // UE::EmbedSurfacePathLocals
+
+
 //static int FMeshSurfacePath::FindSharedTriangle(const FDynamicMesh3* Mesh, const FMeshSurfacePoint& A, const FMeshSurfacePoint& B)
 //{
 //	if (A.PointType == B.PointType && A.PointType == ESurfacePointType::Vertex)
@@ -615,6 +639,8 @@ bool FMeshSurfacePath::AddViaPlanarWalk(
 	FVector3d EndPt, FVector3d WalkPlaneNormal, TFunction<FVector3d(const FDynamicMesh3*, int)> VertexToPosnFn,
 	bool bAllowBackwardsSearch, double AcceptEndPtOutsideDist, double PtOnPlaneThresholdSq, double BackwardsTolerance)
 {
+	using namespace UE::EmbedSurfacePathLocals;
+	
 	if (!VertexToPosnFn)
 	{
 		VertexToPosnFn = [](const FDynamicMesh3* MeshArg, int VertexID)
@@ -651,17 +677,12 @@ bool FMeshSurfacePath::AddViaPlanarWalk(
 //	return false;
 //}
 
-/**
-* Embed a surface path in mesh provided that the path only crosses vertices and edges except at the start and end, so we can add the path easily with local edge splits and possibly two triangle pokes (rather than needing general remeshing machinery)
-* Also assumes triangles are only crossed over once (except possibly  to loop around to the start triangle on the end triangle)
-* Planar walks naturally create simple paths, so this function can be used on any paths created by single planar walks.
-*
-* @param bUpdatePath Updating the Path array with the new vertices (if false, the path will no longer be valid after running this function)\
-* @param PathVertices Indices of the vertices on the path after embedding succeeds; NOTE these will not be 1:1 with the input Path
-* @return true if embedding succeeded.
-*/
-bool FMeshSurfacePath::EmbedSimplePath(bool bUpdatePath, TArray<int>& PathVertices, bool bDoNotDuplicateFirstVertexID, double SnapElementThresholdSq)
+
+
+bool FMeshSurfacePath::EmbedSimplePath(bool bUpdatePath, TArray<int>& PathVertices, bool bDoNotDuplicateFirstVertexID, double SnapElementThresholdSq, const FEmbedSimplePathSettings& Settings)
 {
+	using namespace UE::EmbedSurfacePathLocals;
+
 	// used to track where the new vertices for *this* path start; used for bDoNotDuplicateFirstVertexID
 	int32 InitialPathIdx = PathVertices.Num();
 	if (!Path.Num())
@@ -669,6 +690,115 @@ bool FMeshSurfacePath::EmbedSimplePath(bool bUpdatePath, TArray<int>& PathVertic
 		return true;
 	}
 
+
+	if (Settings.bSimplifyPathBySnapping)
+	{
+		// Helper to snap tri and edge surface points to vertices, if they're within the threshold distance
+		auto SnapToVertex = [this, SnapElementThresholdSq](FMeshSurfacePoint& Pt) -> bool
+		{
+			double CloseDSq = SnapElementThresholdSq;
+			if (Pt.PointType == ESurfacePointType::Triangle)
+			{
+				FTriangle3d Tri;
+				Mesh->GetTriVertices(Pt.ElementID, Tri.V[0], Tri.V[1], Tri.V[2]);
+				FVector3d NewPt = Tri.BarycentricPoint(Pt.BaryCoord);
+				int32 ClosestIdx = INDEX_NONE;
+				if (SelectSnapPoint(Tri.V, NewPt, ClosestIdx, SnapElementThresholdSq))
+				{
+					Pt = FMeshSurfacePoint(Mesh->GetTriangle(Pt.ElementID)[ClosestIdx]);
+					return true;
+				}
+			}
+			else if (Pt.PointType == ESurfacePointType::Edge)
+			{
+				FIndex2i EdgeV = Mesh->GetEdgeV(Pt.ElementID);
+				FVector3d EdgePos[2];
+				Mesh->GetEdgeV(Pt.ElementID, EdgePos[0], EdgePos[1]);
+				FVector3d NewPt = Pt.BaryCoord[0] * EdgePos[0] + Pt.BaryCoord[1] * EdgePos[1];
+				int32 ClosestIdx = INDEX_NONE;
+				if (SelectSnapPoint(EdgePos, NewPt, ClosestIdx, SnapElementThresholdSq))
+				{
+					Pt = FMeshSurfacePoint(EdgeV[ClosestIdx]);
+					return true;
+				}
+			}
+			return false;
+		};
+
+		int32 NumRemoved = 0;
+		TMap<int32, int32> FirstIndices;
+		FirstIndices.Reserve(Path.Num());
+		for (int32 PathIdx = 0; PathIdx < Path.Num(); ++PathIdx)
+		{
+			int32 OffsetPathIdx = PathIdx - NumRemoved;
+			if (NumRemoved > 0)
+			{
+				Path[OffsetPathIdx] = Path[PathIdx];
+			}
+
+			if (SnapToVertex(Path[OffsetPathIdx].Key))
+			{
+				int32 LastValidIdx = PathIdx - 1 - NumRemoved;
+				// if snapped point is a duplicate, can skip it
+				if (LastValidIdx >= 0 &&
+					Path[LastValidIdx].Key.PointType == ESurfacePointType::Vertex &&
+					Path[LastValidIdx].Key.ElementID == Path[OffsetPathIdx].Key.ElementID)
+				{
+					NumRemoved++;
+				}
+			}
+		}
+		if (NumRemoved > 0)
+		{
+			Path.RemoveAt(Path.Num() - NumRemoved, NumRemoved);
+		}
+	}
+
+	// Snapping on near-degenerate paths can cause the path to revisit a vertex.
+	// For paths that were intended to be straight line connections, we can skip all points between the duplicated vertex(/vertices)
+	if (Settings.bRemovePathLoops)
+	{
+		int32 NumRemoved = 0;
+		TMap<int32, int32> SeenVertices;
+		SeenVertices.Reserve(Path.Num());
+		
+		for (int32 PathIdx = 0; PathIdx < Path.Num(); ++PathIdx)
+		{
+			int32 OffsetPathIdx = PathIdx - NumRemoved;
+			if (NumRemoved > 0)
+			{
+				Path[OffsetPathIdx] = Path[PathIdx];
+			}
+
+			if (Path[OffsetPathIdx].Key.PointType == ESurfacePointType::Vertex)
+			{
+				int32* SeenBefore = SeenVertices.Find(Path[OffsetPathIdx].Key.ElementID);
+				if (SeenBefore)
+				{
+					int32 SeenBeforeVal = *SeenBefore;
+					NumRemoved += OffsetPathIdx - SeenBeforeVal;
+					ensure(OffsetPathIdx > SeenBeforeVal);
+					for (int32 WalkBack = OffsetPathIdx - 1; WalkBack > SeenBeforeVal; --WalkBack)
+					{
+						if (Path[WalkBack].Key.PointType == ESurfacePointType::Vertex)
+						{
+							SeenVertices.Remove(Path[WalkBack].Key.ElementID);
+						}
+					}
+				}
+				else
+				{
+					SeenVertices.Add(Path[OffsetPathIdx].Key.ElementID, OffsetPathIdx);
+				}
+			}
+		}
+		if (NumRemoved > 0)
+		{
+			Path.RemoveAt(Path.Num() - NumRemoved, NumRemoved);
+		}
+	}
+
+
 	int32 PathNum = Path.Num();
 	const FMeshSurfacePoint& OrigEndPt = Path[PathNum - 1].Key;
 	
@@ -686,12 +816,13 @@ bool FMeshSurfacePath::EmbedSimplePath(bool bUpdatePath, TArray<int>& PathVertic
 	
 	if (Path[0].Key.PointType == ESurfacePointType::Triangle)
 	{
-		// TODO: poke triangle, and place initial vertex
+		// poke triangle, and place initial vertex
 		FDynamicMesh3::FPokeTriangleInfo PokeInfo;
 		Mesh->PokeTriangle(Path[0].Key.ElementID, Path[0].Key.BaryCoord, PokeInfo);
 		if (EndPtUpdated.PointType == ESurfacePointType::Triangle && Path[0].Key.ElementID == EndPtUpdated.ElementID)
 		{
-			EndPtUpdated = RelocateTrianglePointAfterRefinement(Mesh, EndPtPos, { PokeInfo.NewTriangles.A, PokeInfo.NewTriangles.B, PokeInfo.OriginalTriangle }, SnapElementThresholdSq);
+			TStaticArray<int, 3> EndCandidateTris{ PokeInfo.NewTriangles.A, PokeInfo.NewTriangles.B, PokeInfo.OriginalTriangle };
+			EndPtUpdated = RelocateTrianglePointAfterRefinement(Mesh, EndPtPos, EndCandidateTris, SnapElementThresholdSq);
 		}
 		PathVertices.Add(PokeInfo.NewVertex);
 		StartProcessIdx = 1;
@@ -707,20 +838,53 @@ bool FMeshSurfacePath::EmbedSimplePath(bool bUpdatePath, TArray<int>& PathVertic
 		if (Pt.PointType == ESurfacePointType::Edge)
 		{
 			ensure(Mesh->IsEdge(Pt.ElementID));
+
+			if (Settings.bAllowEdgeFlipToCollapseTinyEdges && !PathVertices.IsEmpty() && !Mesh->IsBoundaryEdge(Pt.ElementID))
+			{
+				// If the edge we insert would be degenerate, try to flip the edge rather than splitting it
+				// Note this is equivalent to splitting the edge then immediately collapsing the split
+				FVector3d EdgeVerts[2];
+				Mesh->GetEdgeV(Pt.ElementID, EdgeVerts[0], EdgeVerts[1]);
+				FVector3d EdgePt = EdgeVerts[0] * Pt.BaryCoord[0] + EdgeVerts[1] * Pt.BaryCoord[1];
+				FVector3d LastPos = Mesh->GetVertex(PathVertices.Last());
+				if (FVector3d::DistSquared(LastPos, EdgePt) < SnapElementThresholdSq)
+				{
+					FIndex2i OppVIDs = Mesh->GetEdgeOpposingV(Pt.ElementID);
+					FVector3d OpposingVerts[2]{ Mesh->GetVertex(OppVIDs.A), Mesh->GetVertex(OppVIDs.B) };
+					// Note these Normals might be facing the opposite direction from the normals as the edge vertex order doesn't necessarily match the first tri winding
+					// but we are just testing for relative consistency in orientation, so this is ok
+					FVector3d Normals[2]{ VectorUtil::Normal(EdgeVerts[0], EdgeVerts[1], OpposingVerts[0]), VectorUtil::Normal(EdgeVerts[1], EdgeVerts[0], OpposingVerts[1]) };
+					FVector3d FlipNormals[2]{ VectorUtil::Normal(OpposingVerts[0],OpposingVerts[1],EdgeVerts[1]), VectorUtil::Normal(OpposingVerts[1], OpposingVerts[0], EdgeVerts[0]) };
+					// Only flip if (1) the original triangle are facing the same way and (2) flipping won't create a 'fold over' 
+					// If the triangle(s) are degenerate, we don't have a clear normal so this test is invalid; in this case we allow the flip,
+					// to avoid introducing even more degenerate geometry
+					if (Normals[1].IsZero() || (Normals[0].Dot(Normals[1]) >= 0 && Normals[0].Dot(FlipNormals[0]) >= 0. && Normals[0].Dot(FlipNormals[1]) >= 0))
+					{
+						// Note because EmbedSimplePath is only allowed to cross each triangle at most once (except the start/end tri)
+						// we can assume no future path points reference the edge we flip.
+
+						FDynamicMesh3::FEdgeFlipInfo FlipInfo;
+						if (Mesh->FlipEdge(Pt.ElementID, FlipInfo) == EMeshResult::Ok)
+						{
+							if (EndPtUpdated.PointType == ESurfacePointType::Triangle && FlipInfo.Triangles.Contains(EndPtUpdated.ElementID))
+							{
+								TStaticArray<int, 2> TriInds = { FlipInfo.Triangles.A, FlipInfo.Triangles.B };
+								EndPtUpdated = RelocateTrianglePointAfterRefinement(Mesh, EndPtPos, TriInds, SnapElementThresholdSq);
+							}
+							// don't add a vertex to the path in this case; elements accessible from the split edge will be accessible from the previous path vertex after the flip
+							continue;
+						}
+					}
+				}
+			}
+
 			FDynamicMesh3::FEdgeSplitInfo SplitInfo;
-			Mesh->SplitEdge(Pt.ElementID, SplitInfo, Pt.BaryCoord[0]);
+			Mesh->SplitEdge(Pt.ElementID, SplitInfo, Pt.GetEdgeSplitParam());
 			PathVertices.Add(SplitInfo.NewVertex);
 			if (EndPtUpdated.PointType == ESurfacePointType::Triangle && SplitInfo.OriginalTriangles.Contains(EndPtUpdated.ElementID))
 			{
-				TArray<int> TriInds = { EndPtUpdated.ElementID };
-				if (SplitInfo.OriginalTriangles.A == EndPtUpdated.ElementID)
-				{
-					TriInds.Add(SplitInfo.NewTriangles.A);
-				}
-				else
-				{
-					TriInds.Add(SplitInfo.NewTriangles.B);
-				}
+				TStaticArray<int, 2> TriInds { EndPtUpdated.ElementID, 
+				 SplitInfo.OriginalTriangles.A == EndPtUpdated.ElementID ? SplitInfo.NewTriangles.A : SplitInfo.NewTriangles.B };
 				EndPtUpdated = RelocateTrianglePointAfterRefinement(Mesh, EndPtPos, TriInds, SnapElementThresholdSq);
 			}
 			else if (PathIdx != PathNum - 1 && EndPtUpdated.PointType == ESurfacePointType::Edge && Pt.ElementID == EndPtUpdated.ElementID)
@@ -752,7 +916,7 @@ bool FMeshSurfacePath::EmbedSimplePath(bool bUpdatePath, TArray<int>& PathVertic
 		else if (EndPtUpdated.PointType == ESurfacePointType::Edge)
 		{
 			FDynamicMesh3::FEdgeSplitInfo SplitInfo;
-			Mesh->SplitEdge(EndPtUpdated.ElementID, SplitInfo, EndPtUpdated.BaryCoord[0]);
+			Mesh->SplitEdge(EndPtUpdated.ElementID, SplitInfo, EndPtUpdated.GetEdgeSplitParam());
 			PathVertices.Add(SplitInfo.NewVertex);
 		}
 		else
diff --git a/Engine/Source/Runtime/GeometryCore/Private/Operations/MeshMeshCut.cpp b/Engine/Source/Runtime/GeometryCore/Private/Operations/MeshMeshCut.cpp
--- a/Engine/Source/Runtime/GeometryCore/Private/Operations/MeshMeshCut.cpp
+++ b/Engine/Source/Runtime/GeometryCore/Private/Operations/MeshMeshCut.cpp
@@ -389,7 +389,7 @@ namespace MeshCut
 				else
 				{
 					EmbeddedPath.Reset();
-					if (SurfacePath.EmbedSimplePath(false, EmbeddedPath, false, SnapToleranceSq))
+					if (SurfacePath.EmbedSimplePath(false, EmbeddedPath, false, SnapToleranceSq, FMeshSurfacePath::FEmbedSimplePathSettings::WithSimplification()))
 					{
 						ensure(EmbeddedPath.Num() > 0 && EmbeddedPath[0] == PtA.ElemID);
 						if (VertexChains)
diff --git a/Engine/Source/Runtime/GeometryCore/Public/Operations/EmbedSurfacePath.h b/Engine/Source/Runtime/GeometryCore/Public/Operations/EmbedSurfacePath.h
--- a/Engine/Source/Runtime/GeometryCore/Public/Operations/EmbedSurfacePath.h
+++ b/Engine/Source/Runtime/GeometryCore/Public/Operations/EmbedSurfacePath.h
@@ -29,7 +29,10 @@ enum class ESurfacePointType
 
 struct FMeshSurfacePoint
 {
+	// The vertex, edge or triangle ID depending on PointType
 	int ElementID;
+
+	// Weights for each element in order; i.e. for an edge the surface point is at Edge.V[0] * BaryCoord[0] + Edge.V[1] * BaryCoord[1]. For a vertex, always (1,0,0)
 	FVector3d BaryCoord;
 	ESurfacePointType PointType;
 
@@ -39,12 +42,28 @@ struct FMeshSurfacePoint
 	FMeshSurfacePoint(int TriangleID, const FVector3d& BaryCoord) : ElementID(TriangleID), BaryCoord(BaryCoord), PointType(ESurfacePointType::Triangle)
 	{
 	}
+	UE_DEPRECATED(5.8, "Use the static MakeEdgePoint method; note that it takes LerpParam instead of FirstCoordWt, which is 1-FirstCoordWt")
 	FMeshSurfacePoint(int EdgeID, double FirstCoordWt) : ElementID(EdgeID), BaryCoord(FirstCoordWt, 1-FirstCoordWt, 0), PointType(ESurfacePointType::Edge)
 	{
 	}
 	FMeshSurfacePoint(int VertexID) : ElementID(VertexID), BaryCoord(1,0,0), PointType(ESurfacePointType::Vertex)
 	{
 	}
+	FMeshSurfacePoint(int InElementID, FVector3d InBaryCoord, ESurfacePointType InPointType) : ElementID(InElementID), BaryCoord(InBaryCoord), PointType(InPointType) 
+	{
+	}
+
+	static FMeshSurfacePoint MakeEdgePoint(int32 EdgeID, double LerpParam)
+	{
+		return FMeshSurfacePoint(EdgeID, FVector3d(1 - LerpParam, LerpParam, 0.), ESurfacePointType::Edge);
+	}
+
+	// Get the parameter which should be passed to a SplitEdge() call to insert this surface point on its edge. Only meaningful for edge points.
+	double GetEdgeSplitParam() const
+	{
+		checkSlow(PointType == ESurfacePointType::Edge);
+		return BaryCoord[1];
+	}
 
 	GEOMETRYCORE_API FVector3d Pos(const FDynamicMesh3 *Mesh) const;
 };
@@ -134,14 +153,48 @@ public:
 	// */
 	//bool EmbedPath(bool bUpdatePath, TArray<int>& PathVertices, TFunction<bool(const TArray<FVector2d>& Vertices, const TArray<FIndex3i>& LabelledEdges, TArray<FVector2d>& OutVertices, TArray<int32>& OutVertexMap, TArray<FIndex3i>& OutTriangles)> MeshGraphFn);
 
+	/**
+	 * Additional options for controlling simple path embedding
+	 */
+	struct FEmbedSimplePathSettings
+	{
+		// Whether to snap intermediate triange/edge path points to vertices, using the SnapElementThresholdSq
+		bool bSimplifyPathBySnapping = false;
+
+		// Especially after simplification by snapping, we could create paths that revisit the same vertex multiple times --
+		// this option will remove such loops by removing the path between the repeated vertex; i.e. Path A B C D B E becomes just A B E
+		// Note this option should not be used if embedding curved paths, as the loop may be an intentional feature in this case
+		bool bRemovePathLoops = false;
+
+		// If an edge split would create an edge smaller than the snap threshold, allow an edge flip instead.
+		// Note this is equivalent to splitting then immediately collapsing. Flip will not be performed if it would introduce a fold-over.
+		bool bAllowEdgeFlipToCollapseTinyEdges = false;
+
+		FEmbedSimplePathSettings() {}
+
+		// @return Settings with all path simplification options enabled
+		static FEmbedSimplePathSettings WithSimplification()
+		{
+			FEmbedSimplePathSettings Settings;
+			Settings.bSimplifyPathBySnapping = true;
+			Settings.bRemovePathLoops = true;
+			Settings.bAllowEdgeFlipToCollapseTinyEdges = true;
+			return Settings;
+		}
+	};
+
 	/**
 	 * Embed a surface path in mesh provided that the path only crosses vertices and edges except at the start and end, so we can add the path easily with local edge splits and possibly two triangle pokes (rather than needing general remeshing machinery)
 	 *
 	 * @param bUpdatePath Updating the Path array with the new vertices (if false, the path will no longer be valid after running this function)
 	 * @param PathVertices Indices of the vertices on the path after embedding succeeds; NOTE these will not be 1:1 with the input Path
+	 * @param bDoNotDuplicateFirstVertexID Useful if repeatedly calling EmbedSimplePath to extend a path. If true, will not add the first path vertex if it matches the last vertex of the initial, passed-in PathVertices.
+	 * @param SnapElementThresholdSq Squared distance threshold below which path vertices can be snapped to existing elements
+	 * @param Settings Additional options controlling how the path is embedded (e.g. to enable more aggressive snapping)
 	 * @return true if embedding succeeded.
 	 */
-	GEOMETRYCORE_API bool EmbedSimplePath(bool bUpdatePath, TArray<int>& PathVertices, bool bDoNotDuplicateFirstVertexID = true, double SnapElementThresholdSq = FMathf::ZeroTolerance*100);
+	GEOMETRYCORE_API bool EmbedSimplePath(bool bUpdatePath, TArray<int>& PathVertices, bool bDoNotDuplicateFirstVertexID = true, 
+		double SnapElementThresholdSq = FMathf::ZeroTolerance*100, const FEmbedSimplePathSettings& Settings = FEmbedSimplePathSettings());
 
 	// TODO: add functionality to delete 'inside' of path -- but how do we determine what is inside?  ref MeshFacesFromLoop.cs in geometry3sharp
 };
