Polyline graph:
- add way to identify an edge via an owner object (avoid the need for dynamic tag identification)
- setup tests of PolylineGraph main functions (Split, Merge, Connect)

#tests run the auto tests and added more of them
#Synced-CL 48789434

diff --git a/Engine/Source/Runtime/GeometryCore/Public/Polyline.h b/Engine/Source/Runtime/GeometryCore/Public/Polyline.h
--- a/Engine/Source/Runtime/GeometryCore/Public/Polyline.h
+++ b/Engine/Source/Runtime/GeometryCore/Public/Polyline.h
@@ -159,6 +159,14 @@ public:
 		Vertices.Append(NewVertices);
 	}
 
+	/**
+	 * Add a list of Vertices to the polyline
+	 */
+	void AppendVertices(const VectorType* VerticesPtr, int Count)
+	{
+		Vertices.Append(VerticesPtr, Count);
+	}
+
 	/**
 	 * Add a list of Vertices to the polyline
 	 */
@@ -189,6 +197,14 @@ public:
 		Vertices.RemoveAt(VertexIndex);
 	}
 
+	/**
+	* Remove multiple vertices of the polyline (existing Vertices are shifted)
+	*/
+	void RemoveVertices(int VertexIndex, int Count)
+	{
+		Vertices.RemoveAt(VertexIndex, Count);
+	}
+
 	/**
 	 * Replace the list of Vertices with a new list
 	 */
