move the polyline and polygon distancesquared logic to a common place in curveutils, and remove intermediate conversions to segment type + extra normalizations. Add option to get the closest point directly + make computing the 'segment parameter' optional.
+ add test coverage

#Synced-CL 48949184

diff --git a/Engine/Plugins/Runtime/GeometryProcessing/Source/DynamicMesh/Public/MeshConstraints.h b/Engine/Plugins/Runtime/GeometryProcessing/Source/DynamicMesh/Public/MeshConstraints.h
--- a/Engine/Plugins/Runtime/GeometryProcessing/Source/DynamicMesh/Public/MeshConstraints.h
+++ b/Engine/Plugins/Runtime/GeometryProcessing/Source/DynamicMesh/Public/MeshConstraints.h
@@ -216,10 +216,18 @@ public:
 
 	virtual FVector3d Project(const FVector3d& Point, int /*Identifier*/) override
 	{
-		int32 NearestSegmentIdx;
-		double NearestSegmentParam;
-		double DistSqr = DistanceSquared(Point, NearestSegmentIdx, NearestSegmentParam);
-		return GetSegmentPoint(NearestSegmentIdx, NearestSegmentParam);
+		int32 NearestSegmentIdx_Unused;
+		double DistSq_Unused, UnitParam_Unused;
+		FVector3d ClosestPoint;
+		if (CurveUtil::ProjectToCurve<double, FVector3d>(Vertices, false, Point, DistSq_Unused, NearestSegmentIdx_Unused, ClosestPoint, UnitParam_Unused, nullptr))
+		{
+			return ClosestPoint;
+		}
+		else
+		{
+			// Just return the input if the projection fails (i.e., if the constraint curve is empty)
+			return Point;
+		}
 	}
 };
 
diff --git a/Engine/Source/Runtime/GeometryCore/Public/Curve/CurveUtil.h b/Engine/Source/Runtime/GeometryCore/Public/Curve/CurveUtil.h
--- a/Engine/Source/Runtime/GeometryCore/Public/Curve/CurveUtil.h
+++ b/Engine/Source/Runtime/GeometryCore/Public/Curve/CurveUtil.h
@@ -6,7 +6,8 @@
 #include "SegmentTypes.h"
 #include "LineTypes.h"
 #include "BoxTypes.h"
-#include "SegmentTypes.h"
+
+#include "Math/NumericLimits.h"
 
 namespace UE
 {
@@ -164,6 +165,94 @@ using namespace UE::Math;
 		}
 		return Sum / FMathd::TwoPi;
 	}
+	
+	/**
+	 * Project ToProject to a segment connecting points SegA,SegB, and return the unit interpolation parameter
+	 * 
+	 * @param SegA First point of the segment
+	 * @param SegB Second point of the segment
+	 * @param OutSegLenSq Will be set to the squared length of the segment
+	 * @param ToProject Point to project
+	 * @param ToleranceSquared If the squared segment length is less than this, will fall back to using the segment midpoint. (Defaults to 1e-16, which is FMathd::ZeroTolerance squared.)
+	 * @return Unit interpolation parameter, s.t. Lerp(SegA, SegB, ReturnValue) will give the projected point
+	 */
+	template<typename VectorType, typename RealType = VectorType::FReal>
+	inline RealType ProjectToSegment(const VectorType& SegA, const VectorType& SegB, RealType& OutSegLenSq, const VectorType& ToProject, RealType ToleranceSquared = (RealType)1e-16)
+	{
+		VectorType ABVec = SegB - SegA;
+		OutSegLenSq = ABVec.SquaredLength();
+		return FMath::Clamp(OutSegLenSq > ToleranceSquared ? ((ToProject - SegA) | ABVec) / OutSegLenSq : .5, (RealType)0., (RealType)1.);
+	}
+
+	/**
+	 * Project the QueryPoint to the given curve, reporting squared distance, segment, and closest point
+	 * 
+	 * @param Vertices Vertices of the curve
+	 * @param bLoop Whether the curve is a loop (i.e., whether we should include a segment from the last vertex to the first)
+	 * @param QueryPoint Point to query for distance from curve
+	 * @param OutDistanceSq Squared distance to query point
+	 * @param OutNearestSegIndex Index of the nearest segment on the curve, or INDEX_NONE if the curve has no segments (has length < 2)
+	 * @param OutNearestPoint The nearest point on the curve
+	 * @param OutNearestSegmentUnitParam The 'unit parameter' along the closest segment -- i.e. the parameter from 0 to 1 that you would pass to a Lerp to get the closest point on that segment
+	 * @param OutNearestSegmentParam If non-null, will be set to the 'segment parameter' along the closest segment -- i.e. the parameter from -SegmentExtent to +SegmentExtent where 0 is the segment midpoint
+	 * @return true if a closest point was found (i.e. if the curve had vertices)
+	 */
+	template<typename RealType, typename VectorType>
+	bool ProjectToCurve(const TArrayView<const VectorType>& Vertices, bool bLoop, const VectorType& QueryPoint, RealType& OutDistanceSq, int& OutNearestSegIndex, VectorType& OutNearestPoint, RealType& OutNearestSegmentUnitParam, RealType* OutNearestSegmentParam = nullptr)
+	{
+		OutDistanceSq = TNumericLimits<RealType>::Max();
+		OutNearestSegIndex = INDEX_NONE;
+		OutNearestSegmentUnitParam = TNumericLimits<RealType>::Max();
+		if (OutNearestSegmentParam)
+		{
+			*OutNearestSegmentParam = TNumericLimits<RealType>::Max();
+		}
+		
+		// handle if we don't have a curve
+		if (Vertices.Num() < 2)
+		{
+			if (Vertices.Num() == 1)
+			{
+				OutNearestPoint = Vertices[0];
+				OutDistanceSq = VectorType::DistSquared(OutNearestPoint, QueryPoint);
+				return true;
+			}
+			else
+			{
+				return false;
+			}
+		}
+
+		auto ProjectToSegmentHelper = [&Vertices, &QueryPoint, &OutNearestSegIndex, &OutNearestPoint, &OutNearestSegmentUnitParam, OutNearestSegmentParam, &OutDistanceSq](int32 SegIdx, int32 EndIdx)
+		{
+			VectorType A = Vertices[SegIdx], B = Vertices[EndIdx];
+			RealType SegLenSq;
+			RealType LerpParam = ProjectToSegment(A, B, SegLenSq, QueryPoint);
+			VectorType ProjectedSegmentPoint = FMath::Lerp(A, B, LerpParam);
+			RealType DSq = VectorType::DistSquared(ProjectedSegmentPoint, QueryPoint);
+			if (DSq < OutDistanceSq)
+			{
+				OutDistanceSq = DSq;
+				OutNearestSegIndex = SegIdx;
+				OutNearestPoint = ProjectedSegmentPoint;
+				OutNearestSegmentUnitParam = LerpParam;
+				if (OutNearestSegmentParam)
+				{
+					// Convert LerpParam (0 to 1 range) to SegParam (+/- Extent aka Half-Length range)
+					*OutNearestSegmentParam = (LerpParam - (RealType).5) * FMath::Sqrt(SegLenSq);
+				}
+			}
+		};
+		for (int SegIdx = 0, N = Vertices.Num(); SegIdx + 1 < N; ++SegIdx)
+		{
+			ProjectToSegmentHelper(SegIdx, SegIdx + 1);
+		}
+		if (bLoop)
+		{
+			ProjectToSegmentHelper(Vertices.Num() - 1, 0);
+		}
+		return true;
+	}
 
 	/**
 	 * @return true if the given query point is inside the closed curve (in the XY plane), based on the winding integral
diff --git a/Engine/Source/Runtime/GeometryCore/Public/Polygon2.h b/Engine/Source/Runtime/GeometryCore/Public/Polygon2.h
--- a/Engine/Source/Runtime/GeometryCore/Public/Polygon2.h
+++ b/Engine/Source/Runtime/GeometryCore/Public/Polygon2.h
@@ -594,36 +594,10 @@ public:
 	 */
 	T DistanceSquared(const TVector2<T>& QueryPoint, int& NearestSegIndexOut, T& NearestSegParamOut) const
 	{
-		NearestSegIndexOut = -1;
-		NearestSegParamOut = TNumericLimits<T>::Max();
-		T dist = TNumericLimits<T>::Max();
-		int N = Vertices.Num();
-		for (int vi = 0; vi < N; ++vi) 
-		{
-			// @todo can't we just use segment function here now?
-			TSegment2<T> seg = TSegment2<T>(Vertices[vi], Vertices[(vi + 1) % N]);
-			T t = (QueryPoint - seg.Center).Dot(seg.Direction);
-			T d = TNumericLimits<T>::Max();
-			if (t >= seg.Extent)
-			{
-				d = UE::Geometry::DistanceSquared(seg.EndPoint(), QueryPoint);
-			}
-			else if (t <= -seg.Extent)
-			{
-				d = UE::Geometry::DistanceSquared(seg.StartPoint(), QueryPoint);
-			}
-			else
-			{
-				d = (seg.PointAt(t) - QueryPoint).SquaredLength();
-			}
-			if (d < dist) 
-			{
-				dist = d;
-				NearestSegIndexOut = vi;
-				NearestSegParamOut = TMathUtil<T>::Clamp(t, -seg.Extent, seg.Extent);
-			}
-		}
-		return dist;
+		TVector2<T> NearestPoint_Unused;
+		T DistSq, UnitParam_Unused;
+		CurveUtil::ProjectToCurve<T, TVector2<T>>(Vertices, true, QueryPoint, DistSq, NearestSegIndexOut, NearestPoint_Unused, UnitParam_Unused, &NearestSegParamOut);
+		return DistSq;
 	}
 
 
@@ -634,8 +608,11 @@ public:
 	 */
 	T DistanceSquared(const TVector2<T>& QueryPoint) const
 	{
-		int seg; T segt;
-		return DistanceSquared(QueryPoint, seg, segt);
+		int Seg_Unused;
+		TVector2<T> NearestPoint_Unused;
+		T DistSq, UnitParam_Unused;
+		CurveUtil::ProjectToCurve<T, TVector2<T>>(Vertices, true, QueryPoint, DistSq, Seg_Unused, NearestPoint_Unused, UnitParam_Unused, nullptr);
+		return DistSq;
 	}
 
 
diff --git a/Engine/Source/Runtime/GeometryCore/Public/Polyline.h b/Engine/Source/Runtime/GeometryCore/Public/Polyline.h
--- a/Engine/Source/Runtime/GeometryCore/Public/Polyline.h
+++ b/Engine/Source/Runtime/GeometryCore/Public/Polyline.h
@@ -4,6 +4,7 @@
 
 #include "VectorTypes.h"
 #include "SegmentTypes.h"
+#include "Curve/CurveUtil.h"
 #include "LineTypes.h"
 #include "BoxTypes.h"
 
@@ -394,36 +395,10 @@ public:
 	 */
 	T DistanceSquared(const VectorType& QueryPoint, int& NearestSegIndexOut, T& NearestSegParamOut) const
 	{
-		NearestSegIndexOut = -1;
-		NearestSegParamOut = TNumericLimits<T>::Max();
-		T dist = TNumericLimits<T>::Max();
-		int N = SegmentCount();
-		for (int vi = 0; vi < N; ++vi)
-		{
-			// @todo can't we just use segment function here now?
-			SegmentType seg = SegmentType(Vertices[vi], Vertices[vi+1]);
-			T t = (QueryPoint - seg.Center).Dot(seg.Direction);
-			T d = TNumericLimits<T>::Max();
-			if (t >= seg.Extent)
-			{
-				d = UE::Geometry::DistanceSquared(seg.EndPoint(), QueryPoint);
-			}
-			else if (t <= -seg.Extent)
-			{
-				d = UE::Geometry::DistanceSquared(seg.StartPoint(), QueryPoint);
-			}
-			else
-			{
-				d = UE::Geometry::DistanceSquared(seg.PointAt(t), QueryPoint);
-			}
-			if (d < dist)
-			{
-				dist = d;
-				NearestSegIndexOut = vi;
-				NearestSegParamOut = TMathUtil<T>::Clamp(t, -seg.Extent, seg.Extent);
-			}
-		}
-		return dist;
+		VectorType NearestPoint_Unused;
+		T DistSq, UnitParam_Unused;
+		CurveUtil::ProjectToCurve<T, VectorType>(Vertices, false, QueryPoint, DistSq, NearestSegIndexOut, NearestPoint_Unused, UnitParam_Unused, &NearestSegParamOut);
+		return DistSq;
 	}
 
 
@@ -435,8 +410,11 @@ public:
 	 */
 	T DistanceSquared(const VectorType& QueryPoint) const
 	{
-		int seg; T segt;
-		return DistanceSquared(QueryPoint, seg, segt);
+		int NearestSegIndex_Unused;
+		VectorType NearestPoint_Unused;
+		T DistSq, UnitParam_Unused;
+		CurveUtil::ProjectToCurve<T, VectorType>(Vertices, false, QueryPoint, DistSq, NearestSegIndex_Unused, NearestPoint_Unused, UnitParam_Unused, nullptr);
+		return DistSq;
 	}
 
 
