TTangentBezierSpline: Stop trying to keep auto tangents correct

#virtualized
#Synced-CL 48924104

diff --git a/Engine/Source/Runtime/GeometryCore/Public/Splines/TangentBezierSpline.h b/Engine/Source/Runtime/GeometryCore/Public/Splines/TangentBezierSpline.h
--- a/Engine/Source/Runtime/GeometryCore/Public/Splines/TangentBezierSpline.h
+++ b/Engine/Source/Runtime/GeometryCore/Public/Splines/TangentBezierSpline.h
@@ -496,10 +496,6 @@ public:
 
 		// update flag on the internal spline
 		InternalSpline.SetClosedLoopFlag(bInClosedLoop);
-	    	
-	    
-	    // Ensure auto tangents are updated after changing loop state
-	    UpdateTangents();
     }
     
 	/**
@@ -566,9 +562,6 @@ public:
 	    {
 	        TangentModes.Add(Point.TangentMode);
 	    }
-	    
-	    // Update tangents to ensure proper curve continuity
-	    UpdateTangents();
 	    	    
 	    return true;
 	}
@@ -752,17 +745,7 @@ public:
 
 		TangentModes.InsertDefaulted(InsertedPointIndex);
 		SetPointTangentMode(InsertedPointIndex, ControlPoint.TangentMode);
-				
-		// Update adjacent points if they have automatic tangents
-		if (InsertedPointIndex > 0 && IsAutoTangent(InsertedPointIndex-1))
-		{
-			UpdatePointTangents(InsertedPointIndex-1, TangentModes[InsertedPointIndex-1]);
-		}
-        
-		if (InsertedPointIndex+1 < GetNumPoints() && IsAutoTangent(InsertedPointIndex+1))
-		{
-			UpdatePointTangents(InsertedPointIndex+1, TangentModes[InsertedPointIndex+1]);
-		}
+
 		return InsertedPointIndex;
 	}
 
@@ -798,17 +781,7 @@ public:
 		
 		TangentModes.InsertDefaulted(InsertedPointIndex);
 		SetPointTangentMode(InsertedPointIndex, ControlPoint.TangentMode);
-		
-		// Update adjacent points if they have automatic tangents
-		if (InsertedPointIndex > 0 && IsAutoTangent(InsertedPointIndex-1))
-		{
-			UpdatePointTangents(InsertedPointIndex-1, TangentModes[InsertedPointIndex-1]);
-		}
-        
-		if (InsertedPointIndex + 1 < GetNumPoints() && IsAutoTangent(InsertedPointIndex+1))
-		{
-			UpdatePointTangents(InsertedPointIndex+1, TangentModes[InsertedPointIndex+1]);
-		}
+
         return InsertedPointIndex;
     }
     
@@ -865,12 +838,6 @@ public:
         
 		TangentModes.InsertDefaulted(0);
 		SetPointTangentMode(0, ControlPoint.TangentMode);
-		
-		// Update tangents for the next point if it's auto
-		if (GetNumPoints() > 0 && IsAutoTangent(1))
-		{
-			UpdatePointTangents(1, TangentModes[1]);
-		}
 	}
     /** 
 	 * Appends a control point to the end of the spline.
@@ -895,20 +862,11 @@ public:
     	{
     		return;
     	}
-		
     	
 	    // Normal case - add segment to existing spline
         ValueType P0 = GetValue(NewIndex - 1);
-	    
-		        
 		ValueType P3 = ControlPoint.Position;
     
-		// Update the previous point's tangents if needed
-		if (IsAutoTangent(NewIndex - 1))
-		{
-			UpdatePointTangents(NewIndex - 1, TangentModes[NewIndex - 1]);
-		}
-    
 		// Get its outgoing tangent (may have been updated above)
 		ValueType OutTangentPrev = GetTangentOut(NewIndex - 1);
     
@@ -924,19 +882,6 @@ public:
 		
 		TangentModes.InsertDefaulted(NewIndex);
 		SetPointTangentMode(NewIndex, ControlPoint.TangentMode);
-
-        // If this is auto mode, update tangents for adjacent points
-        if (ControlPoint.TangentMode == ETangentMode::Auto || 
-            ControlPoint.TangentMode == ETangentMode::AutoClamped ||
-            ControlPoint.TangentMode == ETangentMode::LegacyAuto || 
-			ControlPoint.TangentMode == ETangentMode::LegacyAutoClamped)
-        {
-            // Update previous point's tangents again to ensure continuity
-            if (IsAutoTangent(NewIndex - 1))
-            {
-                UpdatePointTangents(NewIndex - 1, TangentModes[NewIndex - 1]);
-            }
-        }
     }
 	
 	/**
@@ -1025,10 +970,6 @@ public:
 		{
 			TangentModes.SetNum(GetNumPoints());
 		}
-    
-		// Update tangents for adjacent points
-		UpdateTangents();
-		
     }
 
 	/** 
@@ -1089,27 +1030,6 @@ public:
             SetTangentIn(Index, ControlPoint.TangentIn);
             SetTangentOut(Index, ControlPoint.TangentOut);
         }
-        
-        // Update tangents based on tangent mode
-		if (bTangentModeChanged)
-		{
-			UpdatePointTangents(Index, ControlPoint.TangentMode);
-		}
-
-		// Update surrounding points if position changed and they're auto tangents
-		// Only update neighbors for non-quaternion splines
-        if (bPositionChanged)
-        {
-            if (Index > 0 && IsAutoTangent(Index-1))
-            {
-                UpdatePointTangents(Index-1, TangentModes[Index-1]);
-            }
-            
-            if (Index < NumPoints-1 && IsAutoTangent(Index+1))
-            {
-                UpdatePointTangents(Index+1, TangentModes[Index+1]);
-            }
-        }
     }
 
 	void SetValue(int32 Index, const ValueType& NewValue)
@@ -1303,15 +1223,20 @@ public:
 	  */
 	void SetPointTangentMode(int32 Index, ETangentMode Mode)
 	{
-		if (Index < 0 || Index > GetNumPoints())
+		const int32 NumPoints = GetNumPoints();
+
+		if (Index < 0 || Index >= NumPoints)
 		{
 			return;
 		}
+
+		// If 0 <= Index < NumPoints and TangentModes.Num() == NumPoints, then Index is a valid element in TangentModes.
+		if (!ensureAlwaysMsgf(TangentModes.Num() == NumPoints, TEXT("Incorrect number of tangent modes, resizing!")))
+		{
+			TangentModes.SetNum(NumPoints);
+		}
         
 		TangentModes[Index] = Mode;
-
-		// Update tangents based on new mode
-        UpdatePointTangents(Index, Mode);
 	}
 	
 	virtual void Clear() override
@@ -1331,16 +1256,6 @@ public:
 		const int32 ControlPointIndex = (Index == 0) ? 0 : ((Index - 1) * 4 + 3);
 		const int32 NewIndex = (InternalSpline.SetParameter(ControlPointIndex, NewParameter) + 1) / 4;
 
-		// We likely need to do better in PolyBezier's SetParameter to account for user tangents, but this works for now.
-		if (NewIndex != Index)
-		{
-			UpdateTangents();
-
-			UE_LOG(LogSpline, Verbose, TEXT("\t"));
-			UE_LOG(LogSpline, Verbose, TEXT("After Tangent Recompute:"));
-			InternalSpline.Dump();
-		}
-
 		return NewIndex;
     }
 
@@ -1365,29 +1280,37 @@ public:
 	void UpdateTangents()
 	{
 	    const int32 NumPoints = GetNumPoints();
+
 	    if (NumPoints < 2)
 	    {
 	        return;
 	    }
+
 		// Update tangents for each point according to its mode
 		for (int32 i = 0; i < NumPoints; ++i)
 		{
-			UpdatePointTangents(i, TangentModes[i]);
+			UpdateTangent(i);
 		}
 	}
 
 	/**
 	 * Updates tangents for a specific point based on its tangent mode
 	 * @param Index Index of the point to update
-	 * @param Mode Tangent mode to apply
 	 */
-	void UpdatePointTangents(int32 Index, ETangentMode Mode)
+	void UpdateTangent(int32 Index)
 	{
 		const int32 NumPoints = GetNumPoints();
 		const int32 NumSegments = GetNumberOfSegments();
-		if (Index < 0 || Index >= NumPoints || Mode == ETangentMode::User || Mode == ETangentMode::Broken)
+
+		if (Index < 0 || Index >= NumPoints || !ensureAlwaysMsgf(TangentModes.IsValidIndex(Index), TEXT("No tangent mode for point!")))
+		{
+			return;
+		}
+
+		const ETangentMode Mode = TangentModes[Index];
+
+		if (Mode == ETangentMode::User || Mode == ETangentMode::Broken)
 		{
-			// Invalid index or user-defined mode, no action needed
 			return;
 		}
 
@@ -1437,7 +1360,6 @@ public:
 		const float dTL = FMath::Max(UE_KINDA_SMALL_NUMBER, CurrentParam  - PrevParam); // span to the *left*
 		const float dTR = FMath::Max(UE_KINDA_SMALL_NUMBER, NextParam - CurrentParam ); // span to the *right*
 		
-		
 		// Calculate tangents based on mode
 		ValueType InTangent = GetTangentIn(Index); // Keep existing if not changing
 		ValueType OutTangent = GetTangentOut(Index);
@@ -1548,11 +1470,11 @@ public:
 				OutTangent = Tangent;
 			}
 			break;
-		}		
+		}
+
 		// Apply the calculated tangents
 		SetTangentIn(Index, InTangent);
 		SetTangentOut(Index, OutTangent);
-
 	}
 
 	virtual ValueType EvaluateImpl(float Parameter) const override
@@ -1574,16 +1496,11 @@ public:
 	void Reparameterize(EParameterizationPolicy Policy = EParameterizationPolicy::Uniform)
     {
 	    InternalSpline.Reparameterize(Policy);
-		
-	    // Update tangents to maintain proper curve shape
-	    UpdateTangents();
     }
 
 	void SetKnotVector(const TArray<FKnot>& InKnots)
 	{
 		InternalSpline.SetKnotVector(InKnots);
-
-		UpdateTangents();
 	}
 	
 	const TArray<FKnot>& GetKnotVector() const
@@ -1629,7 +1546,6 @@ public:
 	void SetStationaryEndpoints(bool bInStationaryEndpoints)
 	{
 		bStationaryEndpoint = bInStationaryEndpoints;
-		UpdateTangents(); // Update tangents when this flag changes
 	}
 
 	bool IsStationaryEndpoints() const
