Fix for cook difference when building chaos convex meshes with compact tset enabled.

This is for the UE::Geometry::TConvexHull3 path.
Cooked data output depends on the ordering of FHullConnectivity::TrisWithPoints.
Removal from and addition to this set, thus the ordering, depends on the input vertices.

#tests Dev editor cook
#Synced-CL 48838583

diff --git a/Engine/Source/Runtime/GeometryCore/Private/CompGeom/ConvexHull3.cpp b/Engine/Source/Runtime/GeometryCore/Private/CompGeom/ConvexHull3.cpp
--- a/Engine/Source/Runtime/GeometryCore/Private/CompGeom/ConvexHull3.cpp
+++ b/Engine/Source/Runtime/GeometryCore/Private/CompGeom/ConvexHull3.cpp
@@ -11,6 +11,7 @@
 
 #include "Algo/RemoveIf.h"
 #include "Async/ParallelFor.h"
+#include "Containers/SparseSet.h"
 
 namespace UE
 {
@@ -268,7 +269,9 @@ struct FHullConnectivity
 
 	TArray<FIndex3i> TriNeighbors;
 	TArray<FVisiblePoints> VisiblePoints;
-	TSet<int32> TrisWithPoints;
+
+	// Prefering a sparse set here to retain identical cook output before/after the compact set switch.
+	TSparseSet<int32> TrisWithPoints;
 
 	// If positive, this threshold additionally filters which points are considered 'visible' as only points at least this far from the plane
 	double VisibleDistanceThreshold = -FMathd::MaxReal;
