Add the ability to directly specify negative space that a convex decomposition should use (separately / in addition to any automatic detected navigable space)

Expose in the Dataflow leaf convex decomposition node

#Synced-CL 49122595

diff --git a/Engine/Plugins/Experimental/GeometryCollectionPlugin/Source/GeometryCollectionNodes/Private/Dataflow/GeometryCollectionUtilityNodes.cpp b/Engine/Plugins/Experimental/GeometryCollectionPlugin/Source/GeometryCollectionNodes/Private/Dataflow/GeometryCollectionUtilityNodes.cpp
--- a/Engine/Plugins/Experimental/GeometryCollectionPlugin/Source/GeometryCollectionNodes/Private/Dataflow/GeometryCollectionUtilityNodes.cpp
+++ b/Engine/Plugins/Experimental/GeometryCollectionPlugin/Source/GeometryCollectionNodes/Private/Dataflow/GeometryCollectionUtilityNodes.cpp
@@ -202,6 +202,7 @@ FCreateLeafConvexHullsDataflowNode::FCreateLeafConvexHullsDataflowNode(const UE:
 	RegisterInputConnection(&OptionalSelectionFilter);
 	RegisterInputConnection(&SimplificationDistanceThreshold).SetCanHidePin(true).SetPinIsHidden(true);
 	RegisterInputConnection(&ConvexDecompositionSettings).SetCanHidePin(true).SetPinIsHidden(true);
+	RegisterInputConnection(&RequireNegativeSpaceCovering).SetCanHidePin(true).SetPinIsHidden(true);
 	RegisterOutputConnection(&Collection);
 	RegisterOutputConnection(&SphereCovering);
 
@@ -253,6 +254,9 @@ void FCreateLeafConvexHullsDataflowNode::Evaluate(UE::Dataflow::FContext& Contex
 			LeafSettings.DecompositionSettings.NegativeSpaceMinRadius = InDecompSettings.NegativeSpaceMinRadius;
 			LeafSettings.DecompositionSettings.NegativeSpaceTolerance = InDecompSettings.NegativeSpaceTolerance;
 
+			const FDataflowSphereCovering& RequireCovering = GetValue(Context, &RequireNegativeSpaceCovering);
+			LeafSettings.DecompositionSettings.CustomEmptySpace = RequireCovering.Spheres.Num() > 0 ? &RequireCovering.Spheres : nullptr;
+
 			LeafSettings.bComputeIntersectionsBeforeHull = bComputeIntersectionsBeforeHull;
 			TArray<FGeometryCollectionConvexUtility::FSphereCoveringInfo> SphereCoverings;
 			FGeometryCollectionConvexUtility::GenerateLeafConvexHulls(*GeomCollection, bRestrictToSelection, SelectedBones, LeafSettings, &SphereCoverings);
diff --git a/Engine/Plugins/Experimental/GeometryCollectionPlugin/Source/GeometryCollectionNodes/Public/Dataflow/GeometryCollectionUtilityNodes.h b/Engine/Plugins/Experimental/GeometryCollectionPlugin/Source/GeometryCollectionNodes/Public/Dataflow/GeometryCollectionUtilityNodes.h
--- a/Engine/Plugins/Experimental/GeometryCollectionPlugin/Source/GeometryCollectionNodes/Public/Dataflow/GeometryCollectionUtilityNodes.h
+++ b/Engine/Plugins/Experimental/GeometryCollectionPlugin/Source/GeometryCollectionNodes/Public/Dataflow/GeometryCollectionUtilityNodes.h
@@ -194,6 +194,10 @@ private:
 	UPROPERTY(EditAnywhere, Category = Convex, meta = (DataflowInput))
 	FDataflowConvexDecompositionSettings ConvexDecompositionSettings;
 
+	/** Negative space that must be used. Applies in addition to any automatically-detected covering added via ConvexDecompositionSettings. */
+	UPROPERTY(meta = (DataflowInput))
+	FDataflowSphereCovering RequireNegativeSpaceCovering;
+
 	UPROPERTY(EditAnywhere, Category = "Debug Draw", meta = (ShowOnlyInnerProperties))
 	FDataflowNodeDebugDrawSettings DebugDrawRenderSettings;
 
diff --git a/Engine/Source/Runtime/Experimental/Chaos/Private/GeometryCollection/GeometryCollectionConvexUtility.cpp b/Engine/Source/Runtime/Experimental/Chaos/Private/GeometryCollection/GeometryCollectionConvexUtility.cpp
--- a/Engine/Source/Runtime/Experimental/Chaos/Private/GeometryCollection/GeometryCollectionConvexUtility.cpp
+++ b/Engine/Source/Runtime/Experimental/Chaos/Private/GeometryCollection/GeometryCollectionConvexUtility.cpp
@@ -1431,7 +1431,7 @@ void HullsFromGeometry(
 
 			if (OptionalDecompositionSettings)
 			{
-				bool bAttemptDecomposition = OptionalDecompositionSettings->MaxHullsPerGeometry > 1 || OptionalDecompositionSettings->ErrorTolerance > 0.0 || OptionalDecompositionSettings->bProtectNegativeSpace;
+				bool bAttemptDecomposition = OptionalDecompositionSettings->MaxHullsPerGeometry > 1 || OptionalDecompositionSettings->ErrorTolerance > 0.0 || OptionalDecompositionSettings->bProtectNegativeSpace || OptionalDecompositionSettings->CustomEmptySpace;
 				double InitialGeoVolume = 0.0, InitialHullVolume = 0.0;
 				if (bAttemptDecomposition && (OptionalDecompositionSettings->MinGeoVolumeToDecompose > 0.0 || OptionalDecompositionSettings->MaxGeoToHullVolumeRatioToDecompose < 1.0))
 				{
@@ -1482,7 +1482,7 @@ void HullsFromGeometry(
 					const bool bIsSolid = Decomposition.IsInputSolid();
 					Decomposition.bTreatAsSolid = bIsSolid;
 
-					if (OptionalDecompositionSettings->bProtectNegativeSpace)
+					if (OptionalDecompositionSettings->bProtectNegativeSpace || OptionalDecompositionSettings->CustomEmptySpace)
 					{
 						// Use settings tuned for navigation-driven decomposition
 
@@ -1499,7 +1499,19 @@ void HullsFromGeometry(
 						Decomposition.ConvexEdgeAngleMoreSamplesThreshold = 180;
 						Decomposition.ThickenAfterHullFailure = PreprocessOptions.ThickenInputAfterHullFailure;
 
-						Decomposition.InitializeNegativeSpace(Settings);
+						if (OptionalDecompositionSettings->bProtectNegativeSpace)
+						{
+							Decomposition.InitializeNegativeSpace(Settings);
+						}
+						if (OptionalDecompositionSettings->CustomEmptySpace)
+						{
+							// Get the matrix transform from CustomEmptySpace to bone-local space
+							FMatrix ToBoneLocal = OptionalDecompositionSettings->CustomEmptySpaceWorldTransform.ToMatrixWithScale() * GlobalTransformArray[Idx].ToInverseMatrixWithScale();
+							// conservatively scale radii in the non-uniform case, since we don't want the spheres to overlap the input geometry
+							double ToBoneLocalScale = ToBoneLocal.GetMinimumAxisScale();
+							Decomposition.AddCoveringToInitialNegativeSpace(*OptionalDecompositionSettings->CustomEmptySpace, ToBoneLocal, ToBoneLocalScale);
+						}
+
 						if (OutComputedNavigableSpheres)
 						{
 							FGeometryCollectionConvexUtility::FSphereCoveringInfo SphereInfo
diff --git a/Engine/Source/Runtime/Experimental/Chaos/Public/GeometryCollection/GeometryCollectionConvexUtility.h b/Engine/Source/Runtime/Experimental/Chaos/Public/GeometryCollection/GeometryCollectionConvexUtility.h
--- a/Engine/Source/Runtime/Experimental/Chaos/Public/GeometryCollection/GeometryCollectionConvexUtility.h
+++ b/Engine/Source/Runtime/Experimental/Chaos/Public/GeometryCollection/GeometryCollectionConvexUtility.h
@@ -224,6 +224,12 @@ public:
 
 		/** Spheres smaller than this are not included in the negative space */
 		float NegativeSpaceMinRadius = 10.0f;
+
+		/** Directly-specified empty space to use *in addition* to any that would be added via bProtectNegativeSpace. If provided, this covering will be used even if bProtectNegativeSpace is false. */
+		const UE::Geometry::FSphereCovering* CustomEmptySpace = nullptr;
+
+		/** Transform to apply to the custom empty space sphere covering, relative to the top-level/world space of the geometry collection. */
+		FTransform CustomEmptySpaceWorldTransform = FTransform::Identity;
 	};
 
 	// Settings to control how convex hulls are generated for rigid/leaf nodes, from the geometry and/or imported collision shapes
diff --git a/Engine/Source/Runtime/GeometryCore/Private/CompGeom/ConvexDecomposition3.cpp b/Engine/Source/Runtime/GeometryCore/Private/CompGeom/ConvexDecomposition3.cpp
--- a/Engine/Source/Runtime/GeometryCore/Private/CompGeom/ConvexDecomposition3.cpp
+++ b/Engine/Source/Runtime/GeometryCore/Private/CompGeom/ConvexDecomposition3.cpp
@@ -1743,6 +1743,24 @@ bool FConvexDecomposition3::InitializeNegativeSpace(const FNegativeSpaceSampleSe
 	return true;
 }
 
+bool FConvexDecomposition3::AddCoveringToInitialNegativeSpace(const FSphereCovering& Covering, FMatrix CoveringToLocalSpace, double RadiusToLocalScale)
+{
+	if (Decomposition.Num() != 1)
+	{
+		return false;
+	}
+	int32 OrigNum = NegativeSpace.Num();
+	RadiusToLocalScale /= FMath::Max(FMathd::ZeroTolerance, ResultTransform.GetScale3D().GetAbsMin());
+	CoveringToLocalSpace *= ((FTransform)ResultTransform).ToInverseMatrixWithScale();
+	for (int32 SphereIdx = 0, Num = Covering.Num(); SphereIdx < Num; ++SphereIdx)
+	{
+		FVector3d Center = CoveringToLocalSpace.TransformPosition(Covering.GetCenter(SphereIdx));
+		double Radius = Covering.GetRadius(SphereIdx) * RadiusToLocalScale;
+		NegativeSpace.AddSphere(Center, Radius);
+	}
+	InitNegativeSpaceConvexPartMapping(OrigNum);
+	return true;
+}
 
 void FConvexDecomposition3::FixHullOverlapsInNegativeSpace(double NegativeSpaceTolerance, double NegativeSpaceMinRadius)
 {
@@ -2930,19 +2948,24 @@ void FConvexDecomposition3::DeleteProximity(TArray<int32>&& ToRemoveProx, bool b
 	}
 }
 
-void FConvexDecomposition3::InitNegativeSpaceConvexPartMapping()
+void FConvexDecomposition3::InitNegativeSpaceConvexPartMapping(int32 StartSphereIndex)
 {
 	// Helper to check each sphere for overlap vs each convex part
-	auto IteratePartsVsSpheres = [this](TFunctionRef<void(FConvexPart& Part, int32 SphereIdx, bool bOverlapsPart)> ProcessPartVsSphere)
+	auto IteratePartsVsSpheres = [this, StartSphereIndex](TFunctionRef<void(FConvexPart& Part, int32 SphereIdx, bool bOverlapsPart)> ProcessPartVsSphere)
 	{
 		for (int32 PartIdx = 0; PartIdx < Decomposition.Num(); ++PartIdx)
 		{
 			FConvexPart& Part = Decomposition[PartIdx];
-			if (!ensure(Part.OverlapsNegativeSpace.IsEmpty())) // expect this Init to only be called when there is no existing mapping
+			if (StartSphereIndex == 0 && !ensure(Part.OverlapsNegativeSpace.IsEmpty())) // If calling on all spheres, expect this Init to only be called when there is no existing mapping
 			{
 				Part.OverlapsNegativeSpace.Reset();
 			}
-			for (int32 SphereIdx = 0; SphereIdx < NegativeSpace.Num(); ++SphereIdx)
+			else if (StartSphereIndex > 0 && !Part.OverlapsNegativeSpace.IsEmpty()) // If calling on subset of spheres, expect this Init to be called only a larger StartSphereIndex than previously mapped spheres
+			{
+				// Spot-check just the last sphere to catch out incorrect usage, since spheres are generally added in increasing order
+				ensure(Part.OverlapsNegativeSpace.Last() < StartSphereIndex);
+			}
+			for (int32 SphereIdx = StartSphereIndex; SphereIdx < NegativeSpace.Num(); ++SphereIdx)
 			{
 				// Note: We assume the part is in the same coordinate space as the negative space spheres
 				bool bOverlapsPart = ConvexPartVsSphereOverlap(Part, NegativeSpace.GetCenter(SphereIdx), NegativeSpace.GetRadius(SphereIdx));
diff --git a/Engine/Source/Runtime/GeometryCore/Public/CompGeom/ConvexDecomposition3.h b/Engine/Source/Runtime/GeometryCore/Public/CompGeom/ConvexDecomposition3.h
--- a/Engine/Source/Runtime/GeometryCore/Public/CompGeom/ConvexDecomposition3.h
+++ b/Engine/Source/Runtime/GeometryCore/Public/CompGeom/ConvexDecomposition3.h
@@ -389,12 +389,21 @@ public:
 	GEOMETRYCORE_API void InitializeFromIndexMesh(TArrayView<const FVector3f> Vertices, TArrayView<const FIntVector> Faces, const FPreprocessMeshOptions& Options, int32 FaceVertexOffset = 0);
 
 	/**
-	 * Find negative space that should be protected. Uses the mesh passed on construction or to InitializeFromMesh or InitializeFromIndexMesh
+	 * Find and set a new negative space that should be protected (clearing any existing). Uses the mesh passed on construction or to InitializeFromMesh or InitializeFromIndexMesh.
 	 * @param Settings	Settings to use to find the negative space
+	 * @param RequestedSamples Optional points where we should sample the negative space, in addition to spheres automatically added via the Settings. Settings still apply to filter (bAllowSamplesInsideMesh) and to determine sphere radii.
 	 * @return			False on failure -- e.g., if there was no mesh available
 	 */
 	GEOMETRYCORE_API bool InitializeNegativeSpace(const FNegativeSpaceSampleSettings& Settings, TArrayView<const FVector3d> RequestedSamples = TArrayView<const FVector3d>());
 
+	/**
+	 * Add a sphere covering to existing negative space. Light-weight compared to InitializeNegativeSpace: Does not filter spheres that overlap or are inside the mesh geometry.
+	 * @param Covering The spheres to add.
+	 * @param CoveringToLocalSpace Matrix transform from input covering space to the local space of input convex decomposition geometry
+	 * @param RadiusToLocalScale Scale to transform sphere radii to the local space of the convex decomposition. Note spheres cannot be non-uniformly scaled in the local space.
+	 */
+	GEOMETRYCORE_API bool AddCoveringToInitialNegativeSpace(const FSphereCovering& Covering, FMatrix CoveringToLocalSpace = FMatrix::Identity, double RadiusToLocalScale = 1.);
+
 	//
 	// Settings
 	//
@@ -823,7 +832,7 @@ public:
 private:
 	
 	// Initialize mappings from Decomposition parts to NegativeSpace indices (FConvexPart::OverlapsNegativeSpace)
-	void InitNegativeSpaceConvexPartMapping();
+	void InitNegativeSpaceConvexPartMapping(int32 StartSphereIndex = 0);
 
 	// Negative space to attempt to protect; can be referenced by the Decomposition parts
 	FSphereCovering NegativeSpace;
