GeometryCore: Clean-up TDenseGrid2/3 and make them more consistent

#Synced-CL 49302611

diff --git a/Engine/Plugins/Runtime/GeometryProcessing/Source/DynamicMesh/Private/Operations/MeshBevel.cpp b/Engine/Plugins/Runtime/GeometryProcessing/Source/DynamicMesh/Private/Operations/MeshBevel.cpp
--- a/Engine/Plugins/Runtime/GeometryProcessing/Source/DynamicMesh/Private/Operations/MeshBevel.cpp
+++ b/Engine/Plugins/Runtime/GeometryProcessing/Source/DynamicMesh/Private/Operations/MeshBevel.cpp
@@ -2256,7 +2256,7 @@ void FMeshBevel::AppendEdgeQuads_Multi(FDynamicMesh3& Mesh, FBevelEdge& Edge)
 	const FQuadGridPatch* PrevQuadPatch = GetConnectedQuadStripRef(BevelEdgeVerts.A, Edge.EdgeIndex);
 	const FQuadGridPatch* NextQuadPatch = GetConnectedQuadStripRef(BevelEdgeVerts.B, Edge.EdgeIndex);
 
-	// now add columns of vertices for each inital vertex along the bevel-edge (by iterating along mesh-edges) 
+	// now add columns of vertices for each initial vertex along the bevel-edge (by iterating along mesh-edges) 
 	for (int32 k = 0; k < NumEdges; ++k)
 	{
 		FIndex2i EdgeV0 = SequentialQuadEdges[k].EdgeV0;
@@ -2749,7 +2749,7 @@ void FMeshBevel::AppendJunctionVertexPolygon_Multi(FDynamicMesh3& Mesh, FBevelVe
 	{
 		FDenseGrid2i VertexGrid;
 		VertexGrid.Resize(NumEdgeVerts, NumEdgeVerts);
-		VertexGrid.AssignAll(-1);
+		VertexGrid.Assign(-1);
 
 		// transfer vertex polygon into border of grid
 		int32 PolygonIdx = 0;
diff --git a/Engine/Plugins/Runtime/MeshModelingToolset/Source/MeshModelingTools/Private/DisplaceMeshTool.cpp b/Engine/Plugins/Runtime/MeshModelingToolset/Source/MeshModelingTools/Private/DisplaceMeshTool.cpp
--- a/Engine/Plugins/Runtime/MeshModelingToolset/Source/MeshModelingTools/Private/DisplaceMeshTool.cpp
+++ b/Engine/Plugins/Runtime/MeshModelingToolset/Source/MeshModelingTools/Private/DisplaceMeshTool.cpp
@@ -861,7 +861,7 @@ namespace DisplaceMeshToolLocals{
 			Parameters.DisplacementMap->GetPlatformData()->Mips.Num() < 1)
 		{
 			Parameters.DisplaceField = FSampledScalarField2f();
-			Parameters.DisplaceField.GridValues.AssignAll(0);
+			Parameters.DisplaceField.GridValues.Assign(0);
 			return;
 		}
 
@@ -869,7 +869,7 @@ namespace DisplaceMeshToolLocals{
 		if (!UE::AssetUtils::ReadTexture(Parameters.DisplacementMap, DisplacementMapValues, false)) 
 		{
 			Parameters.DisplaceField = FSampledScalarField2f();
-			Parameters.DisplaceField.GridValues.AssignAll(0);
+			Parameters.DisplaceField.GridValues.Assign(0);
 		}
 		else
 		{
diff --git a/Engine/Plugins/Runtime/MeshModelingToolset/Source/MeshModelingToolsEditorOnly/Private/SkeletalMesh/SkinWeightsBindingTool.cpp b/Engine/Plugins/Runtime/MeshModelingToolset/Source/MeshModelingToolsEditorOnly/Private/SkeletalMesh/SkinWeightsBindingTool.cpp
--- a/Engine/Plugins/Runtime/MeshModelingToolset/Source/MeshModelingToolsEditorOnly/Private/SkeletalMesh/SkinWeightsBindingTool.cpp
+++ b/Engine/Plugins/Runtime/MeshModelingToolset/Source/MeshModelingToolsEditorOnly/Private/SkeletalMesh/SkinWeightsBindingTool.cpp
@@ -269,9 +269,9 @@ void USkinWeightsBindingTool::Render(IToolsContextRenderAPI* RenderAPI)
 		float PDIScale = RenderAPI->GetCameraState().GetPDIScalingFactor();
 		const TDenseGrid3<FOccupancyGrid3::EDomain>& OccupancyGrid = Occupancy->GetOccupancyStateGrid(); 
 
-		for (int32 Index = 0; Index < OccupancyGrid.Size(); Index++)
+		for (int32 Index = 0; Index < OccupancyGrid.Num(); Index++)
 		{
-			FVector3i OccupancyIndex(OccupancyGrid.ToIndex(Index));
+			FVector3i OccupancyIndex(OccupancyGrid.GetCoords(Index));
 
 			const FOccupancyGrid3::EDomain Domain = OccupancyGrid[OccupancyIndex];
 			if (bShowBoundary && Domain == FOccupancyGrid3::EDomain::Boundary)
diff --git a/Engine/Source/Runtime/GeometryCore/Public/Image/ImageBuilder.h b/Engine/Source/Runtime/GeometryCore/Public/Image/ImageBuilder.h
--- a/Engine/Source/Runtime/GeometryCore/Public/Image/ImageBuilder.h
+++ b/Engine/Source/Runtime/GeometryCore/Public/Image/ImageBuilder.h
@@ -46,7 +46,7 @@ public:
 	 */
 	void Clear(const PixelType& ClearValue)
 	{
-		Image.AssignAll(ClearValue);
+		Image.Assign(ClearValue);
 	}
 
 	/**
diff --git a/Engine/Source/Runtime/GeometryCore/Public/Implicit/SweepingMeshSDF.h b/Engine/Source/Runtime/GeometryCore/Public/Implicit/SweepingMeshSDF.h
--- a/Engine/Source/Runtime/GeometryCore/Public/Implicit/SweepingMeshSDF.h
+++ b/Engine/Source/Runtime/GeometryCore/Public/Implicit/SweepingMeshSDF.h
@@ -690,9 +690,10 @@ private:
 		});
 
 		double max_dist = ExactBand * (double)GetDiagLength(DX);
-		ParallelFor(Grid.Size(), [this, &Origin, DX, NI, NJ, NK, &Distances, max_dist, upper_bound, &closest_tri](int LinearIdx)
+		check(Grid.Num() <= MAX_int32);
+		ParallelFor(static_cast<int32>(Grid.Num()), [this, &Origin, DX, NI, NJ, NK, &Distances, max_dist, upper_bound, &closest_tri](int LinearIdx)
 		{
-			FVector3i Idx = Grid.ToIndex(LinearIdx);
+			FVector3i Idx = Grid.GetCoords(LinearIdx);
 			if (Distances[Idx] == 1) {
 				int I = Idx.X, J = Idx.Y, K = Idx.Z;
 				FVector3d p((float)I * GetDim(DX, 0) + Origin[0], (float)J * GetDim(DX, 1) + Origin[1], (float)K * GetDim(DX, 2) + Origin[2]);
@@ -801,7 +802,7 @@ private:
 		// this returns the FCriticalSection to use for the given span of values
 		auto GetGridSectionLock = [this, SectionSize, &Distances, &GridSections](FVector3i CellGridIndex) -> FCriticalSection*
 		{
-			int64 CellLinearIndex = Distances.ToLinear(CellGridIndex);
+			int64 CellLinearIndex = Distances.GetIndex(CellGridIndex);
 			int64 SectionIndex = CellLinearIndex / SectionSize;
 			checkSlow(SectionIndex <= MAX_int32);
 			return &GridSections[(int32)SectionIndex];
@@ -821,7 +822,8 @@ private:
 
 		TArray<int> PendingCellQueue;
 		TArray<bool> done; 
-		done.SetNumZeroed(Distances.Size());
+		check(Distances.Num() <= MAX_int32);
+		done.SetNumZeroed(static_cast<int32>(Distances.Num()));
 
 		// this lambda combines the per-chunk queues into a single queue which we can pass to a ParallelFor
 		auto FlattenSectionQueues = [&PendingCellQueue, &SectionQueues]()
@@ -875,7 +877,7 @@ private:
 					int near_tid = Spatial->FindNearestTriangle(p, dsqr);
 					Distances[Idx] = (float)FMathd::Sqrt(dsqr);
 					closest_tri[Idx] = near_tid;
-					int idx_linear = Distances.ToLinear(Idx);
+					int idx_linear = static_cast<int32>(Distances.GetIndex(Idx));
 					if (done[idx_linear] == false)
 					{
 						done[idx_linear] = true;
@@ -902,7 +904,7 @@ private:
 				ParallelFor(PendingCellQueue.Num(), [&](int QIdx)
 				{
 					int cur_linear_index = PendingCellQueue[QIdx];
-					FVector3i cur_idx = Distances.ToIndex(cur_linear_index);
+					FVector3i cur_idx = Distances.GetCoords(cur_linear_index);
 					for (FVector3i idx_offset : IndexUtil::GridOffsets26) 
 					{
 						FVector3i nbr_idx = cur_idx + idx_offset;
@@ -910,7 +912,7 @@ private:
 						{
 							continue;
 						}
-						int nbr_linear_idx = Distances.ToLinear(nbr_idx);
+						int nbr_linear_idx = static_cast<int32>(Distances.GetIndex(nbr_idx));
 
 						// Note: this is technically unsafe to do because other threads might be writing to this memory location
 						// which could in theory mean that the value being read is garbage or stale. However the only possible values are true and false. 
diff --git a/Engine/Source/Runtime/GeometryCore/Public/Spatial/DenseGrid2.h b/Engine/Source/Runtime/GeometryCore/Public/Spatial/DenseGrid2.h
--- a/Engine/Source/Runtime/GeometryCore/Public/Spatial/DenseGrid2.h
+++ b/Engine/Source/Runtime/GeometryCore/Public/Spatial/DenseGrid2.h
@@ -1,18 +1,13 @@
 // Copyright Epic Games, Inc. All Rights Reserved.
 
-// Port of geometry3Sharp DenseGrid2
-
 #pragma once
 
 #include "CoreMinimal.h"
-#include "BoxTypes.h"
 #include "IntBoxTypes.h"
 
 #include "HAL/PlatformAtomics.h"
 
-namespace UE
-{
-namespace Geometry
+namespace UE::Geometry
 {
 
 /**
@@ -26,108 +21,156 @@ protected:
 	TArray64<ElemType> Buffer;
 
 	/** dimensions per axis */
-	int64 DimensionX, DimensionY;
+	FVector2i Dimensions;
 
 public:
 	/**
 	 * Create empty grid
 	 */
-	TDenseGrid2() : DimensionX(0), DimensionY(0)
+	TDenseGrid2() : Dimensions(0, 0)
+	{
+	}
+
+	TDenseGrid2(int32 DimX, int32 DimY)
+		: Dimensions(DimX, DimY)
 	{
+		Buffer.SetNumUninitialized(Num());
 	}
 
 	TDenseGrid2(int32 DimX, int32 DimY, ElemType InitialValue)
+		: Dimensions(DimX, DimY)
 	{
-		DimensionX = (int64)DimX;
-		DimensionY = (int64)DimY;
-		Buffer.Init(InitialValue, Size());
+		Buffer.Init(InitialValue, Num());
 	}
 
+	UE_DEPRECATED(5.8, "Please use Num instead.")
 	int64 Size() const
 	{
-		return DimensionX * DimensionY;
+		return Num();
 	}
 
-	int64 Width() const
+	int64 Num() const
 	{
-		return DimensionX;
+		return static_cast<int64>(Dimensions.X) * static_cast<int64>(Dimensions.Y);
 	}
 
-	int64 Height() const
+	int32 Width() const
 	{
-		return DimensionY;
+		return Dimensions.X;
 	}
 
-	FVector2i GetDimensions() const
+	int32 Height() const
 	{
-		return FVector2i((int32)DimensionX, (int32)DimensionY);
+		return Dimensions.Y;
+	}
+
+	const FVector2i& GetDimensions() const
+	{
+		return Dimensions;
 	}
 
 	void Resize(int32 DimX, int32 DimY, EAllowShrinking AllowShrinking = EAllowShrinking::Default)
 	{
-		DimensionX = (int64)DimX;
-		DimensionY = (int64)DimY;
-		Buffer.SetNumUninitialized( DimensionX*DimensionY, AllowShrinking);
+		Dimensions = FVector2i(DimX, DimY);
+		Buffer.SetNumUninitialized(Num(), AllowShrinking);
 	}
 
+	UE_DEPRECATED(5.8, "Please use Assign instead.")
 	void AssignAll(ElemType Value)
 	{
-		for (int64 i = 0, Num = Size(); i < Num; ++i)
+		Assign(Value);
+	}
+
+	void Assign(ElemType Value)
+	{
+		for (int64 Idx = 0, N = Num(); Idx < N; ++Idx)
 		{
-			Buffer[i] = Value;
+			Buffer[Idx] = Value;
 		}
 	}
 
-	void SetMin(const FVector2i& IJK, ElemType NewValue)
+	void SetMin(const FVector2i& XY, ElemType MinValue)
 	{
-		int64 Idx = (int64)IJK.X + (DimensionX * (int64)IJK.Y);
-		if (NewValue < Buffer[Idx])
+		const int64 Idx = GetIndex(XY);
+		if (MinValue < Buffer[Idx])
 		{
-			Buffer[Idx] = NewValue;
+			Buffer[Idx] = MinValue;
 		}
 	}
 
-	void SetMax(const FVector2i& IJK, ElemType NewValue)
+	void SetMax(const FVector2i& XY, ElemType MaxValue)
 	{
-		int64 Idx = (int64)IJK.X + (DimensionX * (int64)IJK.Y);
-		if (NewValue > Buffer[Idx])
+		const int64 Idx = GetIndex(XY);
+		if (MaxValue > Buffer[Idx])
 		{
-			Buffer[Idx] = NewValue;
+			Buffer[Idx] = MaxValue;
 		}
 	}
 
-	constexpr ElemType& operator[](int64 Idx)
+	bool IsValidIndex(int32 X, int32 Y) const
 	{
-		return Buffer[Idx];
+		return X >= 0 && Y >= 0 && X < Dimensions.X && Y < Dimensions.Y;
 	}
-	constexpr const ElemType& operator[](int64 Idx) const
+
+	bool IsValidIndex(const FVector2i& XY) const
 	{
-		return Buffer[Idx];
+		return IsValidIndex(XY.X, XY.Y);
 	}
 
-	constexpr ElemType& operator[](FVector2i Idx)
+	constexpr ElemType& operator[](int64 Index)
 	{
-		return Buffer[(int64)Idx.X + DimensionX * (int64)Idx.Y];
+		return const_cast<ElemType&>(const_cast<const TDenseGrid2*>(this)->operator[](Index));
 	}
-	constexpr const ElemType& operator[](FVector2i Idx) const
+
+	constexpr const ElemType& operator[](int64 Index) const
 	{
-		return Buffer[(int64)Idx.X + DimensionX * (int64)Idx.Y];
+		return Buffer[Index];
 	}
+
+	constexpr ElemType& operator[](FVector2i XY)
+	{
+		return const_cast<ElemType&>(const_cast<const TDenseGrid2*>(this)->operator[](XY));
+	}
+
+	constexpr const ElemType& operator[](FVector2i XY) const
+	{
+		return Buffer[GetIndex(XY)];
+	}
+
 	constexpr ElemType& At(int32 X, int32 Y)
 	{
-		return Buffer[(int64)X + DimensionX * (int64)Y];
+		return const_cast<ElemType&>(const_cast<const TDenseGrid2*>(this)->At(X, Y));
 	}
+
 	constexpr const ElemType& At(int32 X, int32 Y) const
 	{
-		return Buffer[(int64)X + DimensionX * (int64)Y];
+		return Buffer[GetIndex(X, Y)];
+	}
+
+	constexpr ElemType& At(FVector2i XY)
+	{
+		return const_cast<ElemType&>(const_cast<const TDenseGrid2*>(this)->At(XY));
+	}
+
+	constexpr const ElemType& At(FVector2i XY) const
+	{
+		return Buffer[GetIndex(XY)];
 	}
+
+	UE_DEPRECATED(5.8, "Please use the int32 overloads instead.")
 	constexpr ElemType& At(int64 X, int64 Y)
 	{
-		return Buffer[X + DimensionX * Y];
+		check(static_cast<int32>(X) <= MAX_int32);
+		check(static_cast<int32>(Y) <= MAX_int32);
+		return const_cast<ElemType&>(const_cast<const TDenseGrid2*>(this)->At(static_cast<int32>(X), static_cast<int32>(Y)));
 	}
+
+	UE_DEPRECATED(5.8, "Please use the int32 overloads instead.")
 	constexpr const ElemType& At(int64 X, int64 Y) const
 	{
-		return Buffer[X + DimensionX * Y];
+		check(static_cast<int32>(X) <= MAX_int32);
+		check(static_cast<int32>(Y) <= MAX_int32);
+		return const_cast<ElemType&>(const_cast<const TDenseGrid2*>(this)->At(static_cast<int32>(X), static_cast<int32>(Y)));
 	}
 
 	TArray64<ElemType>& GridValues()
@@ -139,43 +182,78 @@ public:
 		return Buffer;
 	}
 
-
-	void GetXPair(int32 X0, int32 Y, ElemType& AOut, ElemType& BOut) const
+	UE_DEPRECATED(5.8, "Please use Process instead.")
+	void Apply(TFunctionRef<ElemType(ElemType)> Func)
 	{
-		int64 Offset = DimensionX * (int64)Y;
-		AOut = Buffer[Offset + X0];
-		BOut = Buffer[Offset + X0 + 1];
+		for (int64 Idx = 0, N = Num(); Idx < N; ++Idx)
+		{
+			Buffer[Idx] = Func(Buffer[Idx]);
+		}
 	}
 
-	void Apply(TFunctionRef<ElemType(ElemType)> ApplyFunc)
+	/**
+	* Call a function on all values.
+	*/
+	template<typename ProcessFunc>
+	void Process(ProcessFunc Func)
 	{
-		for (int64 Idx = 0, Num = Size(); Idx < Num; Idx++)
+		for (int64 Idx = 0, N = Num(); Idx < N; ++Idx)
 		{
-			Buffer[Idx] = F(Buffer[Idx]);
+			Func(Buffer[Idx]);
 		}
 	}
 
+	/**
+	* Call a function on the value at (X,Y).
+	*/
+	template<typename ProcessFunc>
+	void Process(int32 X, int32 Y, ProcessFunc Func)
+	{
+		Func(At(X, Y));
+	}
+
+	/**
+	* Call a function on the value at (X,Y).
+	*/
+	template<typename ProcessFunc>
+	void Process(const FVector2i& XY, ProcessFunc Func)
+	{
+		Func(At(XY.X, XY.Y));
+	}
+
+	void GetXPair(int32 X0, int32 Y, ElemType& AOut, ElemType& BOut) const
+	{
+		const int64 Offset = static_cast<int64>(Dimensions.X) * Y;
+		AOut = Buffer[Offset + X0];
+		BOut = Buffer[Offset + X0 + 1];
+	}
+
 	FAxisAlignedBox2i Bounds() const
 	{
-		return FAxisAlignedBox2i({ 0, 0}, { (int32)DimensionX, (int32)DimensionY});
+		return FAxisAlignedBox2i({0, 0}, {Dimensions.X, Dimensions.Y});
 	}
+
 	FAxisAlignedBox2i BoundsInclusive() const
 	{
-		return FAxisAlignedBox2i({ 0, 0}, { (int32)(DimensionX - 1), (int32)(DimensionY - 1)});
+		return FAxisAlignedBox2i({0, 0}, {Dimensions.X - 1, Dimensions.Y - 1});
 	}
 
-	FVector2i GetCoords(int64 LinearIndex) const
+	FVector2i GetCoords(int64 Index) const
 	{
-		checkSlow(LinearIndex >= 0);
-		return FVector2i((int32)(LinearIndex % (int64)DimensionX), (int32)(LinearIndex / (int64)DimensionX));
+		checkSlow(Index >= 0);
+		const int32 X = static_cast<int32>(Index % static_cast<int64>(Dimensions.X));
+		const int32 Y = static_cast<int32>(Index / static_cast<int64>(Dimensions.X));
+		return FVector2i(X, Y);
 	}
+
 	int64 GetIndex(int32 X, int32 Y) const
 	{
-		return X + DimensionX * Y;
+		return X + static_cast<int64>(Dimensions.X) * Y;
 	}
-	int64 GetIndex(const FVector2i& IJK) const
+
+	int64 GetIndex(const FVector2i& XY) const
 	{
-		return IJK.X + DimensionX * IJK.Y;
+		return GetIndex(XY.X, XY.Y);
 	}
 };
 
@@ -189,28 +267,27 @@ typedef TDenseGrid2<int32> FDenseGrid2i;
 namespace DenseGrid
 {
 
-	inline void AtomicIncrement(FDenseGrid2i& Grid, int32 i, int32 j)
-	{
-		FPlatformAtomics::InterlockedIncrement(&Grid.At(i, j));
-	}
+inline void AtomicIncrement(FDenseGrid2i& Grid, int32 i, int32 j)
+{
+	FPlatformAtomics::InterlockedIncrement(&Grid.At(i, j));
+}
 
-	inline void AtomicDecrement(FDenseGrid2i& Grid, int32 i, int32 j)
+inline void AtomicDecrement(FDenseGrid2i& Grid, int32 i, int32 j)
+{
+	FPlatformAtomics::InterlockedDecrement(&Grid.At(i, j));
+}
+
+inline void AtomicIncDec(FDenseGrid2i& Grid, int32 i, int32 j, bool bDecrement = false)
+{
+	if (bDecrement)
 	{
 		FPlatformAtomics::InterlockedDecrement(&Grid.At(i, j));
 	}
-
-	inline void AtomicIncDec(FDenseGrid2i& Grid, int32 i, int32 j, bool bDecrement = false)
+	else
 	{
-		if (bDecrement)
-		{
-			FPlatformAtomics::InterlockedDecrement(&Grid.At(i, j));
-		}
-		else
-		{
-			FPlatformAtomics::InterlockedIncrement(&Grid.At(i, j));
-		}
+		FPlatformAtomics::InterlockedIncrement(&Grid.At(i, j));
 	}
 }
+}
 
-} // end namespace UE::Geometry
-} // end namespace UE
+}
diff --git a/Engine/Source/Runtime/GeometryCore/Public/Spatial/DenseGrid3.h b/Engine/Source/Runtime/GeometryCore/Public/Spatial/DenseGrid3.h
--- a/Engine/Source/Runtime/GeometryCore/Public/Spatial/DenseGrid3.h
+++ b/Engine/Source/Runtime/GeometryCore/Public/Spatial/DenseGrid3.h
@@ -1,29 +1,24 @@
 // Copyright Epic Games, Inc. All Rights Reserved.
 
-// Port of geometry3Sharp DenseGrid3
-
 #pragma once
 
 #include "CoreMinimal.h"
-#include "BoxTypes.h"
 #include "IntBoxTypes.h"
 
 #include "HAL/PlatformAtomics.h"
 
-namespace UE
-{
-namespace Geometry
+namespace UE::Geometry
 {
 
 /**
- * 3D dense grid of floating-point scalar values. 
+ * 3D dense grid of scalar values. 
  */
 template<typename ElemType>
 class TDenseGrid3
 {
 protected:
 	/** grid of allocated elements */
-	TArray<ElemType> Buffer;
+	TArray64<ElemType> Buffer;
 
 	/** dimensions per axis */
 	FVector3i Dimensions;
@@ -36,21 +31,42 @@ public:
 	{
 	}
 
-	TDenseGrid3(int DimX, int DimY, int DimZ, ElemType InitialValue)
+	TDenseGrid3(int32 DimX, int32 DimY, int32 DimZ)
+		: Dimensions(DimX, DimY, DimZ)
+	{
+		Buffer.SetNumUninitialized(Num());
+	}
+
+	TDenseGrid3(int32 DimX, int32 DimY, int32 DimZ, ElemType InitialValue)
+		: Dimensions(DimX, DimY, DimZ)
+	{
+		Buffer.Init(InitialValue, Num());
+	}
+
+	UE_DEPRECATED(5.8, "Please use Num instead.")
+	int64 Size() const
+	{
+		return Num();
+	}
+
+	int64 Num() const
 	{
-		Resize(DimX, DimY, DimZ);
-		Assign(InitialValue);
+		return static_cast<int64>(Dimensions.X) * static_cast<int64>(Dimensions.Y) * static_cast<int64>(Dimensions.Z);
 	}
 
-	int Size() const
+	int32 Width() const
 	{
-		return Dimensions.X * Dimensions.Y * Dimensions.Z;
+		return Dimensions.X;
 	}
 
-	bool IsValidIndex(const FVector3i& Index) const
+	int32 Height() const
 	{
-		return Index.X >= 0 && Index.Y >= 0 && Index.Z >= 0
-			&& Index.X < Dimensions.X && Index.Y < Dimensions.Y && Index.Z < Dimensions.Z;
+		return Dimensions.Y;
+	}
+
+	int32 Depth() const
+	{
+		return Dimensions.Z;
 	}
 
 	const FVector3i& GetDimensions() const
@@ -58,203 +74,245 @@ public:
 		return Dimensions;
 	}
 
-	void Resize(int DimX, int DimY, int DimZ, EAllowShrinking AllowShrinking = EAllowShrinking::Default)
+	void Resize(int32 DimX, int32 DimY, int32 DimZ, EAllowShrinking AllowShrinking = EAllowShrinking::Default)
 	{
-		check((int64)DimX * (int64)DimY * (int64)DimZ < INT_MAX);
-		Buffer.SetNumUninitialized(DimX * DimY * DimZ, AllowShrinking);
 		Dimensions = FVector3i(DimX, DimY, DimZ);
+		Buffer.SetNumUninitialized(Num(), AllowShrinking);
 	}
 
 	void Assign(ElemType Value)
 	{
-		for (int i = 0; i < Buffer.Num(); ++i)
+		for (int64 Idx = 0, N = Num(); Idx < N; ++Idx)
 		{
-			Buffer[i] = Value;
+			Buffer[Idx] = Value;
 		}
 	}
 
-	void SetMin(const FVector3i& IJK, ElemType F)
+	void SetMin(const FVector3i& XYZ, ElemType MinValue)
 	{
-		int Idx = IJK.X + Dimensions.X * (IJK.Y + Dimensions.Y * IJK.Z);
-		if (F < Buffer[Idx])
+		const int64 Idx = GetIndex(XYZ);
+		if (MinValue < Buffer[Idx])
 		{
-			Buffer[Idx] = F;
+			Buffer[Idx] = MinValue;
 		}
 	}
-	void SetMax(const FVector3i& IJK, ElemType F)
+	void SetMax(const FVector3i& XYZ, ElemType MaxValue)
 	{
-		int Idx = IJK.X + Dimensions.X * (IJK.Y + Dimensions.Y * IJK.Z);
-		if (F > Buffer[Idx])
+		const int64 Idx = GetIndex(XYZ);
+		if (MaxValue > Buffer[Idx])
 		{
-			Buffer[Idx] = F;
+			Buffer[Idx] = MaxValue;
 		}
 	}
 
-	constexpr ElemType& operator[](int Idx)
+	bool IsValidIndex(const FVector3i& XYZ) const
+	{
+		return IsValidIndex(XYZ.X, XYZ.Y, XYZ.Z);
+	}
+
+	bool IsValidIndex(int32 X, int32 Y, int32 Z) const
+	{
+		return X >= 0 && Y >= 0 && Z >= 0 && X < Dimensions.X && Y < Dimensions.Y && Z < Dimensions.Z;
+	}
+
+	constexpr ElemType& operator[](int64 Index)
 	{
-		return Buffer[Idx];
+		return const_cast<ElemType&>(const_cast<const TDenseGrid3*>(this)->operator[](Index));
 	}
-	constexpr const ElemType& operator[](int Idx) const
+
+	constexpr const ElemType& operator[](int64 Index) const
+	{
+		return Buffer[Index];
+	}
+
+	constexpr ElemType& operator[](FVector3i XYZ)
+	{
+		return const_cast<ElemType&>(const_cast<const TDenseGrid3*>(this)->operator[](XYZ));
+	}
+
+	constexpr const ElemType& operator[](FVector3i XYZ) const
 	{
-		return Buffer[Idx];
+		return Buffer[GetIndex(XYZ.X, XYZ.Y, XYZ.Z)];
 	}
 
-	constexpr ElemType& operator[](FVector3i Idx)
+	constexpr ElemType& At(int32 X, int32 Y, int32 Z)
 	{
-		return Buffer[Idx.X + Dimensions.X * (Idx.Y + Dimensions.Y * Idx.Z)];
+		return const_cast<ElemType&>(const_cast<const TDenseGrid3*>(this)->At(X, Y, Z));
 	}
-	constexpr const ElemType& operator[](FVector3i Idx) const
+
+	constexpr const ElemType& At(int32 X, int32 Y, int32 Z) const
 	{
-		return Buffer[Idx.X + Dimensions.X * (Idx.Y + Dimensions.Y * Idx.Z)];
+		return Buffer[GetIndex(X, Y, Z)];
 	}
-	constexpr ElemType& At(int I, int J, int K)
+
+	constexpr ElemType& At(FVector3i XYZ)
 	{
-		return Buffer[I + Dimensions.X * (J + Dimensions.Y * K)];
+		return const_cast<ElemType&>(const_cast<const TDenseGrid3*>(this)->At(XYZ));
 	}
-	constexpr const ElemType& At(int I, int J, int K) const
+
+	constexpr const ElemType& At(FVector3i XYZ) const
 	{
-		return Buffer[I + Dimensions.X * (J + Dimensions.Y * K)];
+		return Buffer[GetIndex(XYZ)];
 	}
 
-	TArray<ElemType>& GridValues()
+	TArray64<ElemType>& GridValues()
 	{
 		return Buffer;
 	}
-	const TArray<ElemType>& GridValues() const
+	const TArray64<ElemType>& GridValues() const
 	{
 		return Buffer;
 	}
 
-	/**
-	* @return the grid value at (X,Y,Z)
-	*/
 	ElemType GetValue(int32 X, int32 Y, int32 Z) const
 	{
-		return Buffer[X + Dimensions.X * (Y + Dimensions.Y * Z)];
+		return At(X, Y, Z);
 	}
 
-	/**
-	* @return the grid value at (X,Y,Z)
-	*/
-	ElemType GetValue(const FVector3i& CellCoords) const
+	ElemType GetValue(const FVector3i& XYZ) const
 	{
-		return Buffer[CellCoords.X + Dimensions.X * (CellCoords.Y + Dimensions.Y * CellCoords.Z)];
+		return At(XYZ);
 	}
 
-	/**
-	* Set the grid value at (X,Y,Z)
-	*/
 	void SetValue(int32 X, int32 Y, int32 Z, ElemType NewValue)
 	{
-		Buffer[X + Dimensions.X * (Y + Dimensions.Y * Z)] = NewValue;
+		At(X, Y, Z) = NewValue;
+	}
+
+	void SetValue(const FVector3i& XYZ, ElemType NewValue)
+	{
+		At(XYZ) = NewValue;
+	}
+
+	template<typename ProcessFunc>
+	UE_DEPRECATED(5.8, "Please use Process instead.")
+	void ProcessValue(int32 X, int32 Y, int32 Z, ProcessFunc Func)
+	{
+		Process(X, Y, Z, Func);
+	}
+
+	template<typename ProcessFunc>
+	UE_DEPRECATED(5.8, "Please use Process instead.")
+	void ProcessValue(const FVector3i& XYZ, ProcessFunc Func)
+	{
+		Process(XYZ, Func);
 	}
 
 	/**
-	* Set the grid value at (X,Y,Z)
+	* Call a function on all values.
 	*/
-	void SetValue(const FVector3i& CellCoords, ElemType NewValue)
+	template<typename ProcessFunc>
+	void Process(ProcessFunc Func)
 	{
-		Buffer[CellCoords.X + Dimensions.X * (CellCoords.Y + Dimensions.Y * CellCoords.Z)] = NewValue;
+		for (int64 Idx = 0, N = Num(); Idx < N; ++Idx)
+		{
+			Func(Buffer[Idx]);
+		}
 	}
 
 	/**
-	* Call an external lambda with a reference to the grid value at (X,Y,Z).
-	* Called as Func(ElemType&), so the caller can both read and write the grid cell
+	* Call a function on the value at (X,Y,Z).
 	*/
 	template<typename ProcessFunc>
-	void ProcessValue(int32 X, int32 Y, int32 Z, ProcessFunc Func)
+	void Process(int32 X, int32 Y, int32 Z, ProcessFunc Func)
 	{
-		int32 Index = X + Dimensions.X * (Y + Dimensions.Y * Z);
-		ElemType& CurValue = Buffer[Index];
-		Func(CurValue);
+		Func(At(X, Y, Z));
 	}
 
 	/**
-	* Call an external lambda with a reference to the grid value at (X,Y,Z).
-	* Called as Func(ElemType&), so the caller can both read and write the grid cell
+	* Call a function on the value at (X,Y,Z).
 	*/
 	template<typename ProcessFunc>
-	void ProcessValue(const FVector3i& CellCoords, ProcessFunc Func)
+	void Process(const FVector3i& XYZ, ProcessFunc Func)
 	{
-		int32 Index = CellCoords.X + Dimensions.X * (CellCoords.Y + Dimensions.Y * CellCoords.Z);
-		ElemType& CurValue = Buffer[Index];
-		Func(CurValue);
+		Func(At(XYZ));
 	}
 
+	void GetXPair(int32 X0, int32 Y, int32 Z, ElemType& AOut, ElemType& BOut) const
+	{
+		const int64 Offset = static_cast<int64>(Dimensions.X) * (static_cast<int64>(Y) + static_cast<int64>(Dimensions.Y) * static_cast<int64>(Z));
+		AOut = Buffer[Offset + static_cast<int64>(X0)];
+		BOut = Buffer[Offset + static_cast<int64>(X0) + 1];
+	}
 
+	FAxisAlignedBox3i Bounds() const
+	{
+		return FAxisAlignedBox3i({0, 0, 0}, {Dimensions.X, Dimensions.Y, Dimensions.Z});
+	}
 
-	void GetXPair(int X0, int Y, int Z, ElemType& AOut, ElemType& BOut) const
+	FAxisAlignedBox3i BoundsInclusive() const
 	{
-		int Offset = Dimensions.X * (Y + Dimensions.Y * Z);
-		AOut = Buffer[Offset + X0];
-		BOut = Buffer[Offset + X0 + 1];
+		return FAxisAlignedBox3i({0, 0, 0}, {Dimensions.X - 1, Dimensions.Y - 1, Dimensions.Z - 1});
 	}
 
-	void Apply(TFunctionRef<ElemType(ElemType)> F)
+	UE_DEPRECATED(5.8, "Please use GetCoords instead.")
+	FVector3i ToIndex(int32 Index) const
 	{
-		for (int Idx = 0, Num = Size(); Idx < Num; Idx++)
-		{
-			Buffer[Idx] = F(Buffer[Idx]);
-		}
+		return GetCoords(Index);
 	}
 
-	FAxisAlignedBox3i Bounds() const
+	FVector3i GetCoords(int64 Index) const
 	{
-		return FAxisAlignedBox3i({0, 0, 0},{Dimensions.X, Dimensions.Y, Dimensions.Z});
+		checkSlow(Index >= 0);
+		const int32 X = static_cast<int32>(Index % static_cast<int64>(Dimensions.X));
+		const int32 Y = static_cast<int32>((Index / static_cast<int64>(Dimensions.X)) % static_cast<int64>(Dimensions.Y));
+		const int32 Z = static_cast<int32>(Index / (static_cast<int64>(Dimensions.X) * static_cast<int64>(Dimensions.Y)));
+		return FVector3i(X, Y, Z);
 	}
-	FAxisAlignedBox3i BoundsInclusive() const
+
+	UE_DEPRECATED(5.8, "Please use GetIndex instead.")
+	int32 ToLinear(int32 X, int32 Y, int32 Z) const
 	{
-		return FAxisAlignedBox3i({0, 0, 0},{Dimensions.X-1, Dimensions.Y-1, Dimensions.Z-1});
+		return GetIndex(X, Y, Z);
 	}
 
-	FVector3i ToIndex(int Idx) const
+	UE_DEPRECATED(5.8, "Please use GetIndex instead.")
+	int32 ToLinear(const FVector3i& XYZ) const
 	{
-		int x = Idx % Dimensions.X;
-		int y = (Idx / Dimensions.X) % Dimensions.Y;
-		int z = Idx / (Dimensions.X * Dimensions.Y);
-		return FVector3i(x, y, z);
+		return GetIndex(XYZ.X, XYZ.Y, XYZ.Z);
 	}
-	int ToLinear(int X, int Y, int Z) const
+
+	int64 GetIndex(int32 X, int32 Y, int32 Z) const
 	{
-		return X + Dimensions.X * (Y + Dimensions.Y * Z);
+		return X + static_cast<int64>(Dimensions.X) * (Y + static_cast<int64>(Dimensions.Y) * Z);
 	}
-	int ToLinear(const FVector3i& IJK) const
+
+	int64 GetIndex(const FVector3i& XYZ) const
 	{
-		return IJK.X + Dimensions.X * (IJK.Y + Dimensions.Y * IJK.Z);
+		return GetIndex(XYZ.X, XYZ.Y, XYZ.Z);
 	}
 };
 
 
 typedef TDenseGrid3<float> FDenseGrid3f;
 typedef TDenseGrid3<double> FDenseGrid3d;
-typedef TDenseGrid3<int> FDenseGrid3i;
+typedef TDenseGrid3<int32> FDenseGrid3i;
 
 // additional utility functions
 namespace DenseGrid
 {
 
-	inline void AtomicIncrement(FDenseGrid3i& Grid, int i, int j, int k)
-	{
-		FPlatformAtomics::InterlockedIncrement(&Grid.At(i,j,k));
-	}
+inline void AtomicIncrement(FDenseGrid3i& Grid, int32 X, int32 Y, int32 Z)
+{
+	FPlatformAtomics::InterlockedIncrement(&Grid.At(X,Y,Z));
+}
+
+inline void AtomicDecrement(FDenseGrid3i& Grid, int32 X, int32 Y, int32 Z)
+{
+	FPlatformAtomics::InterlockedDecrement(&Grid.At(X, Y, Z));
+}
 
-	inline void AtomicDecrement(FDenseGrid3i& Grid, int i, int j, int k)
+inline void AtomicIncDec(FDenseGrid3i& Grid, int32 X, int32 Y, int32 Z, bool bDecrement = false)
+{
+	if (bDecrement)
 	{
-		FPlatformAtomics::InterlockedDecrement(&Grid.At(i, j, k));
+		FPlatformAtomics::InterlockedDecrement(&Grid.At(X, Y, Z));
 	}
-
-	inline void AtomicIncDec(FDenseGrid3i& Grid, int i, int j, int k, bool bDecrement = false)
+	else
 	{
-		if (bDecrement)
-		{
-			FPlatformAtomics::InterlockedDecrement(&Grid.At(i, j, k));
-		}
-		else
-		{
-			FPlatformAtomics::InterlockedIncrement(&Grid.At(i, j, k));
-		}
+		FPlatformAtomics::InterlockedIncrement(&Grid.At(X, Y, Z));
 	}
 }
+}
 
-} // end namespace UE::Geometry
-} // end namespace UE
+}
diff --git a/Engine/Source/Runtime/GeometryCore/Public/Spatial/MeshWindingNumberGrid.h b/Engine/Source/Runtime/GeometryCore/Public/Spatial/MeshWindingNumberGrid.h
--- a/Engine/Source/Runtime/GeometryCore/Public/Spatial/MeshWindingNumberGrid.h
+++ b/Engine/Source/Runtime/GeometryCore/Public/Spatial/MeshWindingNumberGrid.h
@@ -202,7 +202,8 @@ private:
 		// winding grid to track whether a voxel is 'new' or not. 
 		// So, using 3D bitmap instead - is updated at end of each pass.
 		TArray<bool> bits;
-		bits.SetNumZeroed(winding.Size());
+		check(winding.Num() <= MAX_int32);
+		bits.SetNumZeroed(static_cast<int32>(winding.Num()));
 
 		TArray<int> cur_front;
 		for (int32 LinearIdx = 0; LinearIdx < bits.Num(); LinearIdx++)
@@ -242,7 +243,7 @@ private:
 					return;
 				}
 
-				FVector3i ijk = winding.ToIndex(cur_front[FrontIdx]);
+				FVector3i ijk = winding.GetCoords(cur_front[FrontIdx]);
 				float val = winding[ijk];
 
 				// check 26-neighbours to see if we have a crossing in any direction
@@ -252,7 +253,8 @@ private:
 					{
 						continue;
 					}
-					int LinearNIJK = winding.ToLinear(nijk);
+					checkSlow(winding.GetIndex(nijk) <= MAX_int32);
+					int LinearNIJK = static_cast<int32>(winding.GetIndex(nijk));
 					float val2 = winding[LinearNIJK];
 					if (val2 == FMathf::MaxReal) {
 						FVector3d gx((float)nijk.X * dx + origin[0], (float)nijk.Y * dx + origin[1], (float)nijk.Z * dx + origin[2]);
@@ -309,7 +311,8 @@ private:
 		FastWinding->Build(false);
 
 		bool abort = false;
-		ParallelFor(winding.Size(), [this, origin, dx, &winding, &abort](int LinearIdx)
+		check(winding.Num() <= MAX_int32);
+		ParallelFor(static_cast<int32>(winding.Num()), [this, origin, dx, &winding, &abort](int LinearIdx)
 		{
 			if (LinearIdx % 100 == 0)
 			{
@@ -319,7 +322,7 @@ private:
 			{
 				return;
 			}
-			FVector3i ijk = winding.ToIndex(LinearIdx);
+			FVector3i ijk = winding.GetCoords(LinearIdx);
 			FVector3d gx((double)ijk.X * dx + origin[0], (double)ijk.Y * dx + origin[1], (double)ijk.Z * dx + origin[2]);
 			winding[LinearIdx] = (float)FastWinding->FastWindingNumber(gx);
 		});
diff --git a/Engine/Source/Runtime/GeometryCore/Public/Spatial/OccupancyGrid3.h b/Engine/Source/Runtime/GeometryCore/Public/Spatial/OccupancyGrid3.h
--- a/Engine/Source/Runtime/GeometryCore/Public/Spatial/OccupancyGrid3.h
+++ b/Engine/Source/Runtime/GeometryCore/Public/Spatial/OccupancyGrid3.h
@@ -60,9 +60,9 @@ struct FOccupancyGrid3
 			FVector3i(1, 1, 1),
 		};
 
-		ParallelFor(Occupancy.Size(), [&](const int32 OccupancyId)
+		ParallelFor(Occupancy.Num(), [&](const int32 OccupancyId)
 		{
-			const FVector3i OccupancyIndex(Occupancy.ToIndex(OccupancyId));
+			const FVector3i OccupancyIndex(Occupancy.GetCoords(OccupancyId));
 			const int32 Count = Algo::CountIf(CornerOffsets, [&WindingGrid, OccupancyIndex](FVector3i CornerOffset)
 			{
 				const FVector3i CornerIndex(OccupancyIndex + CornerOffset);
diff --git a/Engine/Source/Runtime/GeometryCore/Public/Spatial/SampledScalarField2.h b/Engine/Source/Runtime/GeometryCore/Public/Spatial/SampledScalarField2.h
--- a/Engine/Source/Runtime/GeometryCore/Public/Spatial/SampledScalarField2.h
+++ b/Engine/Source/Runtime/GeometryCore/Public/Spatial/SampledScalarField2.h
@@ -56,7 +56,7 @@ public:
 	void Resize(int64 Width, int64 Height, const ValueType& InitValue)
 	{
 		GridValues.Resize(Width, Height);
-		GridValues.AssignAll(InitValue);
+		GridValues.Assign(InitValue);
 	}
 
 	int32 Width() const
