GeometryCore: Fix warning realted to recent TDenseGrid changes

#Synced-CL 49302611

diff --git a/Engine/Source/Runtime/GeometryCore/Public/Spatial/SampledScalarField2.h b/Engine/Source/Runtime/GeometryCore/Public/Spatial/SampledScalarField2.h
--- a/Engine/Source/Runtime/GeometryCore/Public/Spatial/SampledScalarField2.h
+++ b/Engine/Source/Runtime/GeometryCore/Public/Spatial/SampledScalarField2.h
@@ -71,7 +71,7 @@ public:
 
 	int64 Num() const
 	{
-		return GridValues.Size();
+		return GridValues.Num();
 	}
 
 	/**
