use UE's PHI constant rather recomputing it in geometrycore spherical fibonacci methods

#Synced-CL 49445927

diff --git a/Engine/Source/Runtime/GeometryCore/Public/Sampling/SphericalFibonacci.h b/Engine/Source/Runtime/GeometryCore/Public/Sampling/SphericalFibonacci.h
--- a/Engine/Source/Runtime/GeometryCore/Public/Sampling/SphericalFibonacci.h
+++ b/Engine/Source/Runtime/GeometryCore/Public/Sampling/SphericalFibonacci.h
@@ -5,6 +5,7 @@
 #include "MathUtil.h"
 #include "VectorTypes.h"
 #include "MatrixTypes.h"
+#include "Math/UnrealMathUtility.h"
 
 namespace UE
 {
@@ -43,7 +44,7 @@ public:
 	 */
 	TVector<RealType> Point(int32 Index) const
 	{
-		static const RealType PHI = ( TMathUtil<RealType>::Sqrt(5.0) + 1.0) / 2.0;
+		constexpr RealType PHI = static_cast<RealType>(UE_DOUBLE_GOLDEN_RATIO);
 
 		checkSlow(Index >= 0 && Index < N);
 		RealType div = (RealType)Index / PHI;
@@ -73,7 +74,7 @@ public:
 	 */
 	int32 FindIndex(const TVector<RealType>& P)
 	{
-		static const RealType PHI = (TMathUtil<RealType>::Sqrt(5.0) + 1.0) / 2.0;
+		constexpr RealType PHI = static_cast<RealType>(UE_DOUBLE_GOLDEN_RATIO);
 
 		RealType phi = TMathUtil<RealType>::Min(TMathUtil<RealType>::Atan2(P.Y, P.X), TMathUtil<RealType>::Pi);
 		RealType cosTheta = P.Z;
@@ -171,7 +172,7 @@ public:
 	 */
 	TVector2<RealType> Point(int32 Index) const
 	{
-		static const RealType PHI = (TMathUtil<RealType>::Sqrt(5.0) + 1.0) / 2.0;
+		constexpr RealType PHI = static_cast<RealType>(UE_DOUBLE_GOLDEN_RATIO);
 		checkSlow(Index >= 0 && Index < N);
 
 		// Use 1-based indices to exclude the zero/origin term.
