PolyBezierSpline: structural changes to have 3-point structure (ArriveTangent, Position, LeaveTangent) stored instead of the previous 4-point per segment storage with repeated shared anchor points.

This allows us to simplify closing splines without having to create an additional closing segment and also simplifies Append/Prepend operations not needing to do a open/close dance when mutating.

#virtualized
#Synced-CL 49678995

diff --git a/Engine/Source/Runtime/Core/Public/UObject/FortniteMainBranchObjectVersions.inl b/Engine/Source/Runtime/Core/Public/UObject/FortniteMainBranchObjectVersions.inl
--- a/Engine/Source/Runtime/Core/Public/UObject/FortniteMainBranchObjectVersions.inl
+++ b/Engine/Source/Runtime/Core/Public/UObject/FortniteMainBranchObjectVersions.inl
@@ -739,5 +739,8 @@ FORTNITE_MAIN_VERSION(UAFAssetData, 241)
 // Add an additional rotation offset to root motion modifier and deprecate the 'Opposite' rotation methods
 FORTNITE_MAIN_VERSION(RootMotionModifierRotationOffset, 242)
 
+// Update PolyBezierSpline internal storage structure to use 3 control points per segment instead of 4
+FORTNITE_MAIN_VERSION(PolyBezierSplineAnchorTripleLayout, 243)
+
 // -----<new versions can be added above this line>-------------------------------------------------
 #undef FORTNITE_MAIN_VERSION
diff --git a/Engine/Source/Runtime/Engine/Private/Components/SplineComponent.cpp b/Engine/Source/Runtime/Engine/Private/Components/SplineComponent.cpp
--- a/Engine/Source/Runtime/Engine/Private/Components/SplineComponent.cpp
+++ b/Engine/Source/Runtime/Engine/Private/Components/SplineComponent.cpp
@@ -632,6 +632,32 @@ FSplineCurves FSplineCurves::FromSplineInterface(const FSpline& InProxy)
 		NewCurve.Position = TangentSpline->GetSplinePointsPosition();
 		NewCurve.Rotation = TangentSpline->GetSplinePointsRotation();
 		NewCurve.Scale = TangentSpline->GetSplinePointsScale();
+
+		// Populate LoopKeyOffset from the new spline when closed.
+		if (TangentSpline->IsClosedLoop())
+		{
+			const int32 NumPoints = NewCurve.Position.Points.Num();
+			if (NumPoints > 0)
+			{
+				const float LastInKey = NewCurve.Position.Points.Last().InVal;
+
+				// FSpline already exposes this – SplineComponent uses it in GetInputKeyValueAtSplinePoint.
+				const float ClosingParam = TangentSpline->GetParameterAtIndex(NumPoints);   // “closing” point
+				const float LoopKeyOffset = ClosingParam - LastInKey;
+
+				if (LoopKeyOffset > 0.f)
+				{
+					NewCurve.Position.bIsLooped    = true;
+					NewCurve.Position.LoopKeyOffset = LoopKeyOffset;
+
+					// Keep Rotation / Scale in sync
+					NewCurve.Rotation.bIsLooped     = true;
+					NewCurve.Rotation.LoopKeyOffset = LoopKeyOffset;
+					NewCurve.Scale.bIsLooped        = true;
+					NewCurve.Scale.LoopKeyOffset    = LoopKeyOffset;
+				}
+			}
+		}
 	}
 
 	NewCurve.UpdateSpline(NewCurve.Position.bIsLooped, false);
diff --git a/Engine/Source/Runtime/Engine/Private/Curves/Splines/TangentSpline.cpp b/Engine/Source/Runtime/Engine/Private/Curves/Splines/TangentSpline.cpp
--- a/Engine/Source/Runtime/Engine/Private/Curves/Splines/TangentSpline.cpp
+++ b/Engine/Source/Runtime/Engine/Private/Curves/Splines/TangentSpline.cpp
@@ -187,6 +187,17 @@ FTangentSpline::FTangentSpline(const FLegacyTangentSpline& Other)
 	FUpdateSplineParams UpdateParams;
 	UpdateParams.bClosedLoop = Other.PositionCurve.bIsLooped;
 
+	// If the source curves encode an override loop offset parameter, convert it to an absolute loop position
+	// The source curve is assumed to encode an override loop offset parameter if it
+	// 1) is closed,
+	// 2) has at least 1 point,
+	// 3) specifies a positive offset that is not exactly equal to 1.
+	if (Other.PositionCurve.bIsLooped && Other.PositionCurve.Points.Num() > 0 && Other.PositionCurve.LoopKeyOffset > 0.0f && Other.PositionCurve.LoopKeyOffset != 1.f)
+	{
+		UpdateParams.bLoopPositionOverride = true;
+		UpdateParams.LoopPosition = Other.PositionCurve.Points.Last().InVal + Other.PositionCurve.LoopKeyOffset;
+	}
+
 	UpdateSpline(UpdateParams);
 }
 
@@ -245,6 +256,17 @@ FTangentSpline::FTangentSpline(const FSplineCurves& Other)
 	FUpdateSplineParams UpdateParams;
 	UpdateParams.bClosedLoop = Other.Position.bIsLooped;
 
+	// If the source curves encode an override loop offset parameter, convert it to an absolute loop position
+	// The source curve is assumed to encode an override loop offset parameter if it
+	// 1) is closed,
+	// 2) has at least 1 point,
+	// 3) specifies a positive offset that is not exactly equal to 1.
+	if (Other.Position.bIsLooped && Other.Position.Points.Num() > 0 && Other.Position.LoopKeyOffset > 0.0f && Other.Position.LoopKeyOffset != 1.f)
+	{
+		UpdateParams.bLoopPositionOverride = true;
+		UpdateParams.LoopPosition = Other.Position.Points.Last().InVal + Other.Position.LoopKeyOffset;
+	}
+	
 	UpdateSpline(UpdateParams);
 }
 
@@ -500,13 +522,11 @@ void FTangentSpline::SetClosedLoop(bool bInClosedLoop, bool bUpdateSpline)
 	if (UE::Geometry::Spline::TTangentBezierSpline<FVector>* ScaleChild = GetTypedAttributeChannel<UE::Geometry::Spline::TTangentBezierSpline<FVector>>(ScaleAttrName))
 	{
 		ScaleChild->SetClosedLoop(bInClosedLoop);
-		ScaleChild->SetKnotVector(GetSpline().GetKnotVector());
 	}
 
 	if (UE::Geometry::Spline::TTangentBezierSpline<FQuat>* RotChild = GetTypedAttributeChannel<UE::Geometry::Spline::TTangentBezierSpline<FQuat>>(RotationAttrName))
 	{
 		RotChild->SetClosedLoop(bInClosedLoop);
-		RotChild->SetKnotVector(GetSpline().GetKnotVector());
 	}
 
 	if (bUpdateSpline)
@@ -1095,6 +1115,38 @@ void FTangentSpline::UpdateSpline(const FUpdateSplineParams& Params)
 
 	GetSpline().SetStationaryEndpoints(Params.bStationaryEndpoints);
 	ReparamStepsPerSegment = Params.ReparamStepsPerSegment;
+	
+	// Todo: enable this when hooking up spline input keys
+	constexpr bool bPropagateLoopPosition = false;
+	if constexpr (bPropagateLoopPosition)
+	{
+		if (Params.bClosedLoop)
+		{
+			const int32 NumPoints = GetNumberOfPoints();
+			if (NumPoints > 0)
+			{
+				const float LastParam = GetParameterAtIndex(NumPoints - 1); // FSpline-style parameter
+				const float LoopKey = Params.bLoopPositionOverride
+					? Params.LoopPosition         // user-specified absolute closing key
+					: LastParam + 1.0f;           // legacy default behaviour
+
+				if (LoopKey > LastParam)
+				{
+					GetSpline().SetParameter(GetNumberOfPoints(), LoopKey);
+
+					if (UE::Geometry::Spline::TTangentBezierSpline<FVector>* ScaleChild = GetTypedAttributeChannel<UE::Geometry::Spline::TTangentBezierSpline<FVector>>(ScaleAttrName))
+					{
+						ScaleChild->SetParameter(GetNumberOfPoints(), LoopKey);
+					}
+
+					if (UE::Geometry::Spline::TTangentBezierSpline<FQuat>* RotChild = GetTypedAttributeChannel<UE::Geometry::Spline::TTangentBezierSpline<FQuat>>(RotationAttrName))
+					{
+						RotChild->SetParameter(GetNumberOfPoints(), LoopKey);
+					}
+				}
+			}
+		}
+	}
 
 	UpdateSpline();
 }
diff --git a/Engine/Source/Runtime/GeometryCore/Public/Splines/BSpline.h b/Engine/Source/Runtime/GeometryCore/Public/Splines/BSpline.h
--- a/Engine/Source/Runtime/GeometryCore/Public/Splines/BSpline.h
+++ b/Engine/Source/Runtime/GeometryCore/Public/Splines/BSpline.h
@@ -425,18 +425,6 @@ public:
 		return Index;
 	}
 
-	// Knot vector management
-	const TArray<FKnot>& GetKnotVector() const
-	{
-		return PairKnots;
-	}
-
-	// Get the pair-based knot vector
-	const TArray<FKnot>& GetPairKnots() const
-	{
-		return PairKnots;
-	}
-
 	void ResetKnotVector()
 	{
 		PairKnots.Reset();
@@ -458,18 +446,19 @@ public:
 		return 0;
 	}
 
-	/**
-	 * Sets a custom knot vector using explicit knot/multiplicity pairs
-	 * @param NewKnots - Array of knot values. Must be non-decreasing and have correct length (see GetExpectedNumKnots)
-	 * @return True if the knot vector was valid and set successfully
-	 */
-	bool SetCustomKnots(const TArray<FKnot>& NewKnots)
+	// Knot vector management
+	virtual TArray<FKnot> GetKnotVector() const
+	{
+		return PairKnots;
+	}
+
+	bool ValidateKnotVector(const TArray<FKnot>& InKnots) const
 	{
 		const int32 RequiredKnots = GetExpectedNumKnots();
 
 		// Validate knot multiplicities and calculate total expanded knots  
 		int32 TotalExpandedKnots = 0;
-		for (const FKnot& Knot : NewKnots)
+		for (const FKnot& Knot : InKnots)
 		{
 			if (Knot.Multiplicity == 0)
 			{
@@ -482,28 +471,44 @@ public:
 		if (TotalExpandedKnots < RequiredKnots)
 		{
 			UE_LOG(LogSpline, Warning, TEXT("Invalid knot vector size. Need %d knots, got %d."),
-				RequiredKnots, TotalExpandedKnots);
+				   RequiredKnots, TotalExpandedKnots);
 			return false;
 		}
 
 		// Verify knot values are non-decreasing (strict comparison, no epsilon)
-		for (int32 i = 1; i < NewKnots.Num(); ++i)
+		for (int32 i = 1; i < InKnots.Num(); ++i)
 		{
-			if (NewKnots[i].Value < NewKnots[i - 1].Value)
+			if (InKnots[i].Value < InKnots[i - 1].Value)
 			{
 				UE_LOG(LogSpline, Warning, TEXT("Invalid knot vector: knot values must be non-decreasing."));
 				return false;
 			}
 		}
 
+		return true;
+	}
+
+	/**
+	 * Sets a custom knot vector using explicit knot/multiplicity pairs
+	 * @param NewKnots - Array of knot values. Must be non-decreasing and have correct length (see GetExpectedNumKnots)
+	 * @return True if the knot vector was valid and set successfully
+	 */
+	virtual bool SetKnotVector(const TArray<FKnot>& NewKnots)
+	{
+		if (!ValidateKnotVector(NewKnots))
+		{
+			return false;
+		}
+
 		// Directly assign the knots - no epsilon comparison needed
 		PairKnots = NewKnots;
 
 		// Mark the flat knots cache as dirty
 		MarkFlatKnotsCacheDirty();
 
-		UE_LOG(LogSpline, Verbose, TEXT("Set Custom knot vector - Points: %d, Degree: %d, Total Knots: %d, Unique Knots: %d"),
-			NumKeys(), Degree, TotalExpandedKnots, NewKnots.Num());
+		UE_LOG(LogSpline, Verbose, TEXT("Set knot vector (Direct) - Points: %d, Unique Knots: %d"),
+			NumKeys(), PairKnots.Num());
+
 		PrintKnotVector();
 
 		return true;
@@ -560,7 +565,8 @@ public:
 
 	bool IsClampedEnds() const
 	{
-		return bClampEnds;
+		// Clamping only meaningful for non-closed curves
+		return bClampEnds && !bIsClosedLoop;
 	}
 
 	// Get the valid parameter range in knot space
@@ -589,6 +595,8 @@ protected:
 		}
 
 		PairKnots[KnotIdx].Value = NewValue;
+		// Mark the flat knots cache as dirty
+		MarkFlatKnotsCacheDirty();
 	}
 
 	/**
@@ -625,6 +633,9 @@ protected:
 		FKnot OldA = PairKnots[KnotIdxA];
 		PairKnots[KnotIdxA] = PairKnots[KnotIdxB];
 		PairKnots[KnotIdxB] = OldA;
+		
+		// Mark the flat knots cache as dirty
+		MarkFlatKnotsCacheDirty();
 	}
 
 	/**
@@ -1098,17 +1109,29 @@ protected:
 
 	/**
 	 * Applies clamped multiplicity to the knot vector.
-	 * This ensures that the first and last knots have Degree + 1 multiplicity
+	 * This ensures that the first and last knots have Degree + 1 multiplicity if the ends are clamped.
+	 * For unclamped splines (including closed splines), first and last knots will have a multiplicity of Degree.
 	 */
 	void ApplyClampedKnotsMultiplicity()
 	{
-		if (IsClampedEnds() && PairKnots.Num() > 0)
+		if (PairKnots.Num() == 0)
+		{
+			return;
+		}
+
+		if (IsClampedEnds())
 		{
 			// enforce degree + 1 multiplicity at the endpoints
 			PairKnots[0].Multiplicity = Degree + 1;
 			PairKnots.Last().Multiplicity = Degree + 1;
-			MarkFlatKnotsCacheDirty();
 		}
+		else
+		{
+			PairKnots[0].Multiplicity = Degree;
+			PairKnots.Last().Multiplicity = Degree;
+		}
+
+		MarkFlatKnotsCacheDirty();
 	}
 
 	/**
@@ -1136,9 +1159,11 @@ protected:
 	 */
 	void PrintKnotVector() const
 	{
-		for (int32 Index = 0; Index < PairKnots.Num(); ++Index)
+		TArray<FKnot> KnotVector = GetKnotVector();
+
+		for (int32 Index = 0; Index < KnotVector.Num(); ++Index)
 		{
-			const FKnot& Knot = PairKnots[Index];
+			const FKnot& Knot = KnotVector[Index];
 			UE_LOG(LogSpline, Verbose, TEXT("  [%d] Value: %f, Multiplicity: %u"), Index, Knot.Value, Knot.Multiplicity);
 		}
 
diff --git a/Engine/Source/Runtime/GeometryCore/Public/Splines/PolyBezierSpline.h b/Engine/Source/Runtime/GeometryCore/Public/Splines/PolyBezierSpline.h
--- a/Engine/Source/Runtime/GeometryCore/Public/Splines/PolyBezierSpline.h
+++ b/Engine/Source/Runtime/GeometryCore/Public/Splines/PolyBezierSpline.h
@@ -1,9 +1,9 @@
 // Copyright Epic Games, Inc. All Rights Reserved.
 
 #pragma once
-
 #include "BSpline.h"
 #include "Misc/ScopeExit.h"
+#include "UObject/FortniteMainBranchObjectVersion.h"
 
 namespace UE
 {
@@ -27,6 +27,7 @@ class TPolyBezierSpline :
 {
 public:
 	using Base = TBSpline<VALUETYPE, 3>;
+	using Base::Degree;
 	typedef typename Base::ValueType ValueType;
 	using FWindow = typename Base::FWindow;
 	using FWindowStore = typename Base::FWindowStore;
@@ -267,6 +268,28 @@ public:
         
         return Result;
     }
+	
+	virtual bool Serialize(FArchive& Ar) override
+	{
+		Ar.UsingCustomVersion(FFortniteMainBranchObjectVersion::GUID);
+
+		if (!Base::Serialize(Ar))
+		{
+			return false;
+		}
+
+		if (Ar.IsLoading())
+		{
+			if (Ar.CustomVer(FFortniteMainBranchObjectVersion::GUID) < FFortniteMainBranchObjectVersion::PolyBezierSplineAnchorTripleLayout)
+			{
+				// At this point Values contains the legacy 4-point-per-seg layout.
+				// Convert it in-place to the new anchor triple layout.
+				MigrateLegacyBezierLayout();
+			}
+		}
+
+		return true;
+	}
 
 	// ISplineInterface Implementation
 	virtual void Clear() override { Base::Clear(); }
@@ -278,7 +301,6 @@ public:
         Clone->Values = this->Values;
         Clone->PairKnots = this->PairKnots;
         Clone->bIsClosedLoop = this->bIsClosedLoop;
-        Clone->bClampEnds = this->bClampEnds;
         
         // Copy infinity modes
         Clone->PreInfinityMode = this->PreInfinityMode;
@@ -295,15 +317,37 @@ public:
     		return 0.0f;
     	}
 
-    	// Get Bezier control points for this segment
-    	TStaticArray<ValueType, 4> Coeffs = {};
-    	const int32 BaseIndex = SegmentIndex * 4;
-        
-    	// Get the control points
-    	const ValueType& P0 = Base::GetValue(BaseIndex);
-    	const ValueType& P1 = Base::GetValue(BaseIndex + 1);
-    	const ValueType& P2 = Base::GetValue(BaseIndex + 2);
-    	const ValueType& P3 = Base::GetValue(BaseIndex + 3);
+
+		// Get Bezier control points for this segment
+		TStaticArray<ValueType, 4> Coeffs = {};
+
+		const ValueType* P0Ptr = nullptr;
+		const ValueType* P1Ptr = nullptr;
+		const ValueType* P2Ptr = nullptr;
+		const ValueType* P3Ptr = nullptr;
+
+		if (!IsClosingSegment(SegmentIndex))
+		{
+			P0Ptr = &Base::GetValue(P0Idx(SegmentIndex)); // A_s
+			P1Ptr = &Base::GetValue(P1Idx(SegmentIndex)); // Out_s
+			P2Ptr = &Base::GetValue(P2Idx(SegmentIndex)); // In_{s+1}
+			P3Ptr = &Base::GetValue(P3Idx(SegmentIndex)); // A_{s+1}
+		}
+		else
+		{
+			const int32 NumAnchors = GetNumAnchorPoints();
+
+			// Closing segment: A_last → A_0
+			P0Ptr = &Base::GetValue(Close_P0Idx(NumAnchors)); // A_last
+			P1Ptr = &Base::GetValue(Close_P1Idx(NumAnchors)); // Out_last
+			P2Ptr = &Base::GetValue(Close_P2Idx());           // In_0
+			P3Ptr = &Base::GetValue(Close_P3Idx());           // A_0
+		}
+
+		const ValueType& P0 = *P0Ptr;
+		const ValueType& P1 = *P1Ptr;
+		const ValueType& P2 = *P2Ptr;
+		const ValueType& P3 = *P3Ptr;
 
     	// Setup coefficients in standard Bezier polynomial form relative to test point
     	Coeffs[0] = P0 - Point;                    // constant term
@@ -317,12 +361,35 @@ public:
 
 	float PointDistanceSqToSegmentBBox(const ValueType &Point, int32 SegmentIndex) const
 	{
+		// Decide whether this is the synthetic closing segment
+
 		// Get the control points
-		int32 BaseIndex = SegmentIndex * 4;
-		const ValueType& P0 = Base::GetValue(BaseIndex);
-		const ValueType& P1 = Base::GetValue(BaseIndex + 1);
-		const ValueType& P2 = Base::GetValue(BaseIndex + 2);
-		const ValueType& P3 = Base::GetValue(BaseIndex + 3);
+		const ValueType* P0Ptr = nullptr;
+		const ValueType* P1Ptr = nullptr;
+		const ValueType* P2Ptr = nullptr;
+		const ValueType* P3Ptr = nullptr;
+
+		if (!IsClosingSegment(SegmentIndex))
+		{
+			P0Ptr = &Base::GetValue(P0Idx(SegmentIndex));
+			P1Ptr = &Base::GetValue(P1Idx(SegmentIndex));
+			P2Ptr = &Base::GetValue(P2Idx(SegmentIndex));
+			P3Ptr = &Base::GetValue(P3Idx(SegmentIndex));
+		}
+		else
+		{
+			const int32 NumAnchors = GetNumAnchorPoints();
+
+			P0Ptr = &Base::GetValue(Close_P0Idx(NumAnchors)); // A_last
+			P1Ptr = &Base::GetValue(Close_P1Idx(NumAnchors)); // Out_last
+			P2Ptr = &Base::GetValue(Close_P2Idx());           // In_0
+			P3Ptr = &Base::GetValue(Close_P3Idx());           // A_0
+		}
+
+		const ValueType& P0 = *P0Ptr;
+		const ValueType& P1 = *P1Ptr;
+		const ValueType& P2 = *P2Ptr;
+		const ValueType& P3 = *P3Ptr;
 
 		//	For cubics, convex hull is defined by control points, so use them to compute BBox. 
 		ValueType BBoxMin = P0.ComponentMin(P1);
@@ -341,7 +408,7 @@ public:
 
 	virtual float FindNearest(const ValueType& Point, float& OutSquaredDistance) const override
 	{
-    	const int32 NumSegments = Base::NumKeys() / 4;
+    	const int32 NumSegments = GetNumberOfSegments();
 		float BestDistSq;
 		float BestParam = FindNearestOnSegment(Point, 0, BestDistSq);
 		
@@ -388,9 +455,6 @@ public:
     	return BestParam;
 	}
 
-	// Evaluator methods
-	static int32 GetDegree() { return 3; }
-
     /** 
      * Evaluate nth derivative at parameter
      * @tparam Order - The derivative order (0 = position, 1 = first derivative, etc.)
@@ -419,8 +483,11 @@ public:
     	// Convert global parameter to segment information
     	int32 SegmentIndex;
     	float LocalT;
-    	MapGlobalParameterToLocalSegment(Parameter, SegmentIndex, LocalT);
-		
+		if (!MapGlobalParameterToLocalSegment(Parameter, SegmentIndex, LocalT))
+		{
+			return Math::FSplineValidation::GetDefaultValue<ValueType>();
+		}
+    	
 		const float SegmentScale = GetSegmentParameterRange(SegmentIndex).Length();
 	    
 		// Compute the derivative with proper scaling
@@ -444,116 +511,69 @@ public:
 		const TArray<ValueType>& Points,
 		EParameterizationPolicy ParameterizationPolicy)
 	{
-	    const int32 NumPoints = Points.Num();
-    
-	    // Need at least 4 points for a valid cubic bezier segment
-	    if (NumPoints < 4)
+	   const int32 NumPoints = Points.Num();
+
+	    // We expect Bezier quartets as input: P0,P1,P2,P3 per segment.
+	    if (NumPoints < 4 || (NumPoints % 4) != 0)
 	    {
-	        UE_LOG(LogSpline, Warning, TEXT("SetControlPoints requires at least 4 points to add a valid segment. Got %d points."), NumPoints);
+	        UE_LOG(LogSpline, Warning,
+	            TEXT("SetControlPoints expects 4*k Bezier control points (got %d)."), NumPoints);
 	        return false;
 	    }
-	    
-	    // Clear existing points
+
+	    // Remember desired closed state, then reset storage.
+	    const bool bWasClosed = Base::IsClosedLoop();
 	    Clear();
-	    const bool bIsClosedLoop = Base::IsClosedLoop();
-	    
-	    // For cubic Bezier curves, points come in groups of 4:
-	    // P0 (start), P1 (control), P2 (control), P3 (end)
-	    // Each complete segment requires 4 points
-	    
-	    // Calculate how many complete segments we have
-	    int32 NumSegments = (NumPoints) / 4;
-	    
-	    // Make sure we have at least one complete segment
-	    if (NumSegments < 1)
-	    {
-	        UE_LOG(LogSpline, Warning, TEXT("Not enough points to create a complete segment. Need at least 4 points."));
-	        return false;
-	    }
-	    
-	    // Add each segment
-	    for (int32 SegmentIndex = 0; SegmentIndex < NumSegments; ++SegmentIndex)
+
+	    const int32 NumSegments = NumPoints / 4;
+		
+	    // First segment bootstraps two anchors
 	    {
-	        const int32 BaseIndex = SegmentIndex * 4;
-	        
-	        // Ensure we have enough points for this segment
-	        if (BaseIndex + 3 >= NumPoints)
-	        {
-	            break;
-	        }
-	        
-	        // Extract the 4 control points for this segment
-	        ValueType P0 = Points[BaseIndex];
-	        ValueType P1 = Points[BaseIndex + 1];
-	        ValueType P2 = Points[BaseIndex + 2];
-	        ValueType P3 = Points[BaseIndex + 3];
-	        
-	        // Add the segment
-	        if (SegmentIndex == 0)
-	        {
-	            // For the first segment, use AddBezierSegmentInternal
-	            this->AddBezierSegmentInternal(0, P0, P1, P2, P3, ParameterizationPolicy);
-	        }
-	        else
-	        {
-	            // For subsequent segments, check if P0 matches the last point of the previous segment
-	            ValueType LastPoint = Base::GetValue(Base::NumKeys() - 1);
-	            
-	            if (Math::SizeSquared(P0 - LastPoint) < UE_KINDA_SMALL_NUMBER)
-	            {
-	                // Points are the same, use AppendBezierSegment which only needs P1, P2, P3
-	                this->AppendBezierSegment(P1, P2, P3, ParameterizationPolicy);
-	            }
-	            else
-	            {
-	            	
-	                // Points don't match, use AddBezierSegmentInternal with appropriately computed parameter
-	            	this->AddBezierSegmentInternal(NumSegments, P0, P1, P2, P3, ParameterizationPolicy);
-	            }
-	        }
+	        const ValueType& P0 = Points[0];
+	        const ValueType& P1 = Points[1];
+	        const ValueType& P2 = Points[2];
+	        const ValueType& P3 = Points[3];
+
+	        // A0
+	        AppendAnchor(/*In0*/ P0, /*A0*/ P0, /*Out0*/ P1);
+	        // A1
+	        AppendAnchor(/*In1*/ P2, /*A1*/ P3, /*Out1*/ P3);
 	    }
-	    
-	    // Update closed loop state if needed
-	    if (bIsClosedLoop)
+
+	    // Remaining segments
+	    for (int32 s = 1; s < NumSegments; ++s)
 	    {
-	        // If we're supposed to be a closed loop, make sure the last point connects to the first
-	        ValueType FirstPoint = Base::GetValue(0);
-	        ValueType LastPoint = Base::GetValue(Base::NumKeys() - 1);
-	        
-	        // If not already a closed loop, add a segment to close it
-	        if (Math::SizeSquared(LastPoint - FirstPoint) > UE_KINDA_SMALL_NUMBER && NumSegments > 0)
+	        const int32 b = 4 * s;
+	        const ValueType& P0 = Points[b + 0];
+	        const ValueType& P1 = Points[b + 1];
+	        const ValueType& P2 = Points[b + 2];
+	        const ValueType& P3 = Points[b + 3];
+
+	        // Optionally warn if P0 doesn't match previous A (only a sanity message).
+	        const ValueType& PrevA = this->Values[AnchorBase(s) + 1];
+	        if (Math::SizeSquared(P0 - PrevA) > UE_KINDA_SMALL_NUMBER)
 	        {
-	            // Add a segment that connects back to the start
-	            // We need to compute appropriate control points
-	            int32 LastBaseIndex = (NumSegments - 1) * 4;
-	            
-	            // We'll try to maintain tangent continuity if possible
-	            ValueType OutTangent = (Points[LastBaseIndex + 2] - Points[LastBaseIndex + 1]);
-	            ValueType InTangent = (Points[1] - Points[0]);
-	            
-	            // Scale tangents to match segment length
-	            float SegmentLength = static_cast<float>(Math::Distance(FirstPoint, LastPoint));
-	            OutTangent = Math::GetSafeNormal(OutTangent) * (SegmentLength * 0.33f);
-	            InTangent = Math::GetSafeNormal(InTangent) * (SegmentLength * 0.33f);
-	            
-	            // Create control points
-	            ValueType P1 = LastPoint + OutTangent;
-	            ValueType P2 = FirstPoint - InTangent;
-	            
-	            // Add the closing segment
-	            this->AppendBezierSegment(P1, P2, FirstPoint, ParameterizationPolicy);
+	            UE_LOG(LogSpline, Verbose,
+	                TEXT("SetControlPoints: segment %d P0 differs from previous A (len^2=%g)."),
+	                s, (double)Math::SizeSquared(P0 - PrevA));
 	        }
-	        
-	        Base::SetClosedLoop(true);
-	    }
-	    else
-	    {
-	        Base::SetClosedLoop(false);
+
+	        // Out_s = P1
+	        this->SetValue(AnchorBase(s) + 2, P1);
+
+	        // Append A_{s+1} = P3 with In_{s+1}=P2, Out_{s+1}=A_{s+1} (placeholder until another seg arrives)
+	        AppendAnchor(/*In*/ P2, /*A*/ P3, /*Out*/ P3);
 	    }
-	    
-	    // Reparameterize for consistent parameter distribution
+
+	    // Apply desired loop state:
+	    // - We DO NOT materialize a physical closing quartet anymore.
+	    // - Closing segment is virtual: {A_last, Out_last, In_0, A_0}.
+	    SetClosedLoop(bWasClosed);
+
+	    // Build / rebuild knot vector once from scratch using the policy
+	    // (You already do this today; keep it — it re-spaces and sets multiplicities.)
 	    Reparameterize(ParameterizationPolicy);
-	    
+
 	    return true;
 	}
 
@@ -605,16 +625,17 @@ public:
 	            return false;
 	        }
 	    }
-	    
+		const int32 NumAnchors = GetNumAnchorPoints();
+
 	    // Check connection point
 	    if (bAppend)
 	    {
-	        // First point of new segments should match the last point of existing spline
+	        // First point of new segments should match the last existing anchor
 	        const ValueType& FirstNewPoint = Points[0];
-	        const ValueType& LastExistingPoint = this->Base::GetValue(this->Base::NumKeys() - 1);
+	    	const ValueType& LastExistingAnchor = this->Values[AnchorBase(NumAnchors - 1) + 1];
 	        
 	        // Check if the connection point is reasonably close
-	        const double DistanceSquared = Math::SizeSquared(LastExistingPoint, FirstNewPoint);
+	        const double DistanceSquared = Math::SizeSquared(LastExistingAnchor, FirstNewPoint);
 	        if (DistanceSquared > UE_KINDA_SMALL_NUMBER)
 	        {
 	            UE_LOG(LogSpline, Warning, TEXT("Connection point mismatch when appending segments. Distance: %f"),
@@ -626,10 +647,10 @@ public:
 	    {
 	        // Last point of new segments should match the first point of existing spline
 	        const ValueType& LastNewPoint = Points.Last();
-	        const ValueType& FirstExistingPoint = this->Base::GetValue(0);
+	    	const ValueType& FirstExistingAnchor = this->Values[AnchorBase(0) + 1];
 	        
 	        // Check if the connection point is reasonably close
-	        const double DistanceSquared = Math::SizeSquared(LastNewPoint, FirstExistingPoint);
+	        const double DistanceSquared = Math::SizeSquared(LastNewPoint, FirstExistingAnchor);
 	        if (DistanceSquared > UE_KINDA_SMALL_NUMBER)
 	        {
 	            UE_LOG(LogSpline, Warning, TEXT("Connection point mismatch when prepending segments. Distance: %f"),
@@ -700,23 +721,16 @@ public:
         const ValueType& P3, 
         EParameterizationPolicy ParameterizationPolicy = EParameterizationPolicy::Centripetal)
     {
-		const TOptional<float> LoopKnotDelta = GetLoopKnotDelta();
-		SetClosedLoop(false);
-		ON_SCOPE_EXIT
-		{
-			if (LoopKnotDelta.IsSet())
-			{
-				SetClosedLoop(true);
-				SetParameter(GetNumberOfSegments(), LoopKnotDelta.GetValue());
-			}
-		};
+		const int32 NumAnchors = GetNumAnchorPoints();
+		const int32 NumSegments = GetNumberOfSegments();
 		
-		// Get P0 from the last point of the previous segment
-		ValueType P0 = Base::GetValue(Base::NumKeys() - 1);
-		int32 NumSegments = Base::NumKeys() / 4;
-		// Add segment with append behavior
+		// P0 = last anchor position A_{n-1}
+		const int32 LastAnchor = NumAnchors - 1;
+		const ValueType& P0 = this->Values[AnchorBase(LastAnchor) + 1];
 
-		return AddBezierSegmentInternal(NumSegments, P0, P1, P2, P3, ParameterizationPolicy);
+		// Insert new segment after the current last one
+		//   SegmentIndex == NumSegments means "append at tail"
+		return AddBezierSegmentInternal(NumSegments, P0, P1, P2, P3, ParameterizationPolicy, true);
     }
 
     /**
@@ -735,20 +749,13 @@ public:
 		const ValueType& P2, 
 		EParameterizationPolicy ParameterizationPolicy = EParameterizationPolicy::Centripetal)
 	{
-		const TOptional<float> LoopKnotDelta = GetLoopKnotDelta();
-		SetClosedLoop(false);
-		ON_SCOPE_EXIT
-		{
-			if (LoopKnotDelta.IsSet())
-			{
-				SetClosedLoop(true);
-				SetParameter(GetNumberOfSegments(), LoopKnotDelta.GetValue());
-			}
-		};
+		const int32 NumAnchors = GetNumAnchorPoints();
+		const int32 NumSegments = GetNumberOfSegments();
 		
-		// Get P3 from the first point of the existing spline
-		ValueType P3 = Base::GetValue(0);
-		// Add segment with prepend behavior
+		// P3 = first anchor position A_0
+		const ValueType& P3 = this->Values[AnchorBase(0) + 1];
+
+		// Insert new segment before segment 0 / at the head
 		return AddBezierSegmentInternal(0, P0, P1, P2, P3,
 										ParameterizationPolicy,
 										false); // Prepend to first segment
@@ -792,13 +799,13 @@ public:
 	    const ValueType& Position,
 	    EParameterizationPolicy ParameterizationPolicy = EParameterizationPolicy::Centripetal)
 	{
-    	int32 NumSegments = Base::NumKeys() / 4;
+    	int32 NumSegments = GetNumberOfSegments();
 	    // Validate inputs
 	    if (NumSegments == 0)
 	    {
 	        // Special case for empty spline
 	        int32 NewPointIndex;
-	        if (InitializeFirstSegment(0.0f, Position, NewPointIndex))
+	        if (InitializeFirstSegment(Position, NewPointIndex))
 	            return NewPointIndex;
 	    }
 	    
@@ -829,9 +836,8 @@ public:
         const ValueType& Position,
         EParameterizationPolicy ParameterizationPolicy = EParameterizationPolicy::Centripetal)
     {
-		// todo: stop returning bool, instead interpret INDEX_NONE as failure
 	    if (int32 NewPointIndex;
-	    	InitializeFirstSegment(Parameter, Position, NewPointIndex))
+	    	InitializeFirstSegment(Position, NewPointIndex))
 		{
 			return NewPointIndex;
 		}
@@ -842,12 +848,6 @@ public:
 		Parameter = Base::GetNearestAvailableKnotValue(Parameter, SearchCache);
 		
         FInterval1f ParameterRange = GetParameterSpace();
-		int32 SegmentIndex;
-		float LocalT;
-		MapGlobalParameterToLocalSegment(Parameter, SegmentIndex, LocalT);
-
-		const int32 BaseIndex = SegmentIndex * 4;
-		ValueType P0 = Base::GetValue(BaseIndex);
 		
 		// Handle out-of-range Knots
 		if (Parameter <= ParameterRange.Min)
@@ -860,11 +860,24 @@ public:
 		{
 			return AppendPoint(Position, ParameterizationPolicy);
 		}
- 
-        ValueType P1 = Base::GetValue(BaseIndex + 1);
-        ValueType P2 = Base::GetValue(BaseIndex + 2);
-        ValueType P3 = Base::GetValue(BaseIndex + 3);
- 
+
+		int32 SegmentIndex;
+		float LocalT;
+		if (!MapGlobalParameterToLocalSegment(Parameter, SegmentIndex, LocalT))
+		{
+			return INDEX_NONE;
+		}
+
+		FWindow Window;
+		if (!BuildRawWindowForSegment(SegmentIndex, Window))
+		{
+			return INDEX_NONE;
+		}
+		ValueType P0 = *Window[0];
+		ValueType P1 = *Window[1];
+		ValueType P2 = *Window[2];
+		ValueType P3 = *Window[3];
+		
         // Use De Casteljau algorithm to split the curve
         const float t = LocalT;
         const float mt = 1.0f - t;
@@ -893,69 +906,59 @@ public:
 			TargetSplitParam = Param::PrevDistinct(TargetSplitParam);
 		}
 		
-		for (int32 i = 0; i < 4; ++i)
+		// Bail out if we lost usable range
+		if (TargetSplitParam <= SegmentRange.Min || TargetSplitParam >= SegmentRange.Max)
 		{
-			Base::RemoveValue(BaseIndex);
+			return INDEX_NONE;
 		}
-		
-        // Add the two new segments using the policy
-        AddBezierSegmentInternal(SegmentIndex, Q0, Q1, Q2, Q3, ParameterizationPolicy);
-		Base::InsertKnot( FKnot(TargetSplitParam, this->bClampEnds ? this->GetDegree() : 1));
-		int32 IndexToReturn = AddBezierSegmentInternal(SegmentIndex + 1, R0, R1, R2, R3, ParameterizationPolicy);
-		
-		return IndexToReturn;
-    }
-	
-    /**
-     * Inserts a Bezier segment at the specified parameter
-     * 
-     * @param ValuesChannel - Channel to add the segment to
-     * @param Parameter Parameter value where to insert the segment (Spline Space)
-     * @param P1 First control point
-     * @param P2 Second control point
-     * @param P3 End point
-     * @param ParameterizationPolicy How to distribute Knots
-     * @return Index of the new segment
-     */
-    int32 InsertBezierSegment(
-        float Parameter, 
-        const ValueType& P1, 
-        const ValueType& P2, 
-        const ValueType& P3, 
-        EParameterizationPolicy ParameterizationPolicy = EParameterizationPolicy::Centripetal)
-    {
-		FInterval1f ParameterRange = GetParameterSpace();
-        // Handle edge cases
-        if (Parameter >= ParameterRange.Max)
-        {
-            // Add new segment at the end
-            return AppendBezierSegment(P1, P2, P3, ParameterizationPolicy);
-        }
-        
-        if (Parameter <= ParameterRange.Min)
-        {
-			// we treat P1, P2, P3 as the starting three control points, and connect to existing spline at the start
-            return PrependBezierSegment(P1, P2, P3, ParameterizationPolicy);
-        }
-        
-        // Normal case - insert at parameter
-        const int32 NumSegments = Base::NumKeys() / 4;
-        
-        // Find which segment contains this parameter
-		int32 SegmentIndex;
-		float LocalParam;
-		MapGlobalParameterToLocalSegment(Parameter, SegmentIndex, LocalParam);
-        
-        // Evaluate position at the parameter
-        ValueType PositionAtParam = Base::Evaluate(Parameter);
-        
-        // Use the InsertPoint method to split the curve at the parameter
-        InsertPointAtSegmentParam(SegmentIndex, LocalParam, PositionAtParam, ParameterizationPolicy);
-        
-        // Add the new segment from the inserted point
-        int32 Result = AddBezierSegmentInternal(SegmentIndex+1, PositionAtParam, P1, P2, P3, ParameterizationPolicy);
-    	
-        return Result;
+
+		// --- Anchor/handle surgery in 3-tuple layout ---
+		const int32 NumAnchors = GetNumAnchorPoints();
+		if (NumAnchors < 2)
+		{
+			return INDEX_NONE;
+		}
+
+		// For segment s in an *open* spline, its anchors are A_s and A_{s+1}.
+		// For closed, the closing segment (s == A-1) is between A_{A-1} and A_0.
+		const bool bClosed = IsActuallyClosed();
+		const int32 LeftAnchor = SegmentIndex;
+		const int32 RightAnchor =
+			(!bClosed)
+			? (SegmentIndex + 1)
+			: ((SegmentIndex + 1) % NumAnchors);
+
+		// 1) Left side: Out_left = Q1
+		SetHandleOutOfAnchor(LeftAnchor, Q1);
+
+		// 2) Insert new anchor between left and right with (In_new, A_new, Out_new) = (Q2, Q3, R1)
+		//    For open, this is simply at index RightAnchor.
+		//    For closed, inserting at index RightAnchor still produces a valid cyclic order;
+		//    the notion of "start anchor" is arbitrary on a loop.
+		InsertAnchorAt(RightAnchor, /*In*/ Q2, /*A*/ Q3, /*Out*/ R1);
+
+		// 3) Right side: old right anchor has shifted by +1 after the insert
+		const int32 OldRightAnchorIndex =
+			(!bClosed)
+			? (RightAnchor + 1)
+			: ((RightAnchor + 1) % (NumAnchors + 1)); // one more anchor now in closed loop
+
+		SetHandleInOfAnchor(OldRightAnchorIndex, R2);
+
+		// --- Insert a knot at the split location (incremental) ---
+
+		this->InsertKnot(FKnot(TargetSplitParam, Degree));
+
+		this->ApplyClampedKnotsMultiplicity();
+		ApplyInternalKnotsMultiplicity();
+
+		if (bClosed)
+		{
+			RefreshClosingKnot(ParameterizationPolicy); // optional helper to ensure closing boundary is consistent
+		}
+
+		// The new anchor is at index RightAnchor
+		return RightAnchor;
     }
 
 	/**
@@ -967,7 +970,7 @@ public:
 	*/
     bool RemoveSegment(const int32 SegmentIndex)
     {
-    	const int32 NumSegments = Base::NumKeys() / 4;
+    	const int32 NumSegments = GetNumberOfSegments();
 	    if (SegmentIndex < 0 || SegmentIndex >= NumSegments)
 	    {
 	        return false;
@@ -992,68 +995,127 @@ public:
 	/**
 	 * Removes a point from the spline.
 	 * @param PointIndex - Index of the point to remove
+	 * @return true if successfully removed
 	 */
 	bool RemovePoint(int32 PointIndex)
 	{
-		const int32 NumSegments = GetNumberOfSegments();
-        if (PointIndex < 0 || PointIndex > NumSegments)
-        {
-            return false;
-        }
-		bool bSuccess = true;
-		const int32 SegmentIndex = FMath::Clamp(PointIndex, 0, NumSegments - 1);
-		const int32 BaseIndex = SegmentIndex * 4;
-		const int32 Degree = GetDegree();
-		if (PointIndex == 0)
-		{
-			// Removing the first point - remove the first segment
-			for (int32 i = 3; i >= 0; --i)
-			{
-				bSuccess &= Base::RemoveValue(BaseIndex + i);
-			}
-			// Remove first knot
-			Base::RemoveKnot(0);
+		const int32 NumAnchors  = GetNumAnchorPoints();
+	    const int32 NumSegments = GetNumberOfSegments();
 
-			if (Base::IsClampedEnds())
-			{
-				Base::PairKnots[0].Multiplicity = Degree + 1;
-			}
-		}
-		// else if last point
-		else if (PointIndex == GetNumberOfSegments())
+	    if (NumAnchors == 0)
+	    {
+	        return false;
+	    }
+
+	    // In this model, "points" == anchors
+	    if (PointIndex < 0 || PointIndex >= NumAnchors)
+	    {
+	        return false;
+	    }
+
+		auto RemoveAnchorTriple = [this](int32 AnchorIndex)
 		{
-			// Removing the last point - remove the last segment
-			for (int32 i = 3; i >= 0; --i)
-			{
-				bSuccess &= Base::RemoveValue(BaseIndex + i);
-			}
-	        
-			// Remove end knot
-			Base::RemoveKnot(SegmentIndex + 1);
-			if (Base::IsClampedEnds())
-			{
-				Base::PairKnots.Last().Multiplicity = Degree + 1;
-			}
-		}
-		else
+			const int32 BaseIdx = AnchorBase(AnchorIndex);
+			// Remove in reverse index order so earlier removes don't shift later indices
+			this->RemoveValue(BaseIdx + 2); // Out_i
+			this->RemoveValue(BaseIdx + 1); // A_i
+			this->RemoveValue(BaseIdx + 0); // In_i
+		};
+
+		if (NumAnchors == 1)		// Special case: only 1 anchors → removing one leaves empty spline.
 		{
-			// For an internal point, we need to remove exactly 4 control points:
-			// - P2 and P3 from the previous segment (where P3 is the point)
-			// - P0 and P1 from the next segment (where P0 is the same point)
-        
-			// Start index for removal: P2 of previous segment
-			int32 StartRemoveIndex = (PointIndex - 1) * 4 + 2;
-        
-			// Remove 4 consecutive control points
-			for (int32 i = 0; i < 4; ++i)
-			{
-				Base::RemoveValue(StartRemoveIndex);
-			}
-        
-			Base::RemoveKnot(PointIndex);
+			Base::Values.Reset();
+	        this->PairKnots.Reset();
+	        return true;
 		}
-		Base::Dump();
-		return bSuccess;
+	    else if (NumAnchors == 2)	// Special case: only 2 anchors → removing one leaves 1 anchor, no segments.
+	    {
+	        RemoveAnchorTriple(PointIndex);
+	        // One anchor left, no segments: clear knots
+	        this->PairKnots.Reset();
+	        return true;
+	    }
+
+	    // ---- OPEN SPLINE CASE ----
+	    if (!IsActuallyClosed())
+	    {
+	        // For open curves:
+	        //  - anchors = NumSegments + 1
+	        //  - boundary knots are indices [0 .. NumSegments]
+	        //  - interior boundaries (between anchors j and j+1) correspond to knot index j+1?
+
+	        // A) remove first point (anchor 0): drop first segment and knot 0
+	        if (PointIndex == 0)
+	        {
+	            RemoveAnchorTriple(0);
+
+	            // Remove start boundary knot
+	            if (this->PairKnots.Num() > 0)
+	            {
+	                this->RemoveKnot(0);
+	            }
+
+	            this->ApplyClampedKnotsMultiplicity();
+	            ApplyInternalKnotsMultiplicity();
+	            return true;
+	        }
+
+	        // B) remove last point (anchor N-1): drop last segment and last boundary knot
+	        if (PointIndex == NumAnchors - 1)
+	        {
+	            RemoveAnchorTriple(NumAnchors - 1);
+
+	            // Remove end boundary knot (index == NumSegments)
+	            if (this->PairKnots.Num() > 0)
+	            {
+	                const int32 EndKnotIndex = FMath::Clamp(NumSegments, 0, this->PairKnots.Num() - 1);
+	                this->RemoveKnot(EndKnotIndex);
+	            }
+
+	            this->ApplyClampedKnotsMultiplicity();
+	            ApplyInternalKnotsMultiplicity();
+	            return true;
+	        }
+
+	        // C) remove internal point (0 < j < N-1): merge segments (j-1)->j and j->(j+1)
+	        //    into one (j-1)->(j+1). Shape-wise we keep Out_{j-1} and In_{j+1};
+	        //    dropping A_j's triple is enough.
+	        RemoveAnchorTriple(PointIndex);
+
+	        // Interior knot index aligned with the removed anchor index in the old scheme:
+	        // - anchor j corresponds to a boundary between segments (j-1) and j
+	        // - that boundary’s knot index is j
+	        if (this->PairKnots.Num() > PointIndex)
+	        {
+	            this->RemoveKnot(PointIndex);
+	        }
+
+	        this->ApplyClampedKnotsMultiplicity();
+	        ApplyInternalKnotsMultiplicity();
+	        return true;
+	    }
+
+	    // ---- CLOSED SPLINE CASE ----
+	    // For closed curves:
+	    //  - anchors == segments
+	    //  - all boundaries are "interior" (no clamped ends)
+	    //  - removing anchor j merges two segments (j-1)->j and j->(j+1)
+	    //    into one (j-1)->(j+1); we keep Out_{j-1} and In_{j+1} and just drop A_j.
+	    {
+	        const int32 RemoveAnchorIdx = PointIndex % NumAnchors;
+	        RemoveAnchorTriple(RemoveAnchorIdx);
+
+	        // In the old closed implementation, a point removal deleted the
+	        // corresponding boundary knot at the same index.
+	        if (this->PairKnots.Num() > RemoveAnchorIdx)
+	        {
+	            this->RemoveKnot(RemoveAnchorIdx);
+	        }
+
+	        // No clamped ends in closed curve: all remaining knots should be interior-style.
+	        ApplyInternalKnotsMultiplicity();
+	        return true;
+	    }
 	}
 
     /**
@@ -1069,16 +1131,18 @@ public:
     	const int32 PointIndex,
     	const ValueType& NewValue)
     {
-        const int32 NumSegments = Base::NumKeys() / 4;
-        if (SegmentIndex < 0 || SegmentIndex >= NumSegments ||
-            PointIndex < 0 || PointIndex > 3)
-        {
-            return false;
-        }
+		const int32 NumSegments = GetNumberOfSegments();
+		if (SegmentIndex < 0 || SegmentIndex >= NumSegments ||
+			PointIndex  < 0 || PointIndex  > 3)
+		{
+			return false;
+		}
 
-        // Convert segment+point index to global index in spline values
-        const int32 GlobalIndex = SegmentIndex * 4 + PointIndex;
-        return Base::SetValue(GlobalIndex, NewValue);
+		int32 BezierIndices[4];
+		GetBezierIndices(SegmentIndex, BezierIndices[0], BezierIndices[1], BezierIndices[2], BezierIndices[3]);
+		int32 ValueIndex = BezierIndices[PointIndex];
+		this->SetValue(ValueIndex, NewValue);
+		return true;
     }
 
     /**
@@ -1097,27 +1161,55 @@ public:
         const ValueType& P2,
         const ValueType& P3)
     {
-        const int32 NumSegments = Base::NumKeys() / 4;
+        const int32 NumSegments = GetNumberOfSegments();
         if (SegmentIndex < 0 || SegmentIndex >= NumSegments)
         {
             return false;
         }
 
-        const int32 BaseIndex = SegmentIndex * 4;
-        bool bSuccess = true;
-        bSuccess &= Base::SetValue(BaseIndex, P0);
-        bSuccess &= Base::SetValue(BaseIndex + 1, P1);
-        bSuccess &= Base::SetValue(BaseIndex + 2, P2);
-        bSuccess &= Base::SetValue(BaseIndex + 3, P3);
-    
-        return bSuccess;
+		int32 Idx0, Idx1, Idx2, Idx3;
+		GetBezierIndices(SegmentIndex, Idx0, Idx1, Idx2, Idx3);
+		
+		this->SetValue(Idx0, P0);
+		this->SetValue(Idx1, P1);
+		this->SetValue(Idx2, P2);
+		this->SetValue(Idx3, P3);
+
+		return true;
     }
+
+	/**
+	 * Anchor points are unique segment endpoints.
+	 * All segments have 2 corresponding anchor points.
+	 * For a closed spline, the final segment's anchor points are the first and last points.
+	 */
+	int32 GetNumAnchorPoints() const { return Base::Values.Num() / 3; }
+
+	void SetAnchorPoint(int32 AnchorIndex, const ValueType& NewPosition)
+	{
+		if (AnchorIndex < 0 || AnchorIndex >= GetNumAnchorPoints())
+		{
+			return;
+		}
+		const int32 AnchorValueIndex = AnchorBase(AnchorIndex) + 1;
+		Base::SetValue(AnchorValueIndex, NewPosition);
+	}
+
+	ValueType GetAnchorPoint(int32 AnchorIndex) const
+	{
+		if (AnchorIndex < 0 || AnchorIndex >= GetNumAnchorPoints())
+		{
+			return ValueType();
+		}
+		const int32 AnchorValueIndex = AnchorBase(AnchorIndex) + 1;
+		return Base::GetValue(AnchorValueIndex);
+	}
 	
 	/** Returns the number of Bezier segments in the spline */
 	virtual int32 GetNumberOfSegments() const override
 	{
-		const int32 NumPoints = this->Base::NumKeys();
-		return NumPoints/4;
+		const int32 NumAnchors = GetNumAnchorPoints();
+		return IsActuallyClosed() ? NumAnchors : FMath::Max(0, NumAnchors - 1);
 	}
 
 	/**
@@ -1128,12 +1220,36 @@ public:
 	virtual FInterval1f GetSegmentParameterRange(int32 SegmentIndex) const override
 	{
 		FInterval1f SegmentRange;
-		if (SegmentIndex >= 0 && SegmentIndex < Base::PairKnots.Num() - 1)
+
+		const int32 NumSegments = GetNumberOfSegments();
+		if (SegmentIndex < 0 || SegmentIndex >= NumSegments || Base::PairKnots.Num() == 0 || !Base::PairKnots.IsValidIndex(SegmentIndex))
+		{
+			return SegmentRange;
+		}
+
+		SegmentRange.Min = Base::PairKnots[SegmentIndex].Value;
+
+		if (!IsActuallyClosed())
+		{
+			// Open: knots are standard boundaries, so we still use i+1.
+			if (SegmentIndex + 1 < Base::PairKnots.Num())
+			{
+				SegmentRange.Max = Base::PairKnots[SegmentIndex + 1].Value;
+			}
+			return SegmentRange;
+		}
+
+		// Closed: PairKnots holds starts of each segment; the final segment
+		// closes the loop at Min + LoopKnotDelta.
+		if (SegmentIndex < NumSegments - 1)
 		{
-			SegmentRange.Min = Base::PairKnots[SegmentIndex].Value;
 			SegmentRange.Max = Base::PairKnots[SegmentIndex + 1].Value;
 		}
-        
+		else
+		{
+			SegmentRange.Max = SegmentRange.Min + LoopKnotDelta;
+		}
+
 		return SegmentRange;
 	}
 
@@ -1151,7 +1267,7 @@ public:
 
 		GenerateDistanceKnotsForBezier(Mode);
     	
-		UE_LOG(LogSpline, Verbose, TEXT("Set knot vector - Points: %d, Knots: %d, Mode: %d"),
+		UE_LOG(LogSpline, Verbose, TEXT("Set knot vector (Reparameterize) - Points: %d, Knots: %d, Mode: %d"),
 					   NumValues,  this->PairKnots.Num(), static_cast<int32>(Mode));
 		
 		Base::PrintKnotVector();
@@ -1159,31 +1275,83 @@ public:
 	
 	virtual FInterval1f GetParameterSpace() const override
 	{
-		return Base::PairKnots.Num() > 0
-			? FInterval1f(Base::PairKnots[0].Value, Base::PairKnots.Last().Value)
-			: FInterval1f::Empty();
+		if (Base::PairKnots.Num() == 0)
+		{
+			return FInterval1f::Empty();
+		}
+
+		const float Min = Base::PairKnots[0].Value;
+		const float LastKnot = Base::PairKnots.Last().Value;
+
+		if (!IsActuallyClosed())
+		{
+			return FInterval1f(Min, LastKnot);
+		}
+
+		// Closed: closing boundary is last real knot + LoopKnotDelta
+		return FInterval1f(Min, LastKnot + LoopKnotDelta);
 	}
 
+	/**
+	 * Gets the parameter value for a specific Bezier control point
+	 * Index is a *logical* Bezier control index:
+	 * Index = SegmentIndex * 4 + LocalPointIndex (0..3)
+	 * @param Index - Index of the Bezier control point
+	 * @return Parameter value at the specified control point
+	 */
 	virtual float GetParameter(int32 Index) const override
 	{
-		if (Base::PairKnots.IsEmpty())
+		const int32 NumSegments = GetNumberOfSegments();
+		if (NumSegments <= 0)
+		{
+			return 0.f;
+		}
+
+		const int32 MaxLogicalIndex = NumSegments * 4;
+		if (Index < 0 || Index > MaxLogicalIndex || Base::PairKnots.IsEmpty())
+		{
 			return 0.f;
+		}
 		
-        int32 NumSegments = Base::PairKnots.Num() - 1;
         int32 SegmentIndex = Index / 4;
         int32 LocalPointIndex = Index % 4;
     
-        if (SegmentIndex >= NumSegments)
-            return Base::PairKnots.Last().Value;
-    
-        FInterval1f Range(Base::PairKnots[SegmentIndex].Value, Base::PairKnots[SegmentIndex + 1].Value);
+		// Logical "one past last" control index – treat as the domain end.
+		if (SegmentIndex >= NumSegments)
+		{
+			const FInterval1f Domain = GetParameterSpace();
+			return Domain.Max;
+		}
+
+		// Use the segment parameter range helper, which already knows how to
+		// handle the closing segment for closed loops.
+		const FInterval1f Range = GetSegmentParameterRange(SegmentIndex);
+
         const float T = static_cast<float>(LocalPointIndex) / 3.0f;
         return Range.Interpolate(T);
 	}
 
+	/**
+	* Sets the parameter value for a specific Bezier control point
+	* Index is a *logical* Bezier control index:
+	* Index = SegmentIndex * 4 + LocalPointIndex (0..3)
+	* @param Index - Index of the Bezier control point
+	* @param NewParameter - New parameter value to set
+	* @return The index of the bezier point that was actually updated, or INDEX_NONE if failed
+	*/ 
 	virtual int32 SetParameter(int32 Index, float NewParameter) override
 	{
-		if (Index < 0 || Index >= Base::NumKeys() || Base::NumKeys() < 4)
+		const int32 NumSegments = GetNumberOfSegments();
+		if (NumSegments <= 0)
+		{
+			return INDEX_NONE;
+		}
+
+		// Max logical control index is 4 * NumSegments:
+		// For open: P0..P3 of each seg -> 4 * NumSegments controls, with
+		// the last P3 coincident with the last anchor.
+		const int32 MaxLogicalIndex = NumSegments * 4;
+		if (Index < 0 || Index > MaxLogicalIndex)
 		{
 			return INDEX_NONE;
 		}
@@ -1191,11 +1359,33 @@ public:
 		constexpr int32 PointsPerSegment = 4;
 		const int32 SegmentIndex = Index / PointsPerSegment;
 		const int32 LocalPointIndex = Index % PointsPerSegment;
-		int32 NumSegments = GetNumberOfSegments();
 
-		// Get the original parameter value
-		int32 OldParameterIndex = LocalPointIndex == 0 ? SegmentIndex : SegmentIndex + 1;
-		float OldParameter = Base::PairKnots[OldParameterIndex].Value;
+		const bool bClosingBoundary = IsClosingSegment(SegmentIndex) && (LocalPointIndex != 0); // P1/P2/P3 of last seg
+
+		float OldParameter = 0.f;
+		int32 OldParameterIndex = INDEX_NONE;
+
+		if (!bClosingBoundary)
+		{
+			// Legacy behaviour: real knot index
+			OldParameterIndex = (LocalPointIndex == 0) ? SegmentIndex : SegmentIndex + 1;
+			if (!Base::PairKnots.IsValidIndex(OldParameterIndex))
+			{
+				return INDEX_NONE;
+			}
+			OldParameter = Base::PairKnots[OldParameterIndex].Value;
+		}
+		else
+		{
+			// Closing boundary: virtual "last + delta"
+			if (Base::PairKnots.Num() == 0)
+			{
+				return INDEX_NONE;
+			}
+
+			const float LastKnot = Base::PairKnots.Last().Value;
+			OldParameter = LastKnot + LoopKnotDelta;
+		}
     
 		// Early exit if parameter isn't changing
 		if (OldParameter == NewParameter)
@@ -1222,7 +1412,32 @@ public:
 		}
 		
 		int32 ReturnIndex = INDEX_NONE;
-		
+
+		// --- handle closing boundary separately ---------------------------------
+		if (bClosingBoundary)
+		{
+			// We are moving the virtual closing boundary of the last segment.
+			const float LastKnot = Base::PairKnots.Last().Value;
+			const float NewLoopDelta = NewParameter - LastKnot;
+
+			// Must remain a forward-advancing interval.
+			if (NewLoopDelta <= 0.f)
+			{
+				return Index;
+			}
+
+			LoopKnotDelta = NewLoopDelta;
+			Base::MarkFlatKnotsCacheDirty();
+
+			UE_LOG(LogSpline, Verbose, TEXT("\t"));
+			UE_LOG(LogSpline, Verbose, TEXT("After SetParameter(%d, %f) [closing]:"),
+				Index, NewParameter);
+			Base::Dump();
+
+			return Index;
+		}
+		// -----------------------------------------------------------------------------
+
 		if (LocalPointIndex == 0) // P0
 		{
 			int32 LeftSegmentIndex = SegmentIndex - 1;
@@ -1295,10 +1510,9 @@ public:
 	
 	void FlipSegment(int32 Segment)
 	{
-		static constexpr int32 ControlPointsPerSegment = 4;
-		
 		const int32 NumSegments = GetNumberOfSegments();
-		if (Segment < 0 || Segment >= NumSegments)
+		const int32 NumAnchors  = GetNumAnchorPoints();
+		if (Segment < 0 || Segment >= NumSegments || NumAnchors < 2)
 		{
 			return;
 		}
@@ -1307,92 +1521,100 @@ public:
 		UE_LOG(LogSpline, Verbose, TEXT("Before FlipSegment(%d):"), Segment);
 		Base::Dump();
 		
-		// Flip the segment
+		// 1) Flip this segment's Bezier window: P0<->P3, P1<->P2
 		{
-			int32 P0 = Segment * ControlPointsPerSegment + 0;
-			int32 P1 = Segment * ControlPointsPerSegment + 1;
-			int32 P2 = Segment * ControlPointsPerSegment + 2;
-			int32 P3 = Segment * ControlPointsPerSegment + 3;
-
-			Base::Values.Swap(P0, P3);
-			Base::Values.Swap(P1, P2);
+			int32 Idx0, Idx1, Idx2, Idx3;
+			GetBezierIndices(Segment, Idx0, Idx1, Idx2, Idx3);
+
+			ValueType P0 = this->GetValue(Idx0);
+			ValueType P1 = this->GetValue(Idx1);
+			ValueType P2 = this->GetValue(Idx2);
+			ValueType P3 = this->GetValue(Idx3);
+
+			this->SetValue(Idx0, P3);
+			this->SetValue(Idx1, P2);
+			this->SetValue(Idx2, P1);
+			this->SetValue(Idx3, P0);
 		}
 
-		// Fix up neighbor segment by keeping P3 and P0 coincident and preserving the out-tangent of the end of the neighbor segment
-		if (Segment != 0)
-		{
-			// Set P3 of Segment - 1 to P0 of Segment
-			// Need to preserve P2 relative to P3 on Segment - 1
-			int32 P2 = (Segment - 1) * ControlPointsPerSegment + 2;
-			int32 P3 = (Segment - 1) * ControlPointsPerSegment + 3;
-			int32 P0 = Segment * ControlPointsPerSegment + 0;
+		// 2) Fix left neighbor (Segment-1): keep its out-tangent but move its end point to match
+	    if (Segment > 0 || (IsActuallyClosed() && Segment == 0))
+	    {
+	    	const int32 LeftSeg =
+			 (Segment > 0) ? (Segment - 1) : (NumSegments - 1); // wrap in closed case
 
-			ValueType Delta = Base::Values[P3] - Base::Values[P2];
-			Base::Values[P3] = Base::Values[P0];
-			Base::Values[P2] = Base::Values[P3] - Delta;
-		}
+	    	int32 ThisIdx0, ThisIdx1, ThisIdx2, ThisIdx3;
+	    	int32 LeftIdx0, LeftIdx1, LeftIdx2, LeftIdx3;
 
-		// Fix up neighbor segment by keeping P0 and P3 coincident and preserving the in-tangent of the beginning of the neighbor segment
-		if (Segment != NumSegments - 1)
-		{
-			// Set P0 of Segment + 1 to P3 of Segment
-			// Need to preserve P1 relative to P0 on Segment + 1
-			int32 P0 = (Segment + 1) * ControlPointsPerSegment + 0;
-			int32 P1 = (Segment + 1) * ControlPointsPerSegment + 1;
-			int32 P3 = Segment * ControlPointsPerSegment + 3;
+	    	GetBezierIndices(Segment, ThisIdx0, ThisIdx1, ThisIdx2, ThisIdx3);
+	    	GetBezierIndices(LeftSeg, LeftIdx0, LeftIdx1, LeftIdx2, LeftIdx3);
 
-			ValueType Delta = Base::Values[P0] - Base::Values[P1];
-			Base::Values[P0] = Base::Values[P3];
-			Base::Values[P1] = Base::Values[P0] - Delta;
-		}
-		
-		// Fix up the knot vector
-		Base::SwapKnots(Segment, Segment + 1);
+	    	const ValueType ThisP0 = this->GetValue(ThisIdx0);
+	    	const ValueType LeftP2 = this->GetValue(LeftIdx2);
+	    	const ValueType LeftP3 = this->GetValue(LeftIdx3);
 
-		UE_LOG(LogSpline, Verbose, TEXT("\t"));
-		UE_LOG(LogSpline, Verbose, TEXT("After FlipSegment(%d):"), Segment);
-		Base::Dump();
+	    	const ValueType Delta = LeftP3 - LeftP2; // original out-tangent
+	    	const ValueType NewEnd = ThisP0;
+	    	const ValueType NewIn  = NewEnd - Delta;
+
+	    	this->SetValue(LeftIdx3, NewEnd);
+	    	this->SetValue(LeftIdx2, NewIn);
+	    }
+
+	    // 3) Fix right neighbor (Segment+1): keep its in-tangent but move its start point to match
+	    if (Segment < NumSegments - 1 || IsClosingSegment(Segment))
+	    {
+	    	const int32 RightSeg =
+			 (Segment < NumSegments - 1) ? (Segment + 1) : 0; // wrap in closed case
+
+	    	int32 ThisIdx0, ThisIdx1, ThisIdx2, ThisIdx3;
+	    	int32 RightIdx0, RightIdx1, RightIdx2, RightIdx3;
+
+	    	GetBezierIndices(Segment, ThisIdx0, ThisIdx1, ThisIdx2, ThisIdx3);
+	    	GetBezierIndices(RightSeg, RightIdx0, RightIdx1, RightIdx2, RightIdx3);
+
+	    	const ValueType ThisP3  = this->GetValue(ThisIdx3);
+	    	const ValueType RightP0 = this->GetValue(RightIdx0);
+	    	const ValueType RightP1 = this->GetValue(RightIdx1);
+
+	    	const ValueType Delta   = RightP0 - RightP1; // original in-tangent
+	    	const ValueType NewStart = ThisP3;
+	    	const ValueType NewOut   = NewStart - Delta;
+
+	    	this->SetValue(RightIdx0, NewStart);
+	    	this->SetValue(RightIdx1, NewOut);
+	    }
+
+	    // 4) Swap knot intervals for the flipped segment & its neighbor interval
+	    Base::SwapKnots(Segment, Segment + 1);
+
+	    UE_LOG(LogSpline, Verbose, TEXT("\t"));
+	    UE_LOG(LogSpline, Verbose, TEXT("After FlipSegment(%d):"), Segment);
+	    Base::Dump();
 	}
 	
 	virtual int32 FindIndexForParameter(float Parameter, float& OutLocalParam) const override
 	{
-		int32 NumSegments = Base::PairKnots.Num() - 1;
-
-		// todo: either delete this function in favor of MapGlobalParameterToLocalSegment or fix the loop to binary search with caching
-		
-		for (int32 i = 0; i < NumSegments; ++i)
+		const int32 NumSegments = GetNumberOfSegments();
+		if (NumSegments <= 0)
 		{
-			float Start = Base::PairKnots[i].Value;
-			float End = Base::PairKnots[i + 1].Value;
-			// Exact match on last segment end point
-			if (i == NumSegments - 1 && Parameter == End)
-			{
-				OutLocalParam = 1.0f;
-				return i * 4;
-			}
+			OutLocalParam = 0.0f;
+			return 0;
+		}
 
-			if (Parameter >= Start && Parameter < End)
-			{
-				float SpanLength = End - Start;
+		int32 SegmentIndex = 0;
+		float LocalT       = 0.0f;
+		if (!MapGlobalParameterToLocalSegment(Parameter, SegmentIndex, LocalT))
+		{
+			return INDEX_NONE;
+		}
 
-				OutLocalParam = (SpanLength > FLT_EPSILON)
-					? (Parameter - Start) / SpanLength
-					: 0.0f;
+		SegmentIndex = FMath::Clamp(SegmentIndex, 0, NumSegments - 1);
+		OutLocalParam = LocalT;
 
-				return i * 4;
-			}
-		}
-		OutLocalParam = 0.0f;
-		return 0; // Fallback to first segment
-	}
-	
-	/**
-	 * Sets the closed loop flag without modifying segments
-	 * @param bClosed Whether the spline should be a closed loop
-	 */
-	void SetClosedLoopFlag(bool bClosed)
-	{
-		Base::bIsClosedLoop = bClosed;
+		// This function historically returned a "control point index" = 4*segment,
+		// i.e. the (virtual) P0 index of that segment. We preserve that contract.
+		return SegmentIndex * 4;
 	}
 	
 	virtual void SetClosedLoop(bool bShouldClose) override
@@ -1401,43 +1623,18 @@ public:
 		if (bShouldClose == this->IsClosedLoop())
 			return;
         
-		if (bShouldClose)
-		{
-			// When closing a loop, add a segment to connect the last point back to the first
-			const int32 NumSegments = GetNumberOfSegments();
-        
-			if (NumSegments >= 1) // Need at least one segment to close
-			{
-				// Get first and last points
-				ValueType FirstPoint = this->GetValue(0);
-				ValueType LastPoint = this->GetValue((NumSegments - 1) * 4 + 3);
-            
-				// If the last point isn't already equal to the first point
-				if (!Math::Equals(FirstPoint, LastPoint, UE_SMALL_NUMBER))	
-				{
-					// Calculate control points for a smooth transition
-					ValueType Direction = FirstPoint - LastPoint;
-					ValueType P1 = LastPoint + Direction * 0.33f;
-					ValueType P2 = FirstPoint - Direction * 0.33f;
-					
-					// Add the closing segment
-					AppendBezierSegment(P1, P2, FirstPoint);
-				}
-			}
-		}
-		else
-		{
-			// For an open spline, remove the closing segment
-			const int32 NumSegments = GetNumberOfSegments();
-			if (NumSegments > 0)
-			{
-				// Remove the last segment
-				this->RemoveSegment(NumSegments - 1);
-			}
-		}
-    
-		// Set the flag
-		SetClosedLoopFlag(bShouldClose);
+		// Just flip the topology flag; do not add/remove any control points.
+		Base::SetClosedLoop(bShouldClose);
+
+		this->ApplyClampedKnotsMultiplicity();
+		ApplyInternalKnotsMultiplicity();
+
+		// Note:
+		// - Knot topology (number and placement of knots) is rebuilt by Reparameterize()
+		//   when the owning spline chooses to (eg, FTangentSpline::Reparameterize).
+		// - We intentionally do NOT call AppendBezierSegment / RemoveSegment here anymore,
+		//   because the closing segment is now virtual: {A_last, Out_last, In_0, A_0}.
+
 	}
 
 	/**
@@ -1461,27 +1658,42 @@ public:
 	 */
 	bool MapGlobalParameterToLocalSegment(float GlobalParam, int32& OutSegmentIndex, float& OutLocalParam) const
 	{
-		// Handle parameter outside range
-		if (Base::PairKnots.Num() < 2)
+		const int32 NumSegments = GetNumberOfSegments();
+		const UE::Geometry::FInterval1f ParameterSpace = GetParameterSpace();
+
+		if (ParameterSpace.IsEmpty())
 		{
-			OutSegmentIndex = 0;
-			OutLocalParam = 0.0f;
 			return false;
 		}
-        
-		if (GlobalParam <= Base::PairKnots[0].Value)
+
+		// Left end: snap onto first segment start
+		if (GlobalParam <= ParameterSpace.Min)
 		{
 			OutSegmentIndex = 0;
 			OutLocalParam = 0.0f;
 			return true;
 		}
-        
-		if (GlobalParam >= Base::PairKnots.Last().Value)
+
+		// Right end: snap onto last segment end
+		if (GlobalParam >= ParameterSpace.Max)
 		{
-			OutSegmentIndex = Base::PairKnots.Num() - 2;
+			OutSegmentIndex = NumSegments - 1;
 			OutLocalParam = 1.0f;
 			return true;
 		}
+
+		const float LastKnot = Base::PairKnots.Last().Value;
+
+		// We can now assume GlobalParam is contained by ParameterSpace, no need for clamping.
+
+		// Closed loop: handle the synthetic closing segment [LastKnot, LastKnot + LoopKnotDelta]
+		if (GlobalParam >= LastKnot && ensure(IsActuallyClosed()))
+		{
+			OutSegmentIndex = NumSegments - 1;
+			const float t = (GlobalParam - LastKnot) / FMath::Max(ParameterSpace.Max - LastKnot, UE_SMALL_NUMBER);
+			OutLocalParam = FMath::Clamp(t, 0.0f, 1.0f);
+			return true;
+		}
         
 		int32 SegmentIndex = INDEX_NONE;
 
@@ -1542,43 +1754,381 @@ public:
 		return true;
 	}
 
-	/**
-	 * Gets the number of segments defined by distinct knots
-	 */
-	int32 GetNumDistinctSegments() const
-	{
-		return FMath::Max(0, Base::PairKnots.Num() - 1);
-	}
-
 	int32 FindSegmentIndex(float Parameter, float &OutLocalParam) const
 	{
 		int32 ControlPointIdx = FindIndexForParameter(Parameter, OutLocalParam);
 		return ControlPointIdx / 4; // Each segment has 4 control points
 	}
 
+	void GetBezierIndices(int32 SegmentIndex, int32& OutP0Idx, int32& OutP1Idx, int32& OutP2Idx, int32& OutP3Idx) const
+	{
+		if (!IsClosingSegment(SegmentIndex))
+		{
+			OutP0Idx = P0Idx(SegmentIndex);
+			OutP1Idx = P1Idx(SegmentIndex);
+			OutP2Idx = P2Idx(SegmentIndex);
+			OutP3Idx = P3Idx(SegmentIndex);
+		}
+		else
+		{
+			const int32 NumAnchors = GetNumAnchorPoints();
+
+			OutP0Idx = Close_P0Idx(NumAnchors);
+			OutP1Idx = Close_P1Idx(NumAnchors);
+			OutP2Idx = Close_P2Idx();
+			OutP3Idx = Close_P3Idx();
+		}
+	}
+	
+	/** Per-anchor setters/getters */
+	void SetHandleOutOfAnchor(int32 AnchorIdx, const ValueType& Out)
+	{
+		this->SetValue(AnchorBase(AnchorIdx) + 2, Out);
+	}
+	const ValueType& GetHandleOutOfAnchor(int32 AnchorIdx) const
+	{
+		return this->GetValue(AnchorBase(AnchorIdx) + 2);
+	}
+	
+	void SetHandleInOfAnchor(int32 AnchorIdx, const ValueType& In)
+	{
+		this->SetValue(AnchorBase(AnchorIdx) + 0, In);
+	}
+	const ValueType& GetHandleInOfAnchor(int32 AnchorIdx) const
+	{
+		return this->GetValue(AnchorBase(AnchorIdx) + 0);
+	}
+
+	/** Per-anchor structural inserts */
+	void InsertAnchorAt(int32 At, const ValueType& In, const ValueType& A, const ValueType& Out)
+	{
+		this->InsertValue(AnchorBase(At) + 0, In);
+		this->InsertValue(AnchorBase(At) + 1, A);
+		this->InsertValue(AnchorBase(At) + 2, Out);
+	}
+
+	/** Appends a new anchor point defined by its In, A, and Out values. */
+	void AppendAnchor(const ValueType& In, const ValueType& A, const ValueType& Out)
+	{
+		this->AddValue(In);
+		this->AddValue(A);
+		this->AddValue(Out);
+	}
+	
 	virtual int32 GetExpectedNumKnots() const override
 	{
-		const int32 NumValues = Base::NumKeys();
-		// Each segment has 4 control points. For 1 segment, we need 8 knots (4 start, 4 end).
-		int32 NumSegments = NumValues / 4;
-    
-		// Special handling for incomplete first segment
+		const int32 NumSegments = GetNumberOfSegments();	
+
 		if (NumSegments == 0)
 		{
 			return 0;
 		}
-		return 4 + (NumSegments - 1) * 3 + 4;
+
+		if (!IsActuallyClosed())
+		{
+			// Open cubic poly-Bézier:
+			// start multiplicity = 4, end multiplicity = 4, each interior boundary = 3
+			// Total = 4 + (NumSegments - 1) * 3 + 4
+			return 4 + (NumSegments - 1) * 3 + 4;
+		}
+		else
+		{
+			// Closed cubic poly-Bézier:
+			// all boundaries behave like "interior" → multiplicity = 3
+			// There are NumSegments+1 boundaries including the wrap
+			// Total = 3 * (NumSegments + 1)
+			return 3 * (NumSegments + 1);
+		}
 	}
-	
-	void SetKnotVector(const TArray<FKnot>& NewKnots)
+
+	virtual TArray<FKnot> GetKnotVector() const override
 	{
-		Base::SetCustomKnots(NewKnots);
+		TArray<FKnot> KnotVector = Base::GetKnotVector();
+
+		if (IsActuallyClosed() && ensureAlwaysMsgf(KnotVector.Num() > 0, TEXT("Knot vector is empty, cannot compute the closing knot.")))
+		{
+			KnotVector.Add(FKnot(KnotVector.Last().Value + LoopKnotDelta, Degree));
+		}
+
+		return KnotVector;
 	}
+
+	virtual bool SetKnotVector(const TArray<FKnot>& NewKnots) override
+	{
+		if (!this->ValidateKnotVector(NewKnots))
+		{
+			return false;
+		}
+
+		this->PairKnots = NewKnots;
+
+		if (IsActuallyClosed())
+		{
+			const int32 LastIdx     = this->PairKnots.Num() - 1;
+			const float ClosingKnot = this->PairKnots[LastIdx].Value;
+			const float PrevKnot    = this->PairKnots[LastIdx - 1].Value;
+
+			// Store the closing interval length and drop the closing knot.
+			LoopKnotDelta = ClosingKnot - PrevKnot;
+			ensure(LoopKnotDelta > 0.f);
+
+			this->PairKnots.RemoveAt(LastIdx, 1, EAllowShrinking::Yes);
+		}
+		
+		this->MarkFlatKnotsCacheDirty();
+
+		UE_LOG(LogSpline, Verbose, TEXT("Set knot vector (Direct) - Unique Knots: %d"),
+			   this->PairKnots.Num());
+
+		this->PrintKnotVector();
+
+		return true;
+	}
+
 	void SetFindWindowOverride(FFindWindowOverride Fn) { FindWindowOverrideFn = Fn; }
+	
 private:
+
+	
+	static constexpr int32 AnchorBase(int32 i){ return 3*i; }
+
+	static constexpr int32 P0Idx(int32 s) { return AnchorBase(s) + 1; } // A_s
+	static constexpr int32 P1Idx(int32 s) { return AnchorBase(s) + 2; } // Out_s
+	static constexpr int32 P2Idx(int32 s) { return AnchorBase(s + 1); } // In_{s+1}
+	static constexpr int32 P3Idx(int32 s) { return AnchorBase(s + 1) + 1; } // A_{s+1}
+	// closing
+	static constexpr int32 Close_P0Idx(int32 NumAnchors) { return AnchorBase(NumAnchors - 1) + 1; }
+	static constexpr int32 Close_P1Idx(int32 NumAnchors) { return AnchorBase(NumAnchors - 1) + 2; }
+	static constexpr int32 Close_P2Idx() { return AnchorBase(0); }
+	static constexpr int32 Close_P3Idx() { return AnchorBase(0) + 1; }
+
 	
 	/**
-	 * Computes chord or centripetal distances between segment endpoints (P0 to P3).
+	 * Builds the raw Bezier window for a segment index
+	 * Returns false if segment is invalid.
+	 */ 
+	bool BuildRawWindowForSegment(int32 SegmentIndex, FWindow& OutWindow) const
+	{
+		OutWindow = {};
+
+		const int32 NumAnchors = GetNumAnchorPoints();
+		if (NumAnchors < 2)
+		{
+			return false;
+		}
+
+		const int32 NumSegments = GetNumberOfSegments();
+		if (SegmentIndex < 0 || SegmentIndex >= NumSegments)
+		{
+			return false;
+		}
+
+		auto I = [this](int32 idx){ return &this->Values[idx]; };
+
+		if (!IsClosingSegment(SegmentIndex))
+		{
+			OutWindow[0] = I(P0Idx(SegmentIndex)); // A_s
+			OutWindow[1] = I(P1Idx(SegmentIndex)); // Out_s
+			OutWindow[2] = I(P2Idx(SegmentIndex)); // In_{s+1}
+			OutWindow[3] = I(P3Idx(SegmentIndex)); // A_{s+1}
+		}
+		else
+		{
+			OutWindow[0] = I(Close_P0Idx(NumAnchors)); // A_last
+			OutWindow[1] = I(Close_P1Idx(NumAnchors)); // Out_last
+			OutWindow[2] = I(Close_P2Idx()); // In_0
+			OutWindow[3] = I(Close_P3Idx()); // A_0
+		}
+
+		return Math::FSplineValidation::IsValidWindow<ValueType, Base::WindowSize>(OutWindow);
+	}
+
+	void RefreshClosingKnot(EParameterizationPolicy Policy)
+	{
+		if (!IsActuallyClosed())
+		{
+			LoopKnotDelta = 0.0f;
+			return;
+		}
+
+		const ValueType& P0 = Base::Values[ AnchorBase(GetNumAnchorPoints() - 1) + 1 ]; // A_last
+		const ValueType& P3 = Base::Values[ AnchorBase(0)                        + 1 ]; // A_0
+		LoopKnotDelta = ComputeSegmentLength(P0, P3, Policy);
+
+		Base::MarkFlatKnotsCacheDirty();
+	}
+
+	void MigrateLegacyBezierLayout()
+	{
+	    const int32 NumValues = Base::Values.Num();
+
+		// Empty spline, no migration necessary.
+		if (NumValues == 0)
+		{
+			return;
+		}
+
+		// Handle legacy case where we do not have a complete bezier segment to migrate:
+		switch (NumValues)
+		{
+			// Idx 0 is position, other values are garbage.
+			case 1: [[fallthrough]];
+			case 2: [[fallthrough]];
+			case 3:
+			{
+				ValueType In(ForceInitToZero);
+				ValueType A = Base::Values[0];
+				ValueType Out(ForceInitToZero);
+
+				Base::Values = { In, A, Out };
+				Base::PairKnots.Reset();
+				Base::MarkFlatKnotsCacheDirty();
+
+				return;
+			}
+
+			default:
+				break;
+		}
+
+	    // Legacy layout (actual):
+	    // [ P0, C0_0, C1_0, P1,
+	    //   P1_Duplicate, C0_1, C1_1, P2,
+	    //   P2_Duplicate, C0_2, C1_2, P3,
+	    //   ...
+	    //   P_{n-1}_Duplicate, C0_{n-1}, C1_{n-1}, Pn ]
+
+		// We handle cases 0-3 above
+	    checkf(NumValues >= 4,
+	        TEXT("Legacy PolyBezier must have at least 1 segment (4 values minimum)"));
+
+	    checkf(NumValues % 4 == 0,
+	        TEXT("Legacy PolyBezier layout mismatch (NumValues = %d)"), NumValues);
+
+		const bool bClosed = this->IsClosedLoop();
+
+	    const int32 NumSegments = NumValues / 4;
+
+	    const TArray<ValueType>& Legacy = this->Values;
+		// For open: anchors = NumSegments + 1
+		// For closed: anchors = NumSegments (closing segment is virtual)
+		const int32 NumAnchors = bClosed ? NumSegments : (NumSegments + 1);
+		
+	    TArray<ValueType> NewValues;
+	    NewValues.SetNumUninitialized(NumAnchors * 3);
+
+	    // Anchor positions P_j
+	    auto IndexP = [](int32 j) -> int32
+	    {
+	        // P0 is at index 0.
+	        // For j > 0, P_j is the end of the (j-1)-th segment.
+	        return (j == 0) ? 0 : (3 + 4 * (j - 1));
+	    };
+
+	    // Outgoing handle C0_j for segment j
+	    auto IndexC0 = [](int32 segIdx) -> int32
+	    {
+	        // Segment segIdx starts at 4*segIdx:
+	        // [ P_j_or_dup, C0_j, C1_j, P_{j+1} ]
+	        return 1 + 4 * segIdx;
+	    };
+
+	    // Incoming handle C1_j for segment j
+	    auto IndexC1 = [](int32 segIdx) -> int32
+	    {
+	        return 2 + 4 * segIdx;
+	    };
+
+	    for (int32 AnchorIdx = 0; AnchorIdx < NumAnchors; ++AnchorIdx)
+	    {
+	        ValueType In(ForceInitToZero);
+	        ValueType A(ForceInitToZero);
+	        ValueType Out(ForceInitToZero);
+
+	        // A_j = P_j
+	        const int32 PIndex = IndexP(AnchorIdx);
+	        check(Legacy.IsValidIndex(PIndex));
+	        A = Legacy[PIndex];
+
+	    	// In_j
+	    	if (!bClosed)
+	    	{
+	    		if (AnchorIdx == 0)
+	    		{
+	    			// No incoming segment, start flat.
+	    			In = A;
+	    		}
+	    		else
+	    		{
+	    			// From previous segment's C1
+	    			const int32 PrevSeg = AnchorIdx - 1;
+	    			const int32 C1Index = IndexC1(PrevSeg);
+	    			check(Legacy.IsValidIndex(C1Index));
+	    			In = Legacy[C1Index];
+	    		}
+	    	}
+	    	else
+	    	{
+	    		// Closed: every anchor has an incoming segment.
+	    		const int32 PrevSeg = (AnchorIdx == 0) ? (NumSegments - 1) : (AnchorIdx - 1);
+	    		const int32 C1Index = IndexC1(PrevSeg);
+	    		check(Legacy.IsValidIndex(C1Index));
+	    		In = Legacy[C1Index];
+	    	}
+
+	    	// Out_j
+	    	if (!bClosed)
+	    	{
+	    		if (AnchorIdx < NumSegments)
+	    		{
+	    			const int32 C0Index = IndexC0(AnchorIdx);
+	    			check(Legacy.IsValidIndex(C0Index));
+	    			Out = Legacy[C0Index];
+	    		}
+	    		else
+	    		{
+	    			// Last anchor has no outgoing segment in open curves.
+	    			Out = A;
+	    		}
+	    	}
+	    	else
+	    	{
+	    		// Closed: every anchor has outgoing segment j
+	    		const int32 Seg = AnchorIdx;
+	    		const int32 C0Index = IndexC0(Seg);
+	    		check(Legacy.IsValidIndex(C0Index));
+	    		Out = Legacy[C0Index];
+	    	}
+
+	        const int32 NewBase = AnchorIdx * 3;
+	        NewValues[NewBase + 0] = In;
+	        NewValues[NewBase + 1] = A;
+	        NewValues[NewBase + 2] = Out;
+	    }
+
+	    this->Values = MoveTemp(NewValues);
+
+		// Fix up knot representation for closed loops.
+		// Legacy data uses an explicit closing knot in PairKnots; the new
+		// representation uses LoopKnotDelta instead.
+		if (bClosed && this->PairKnots.Num() > 2)
+		{
+			const int32 LastIdx     = this->PairKnots.Num() - 1;
+			const float ClosingKnot = this->PairKnots[LastIdx].Value;
+			const float PrevKnot    = this->PairKnots[LastIdx - 1].Value;
+
+			// Store the closing interval length and drop the closing knot.
+			LoopKnotDelta = ClosingKnot - PrevKnot;
+
+			this->PairKnots.RemoveAt(LastIdx, 1, EAllowShrinking::Yes);
+		}
+		
+		// Values changed (and potentially PairKnots layout), so flat-knot cache is dirty.
+	    this->MarkFlatKnotsCacheDirty();
+	}
+
+	/**
+	 * Computes chord or centripetal distances between segment endpoints
 	 * Uses those distances to build a proportional knot vector with the proper multiplicities for each Bézier segment.
 	 * @param Mode - Knot generation mode to use
 	 */
@@ -1586,24 +2136,41 @@ private:
 	{
 		TRACE_CPUPROFILER_EVENT_SCOPE(GenerateDistanceKnotsForBezier);
 
-		const TArray<ValueType>& Points = this->Values;
+		// 3-tuple layout:
+		// Values[3*i + 0] = In_i
+		// Values[3*i + 1] = A_i
+		// Values[3*i + 2] = Out_i
 
-		// For a Bezier spline, the control points don't all lie on the curve
-		// So we need to compute chord lengths between segment endpoints only
-		int NumSegments = Points.Num() / 4;
-		
+		const int32 NumAnchors  = GetNumAnchorPoints();
+		const int32 NumSegments = GetNumberOfSegments();
+
+		// Need at least one segment
+		if (NumAnchors < 2 || NumSegments <= 0)
+		{
+			this->PairKnots.Reset();
+			this->MarkFlatKnotsCacheDirty();
+			return;
+		}
+			
 		TArray<float> SegmentLengths;
 		SegmentLengths.Reserve(NumSegments);
-        
+	        
+		// For a Bezier spline, the control points don't all lie on the curve
+		// So we need to compute chord lengths between segment endpoints only.
+		// Segment i goes from A_i to A_{i+1} (open) or A_{(i+1)%N} (closed).
+		const bool bClosed = IsActuallyClosed();
+
 		for (int32 i = 0; i < NumSegments; ++i)
 		{
-			const int32 StartIndex = i * 4;
-			const int32 EndIndex = StartIndex + 3;
-            
-			if (EndIndex < Points.Num())
-			{
-				SegmentLengths.Add(ComputeSegmentLength(Points[StartIndex], Points[EndIndex], Mode));
-			}
+			const int32 StartAnchorIdx = i;
+			const int32 EndAnchorIdx   = bClosed
+				? (i + 1) % NumAnchors
+				: (i + 1); // for open, NumSegments == NumAnchors - 1, so this is always valid
+	            
+			const ValueType& P0 = this->Values[AnchorBase(StartAnchorIdx) + 1]; // A_i
+			const ValueType& P3 = this->Values[AnchorBase(EndAnchorIdx)   + 1]; // A_{i+1} or A_0 (closed)
+
+			SegmentLengths.Add(ComputeSegmentLength(P0, P3, Mode));
 		}
         
 		// Now create a knot vector that respects:
@@ -1640,91 +2207,100 @@ private:
 			this->InsertKnot(FKnot(KnotValue, InMultiplicity));
 		};
 
-		// Add start knot with multiplicity 4
-		SearchParams.DesiredParameter = 0.f;
-		InsertKnotAndUpdateCache(4);
-        
-		// Calculate accumulated distances for internal knots
-		float AccumLength = 0.0f;
-		for (int32 i = 0; i < NumSegments - 1; ++i) // -1 because we don't need knots at the very end
-		{
-			if (i < SegmentLengths.Num())
-			{
-				AccumLength = SafeAdd(AccumLength, SegmentLengths[i]);
-                
-				// Use normalized accumulated length as the desired knot value
-				SearchParams.DesiredParameter = AccumLength;
-				
-				// Interior knots at segment boundaries with multiplicity 3
-				InsertKnotAndUpdateCache(3);
-			}
-		}
-
-		SearchParams.DesiredParameter = SafeAdd(AccumLength, SegmentLengths.Num() > 0
-			? SegmentLengths.Last()
-			: 0.0f);
-
-		InsertKnotAndUpdateCache(4);
-
-		// Mark flat knots cache as dirty
-		this->MarkFlatKnotsCacheDirty();
-	}
-
-    virtual void FindWindow(float Parameter, FWindow& OutWindow, FWindowStore& Scratch) const override
-    {
-        OutWindow = {};
-        if (Base::NumKeys() < 4)  // Need at least one complete Bezier segment
-        {
-            return;
-        }
-		const TArray<ValueType>& Entries = this->Values;
-        // Find which segment we're in
-    	const int32 NumSegments = Base::NumKeys() / 4;
-
-       	int32 SegmentIndex;
-    	float LocalT;
-    	MapGlobalParameterToLocalSegment(Parameter, SegmentIndex, LocalT);
-
-        if (Base::IsClosedLoop())
-        {
-            // For closed loop, wrap the segment index
-            SegmentIndex = SegmentIndex % NumSegments;
-            if (SegmentIndex < 0) SegmentIndex += NumSegments;
-        }
-        
-        // Each Bezier segment has exactly 4 control points
-        const int32 BaseIndex = SegmentIndex * 4;
-
-        if (Base::IsClosedLoop())
+		constexpr int32 EndMult = Degree + 1;
+		constexpr int32 InteriorMult = Degree;
+        if (!bClosed)
         {
-            // Handle wrapping for closed loop
-            const int32 NumPoints = Entries.Num();
-            auto WrapIndex = [NumPoints](int32 Index) -> int32
+            // OPEN:  start mult=4, interior mult=3, end mult=4
+            // Total knots = 4 + (NumSegments - 1) * 3 + 4
+    
+            // Add start knot with multiplicity EndMult = 4
+            SearchParams.DesiredParameter = 0.f;
+            InsertKnotAndUpdateCache(EndMult);
+            
+            // Calculate accumulated distances for internal knots
+            float AccumLength = 0.0f;
+            for (int32 i = 0; i < NumSegments - 1; ++i) 
             {
-                return ((Index % NumPoints) + NumPoints) % NumPoints;
-            };
-
-        	for (int Idx = 0; Idx < Base::WindowSize; ++Idx)
-        	{
-        		OutWindow[Idx] = &Entries[WrapIndex(BaseIndex + Idx)];
-        	}
+                if (i < SegmentLengths.Num())
+                {
+                    AccumLength = SafeAdd(AccumLength, SegmentLengths[i]);
+                    
+                    // Use accumulated length as the desired knot value
+                    SearchParams.DesiredParameter = AccumLength;
+                    
+                    // Interior knots at segment boundaries with multiplicity 3
+                    InsertKnotAndUpdateCache(InteriorMult);
+                }
+            }
+    
+            // Final end knot with multiplicity 4
+            SearchParams.DesiredParameter = SafeAdd(
+                AccumLength,
+                SegmentLengths.Num() > 0 ? SegmentLengths.Last() : 0.0f);
+    
+            InsertKnotAndUpdateCache(EndMult);
         }
         else
         {
-        	// Return null window if BaseIndex is too large to fully populate the window.
-            if (BaseIndex + 3 >= Entries.Num())
+            // CLOSED:
+            // All boundaries behave like interior ones, multiplicity = InteriorMult = 3.
+		    // We create one boundary per segment (at its start) and store the
+		    // total loop length in LoopKnotDelta instead of inserting a
+		    // separate closing knot.
+    
+            float AccumLength = 0.0f;
+    
+            // Boundaries 0 .. NumSegments-1 at starts of each segment
+            for (int32 BoundaryIdx = 0; BoundaryIdx < NumSegments; ++BoundaryIdx)
             {
-                return;
+                SearchParams.DesiredParameter = AccumLength;
+                InsertKnotAndUpdateCache(InteriorMult);
+    
+                if (BoundaryIdx < SegmentLengths.Num())
+                {
+                	// If this is the closing segment (last one), remember its param length
+                	if (BoundaryIdx == NumSegments - 1)
+                	{
+                		LoopKnotDelta = SegmentLengths[BoundaryIdx];
+                	}
+
+                	AccumLength = SafeAdd(AccumLength, SegmentLengths[BoundaryIdx]);
+                }
             }
-
-        	// Standard case - no wrapping
-        	for (int Idx = 0; Idx < Base::WindowSize; ++Idx)
+        	if (NumSegments == 0 || SegmentLengths.Num() == 0)
         	{
-        		OutWindow[Idx] = &Entries[BaseIndex + Idx];
+        		LoopKnotDelta = 0.0f;
         	}
         }
-		
-		if (FindWindowOverrideFn && Math::FSplineValidation::IsValidWindow<ValueType, Base::WindowSize>(OutWindow))
+        // --------------------------------------------------------------------
+
+		// Mark flat knots cache as dirty
+		this->MarkFlatKnotsCacheDirty();
+	}
+
+    virtual void FindWindow(float Parameter, FWindow& OutWindow, FWindowStore& Scratch) const override
+    {
+		OutWindow = {};
+		const int32 NumAnchors = GetNumAnchorPoints();
+		if (NumAnchors < 2)
+		{
+			return;
+		}
+
+		int32 SegmentIndex;
+		float LocalT;
+		if (!MapGlobalParameterToLocalSegment(Parameter, SegmentIndex, LocalT))
+		{
+			return;
+		}
+
+		if (!BuildRawWindowForSegment(SegmentIndex, OutWindow))
+		{
+			return;
+		}
+
+		if (FindWindowOverrideFn)
 		{
 			const FInterval1f SegmentRange = GetSegmentParameterRange(SegmentIndex);
 			FindWindowOverrideFn(OutWindow, SegmentIndex, SegmentRange, Scratch);
@@ -1735,11 +2311,15 @@ private:
     {		
         if (Window.Num() < 4)
         {
-            return ValueType();
+			return Math::FSplineValidation::GetDefaultValue<ValueType>();
         }
+
     	int32 SegmentIndex;
     	float LocalT;
-    	MapGlobalParameterToLocalSegment(Parameter, SegmentIndex, LocalT);
+		if (!MapGlobalParameterToLocalSegment(Parameter, SegmentIndex, LocalT))
+		{
+			return Math::FSplineValidation::GetDefaultValue<ValueType>();
+		}
 
         return TSplineInterpolationPolicy<ValueType>::InterpolateCubic(Window, LocalT);
     }
@@ -1757,7 +2337,7 @@ private:
 		const ValueType& Position, 
 		float& OutSquaredDistance) const
     {
-    	int32 NumSegments = Base::NumKeys() / 4;
+    	int32 NumSegments = GetNumberOfSegments();
     	// Validate segment index
     	if (SegmentIndex < 0 || SegmentIndex >= NumSegments)
     	{
@@ -1778,6 +2358,24 @@ private:
     	return (GlobalParam - SegmentRange.Min) / SegmentLength;
     }
 	
+	/**
+	 * Checks if the spline is actually closed, which is semantically different from IsClosedLoop for 0 and 1 point splines.
+	 */
+	bool IsActuallyClosed() const
+	{
+		return this->IsClosedLoop() && GetNumAnchorPoints() > 1;
+	}
+
+	/**
+	 * Checks if a segment index identifies the closing segment
+	 * @param SegmentIndex - The index to check
+	 * @return true if the segment is the closing segment
+	 */
+	bool IsClosingSegment(int32 SegmentIndex) const
+	{
+		return IsActuallyClosed() && (SegmentIndex == GetNumAnchorPoints() - 1);
+	}
+
 	/**
 	 * Checks if a segment index is valid
 	 * @param Index - Index to check
@@ -1785,42 +2383,49 @@ private:
 	 */
 	bool IsValidSegmentIndex(int32 Index) const
     {
-    	int32 NumSegments = Base::NumKeys() / 4;
+    	int32 NumSegments = GetNumberOfSegments();
     	return Index >= 0 && Index < NumSegments;
     }
 
 	/**
 	 * Helper method to add the first segment to the spline
-	 * @param Parameter - Parameter value where to insert the point (Spline Space)
 	 * @param ControlPoint - Bezier control point to insert
-	 * @param PointIndex - Index of the newly created bezier control point
+	 * @param OutPointIndex - Index of the newly created bezier anchor point
 	 * @return true if the points were added, false otherwise
 	 */
 	bool InitializeFirstSegment(
-		float Parameter,
 		const ValueType& ControlPoint,
-		int32& PointIndex)
+		int32& OutPointIndex)
     {
-    	if (Base::NumKeys() < GetDegree())
-    	{
-    		// First point - just store it directly
-    		PointIndex = Base::AddValue(ControlPoint);
-    		return true;
-    	}
+		const int32 NumAnchors = GetNumAnchorPoints();
+		// 1) No anchors yet: create A0 only, no segment.
+		if (NumAnchors == 0)
+		{
+			// Zero handles: In0 = A0, Out0 = A0
+			AppendAnchor(/*In0*/ ControlPoint, /*A0*/ ControlPoint, /*Out0*/ ControlPoint);
+			OutPointIndex = 0;
+			return true;
+		}
 
-    	// Adding the fourth control point would create a segment
-    	if (Base::NumKeys() == GetDegree())
-    	{
-    		const ValueType& P0 = Base::GetValue(0);
-    		const ValueType& P1 = Base::GetValue(1);
-    		const ValueType& P2 = Base::GetValue(2);
-    		const ValueType& P3 = ControlPoint;
+		// 2) One anchor exists, adding this point will create the first segment.
+		if (NumAnchors == 1)
+		{
+			const ValueType& A0 = Base::Values[P0Idx(0)];
+
+			// Simple auto-tangent guess between A0 and new point
+			ValueType Dir = ControlPoint - A0;
+			ValueType P0  = A0;
+			ValueType P3  = ControlPoint;
+			ValueType P1  = P0 + Dir * ValueType(0.33f);
+			ValueType P2  = P0 + Dir * ValueType(0.66f);
 	        
     		// Add the first segment
     		SetControlPoints({P0, P1, P2, P3}, EParameterizationPolicy::Centripetal);
-    		PointIndex = 3;
+    		OutPointIndex = 1;
     		return true;
     	}
+
+		// 3) We already have at least one segment – let caller take the normal path.
     	return false;
     }
 
@@ -1835,20 +2440,8 @@ private:
 		const ValueType& Position,
 		EParameterizationPolicy ParameterizationPolicy = EParameterizationPolicy::Centripetal)
 	{
-		const TOptional<float> LoopKnotDelta = GetLoopKnotDelta();
-		SetClosedLoop(false);
-		ON_SCOPE_EXIT
-		{
-			if (LoopKnotDelta.IsSet())
-			{
-				SetClosedLoop(true);
-				SetParameter(GetNumberOfSegments(), LoopKnotDelta.GetValue());
-			}
-		};
-    
-		// Create a control point at the beginning
-		ValueType P3 = Base::GetValue(0); // First point of the spline
-    
+		const ValueType& P3 = this->Values[AnchorBase(0) + 1];
+
 		// Calculate reasonable auto-tangents
 		ValueType Dir = P3 - Position;
 		ValueType P1 = Position + Dir * 0.33f;
@@ -1872,27 +2465,19 @@ private:
 		const ValueType& Position,
 		EParameterizationPolicy ParameterizationPolicy)
 	{
-		const TOptional<float> LoopKnotDelta = GetLoopKnotDelta();
-		SetClosedLoop(false);
-		ON_SCOPE_EXIT
-		{
-			if (LoopKnotDelta.IsSet())
-			{
-				SetClosedLoop(true);
-				SetParameter(GetNumberOfSegments(), LoopKnotDelta.GetValue());
-			}
-		};
-    		
+		const int32 NumAnchors = GetNumAnchorPoints();
+		const int32 LastAnchorIdx = NumAnchors - 1;
+
 		// Append - create a point at the end
-		int32 LastIndex = Base::NumKeys() - 1;
-		ValueType P0 = Base::GetValue(LastIndex);
+		const ValueType& P0 = this->Values[AnchorBase(LastAnchorIdx) + 1];
+
         
 		// Calculate reasonable auto-tangents
 		ValueType Dir = Position - P0;
 		ValueType P1 = P0 + Dir * 0.33f;
 		ValueType P2 = P0 + Dir * 0.66f;
 
-		int32 NumSegments = Base::NumKeys() / 4;
+		int32 NumSegments = GetNumberOfSegments();
 		// Add the segment
 		int32 NewIndex = AddBezierSegmentInternal(NumSegments, P0, P1, P2, Position, ParameterizationPolicy/*, StartParam*/);
 
@@ -1927,7 +2512,7 @@ private:
 		// fix internal knots to be clamped to Degree
 		for (int32 i = 1; i < Base::PairKnots.Num() - 1; ++i)
 		{
-			Base::PairKnots[i].Multiplicity = GetDegree();
+			Base::PairKnots[i].Multiplicity = Degree;
 		}
 	}
 
@@ -1940,62 +2525,105 @@ private:
 		EParameterizationPolicy ParameterizationPolicy = EParameterizationPolicy::Centripetal,
 		bool bAppendToSegment = true) // true = connect to end, false = connect to start
 	{
-		const int32 NumSegments = Base::NumKeys() / 4;
-		SegmentIndex = FMath::Clamp(SegmentIndex, 0, NumSegments);
-    
-		// Calculate segment length
-		float SegmentLength = ComputeSegmentLength(P0, P3, ParameterizationPolicy);
-    
-		// Add the control points
-		int32 ControlPointIndex = SegmentIndex * 4;
-		Base::InsertValue(ControlPointIndex, P0);
-		Base::InsertValue(ControlPointIndex + 1, P1);
-		Base::InsertValue(ControlPointIndex + 2, P2);
-		Base::InsertValue(ControlPointIndex + 3, P3);
-
-		// FInterval1f RefRange = GetSegmentParameterRange(SegmentIndex);
-		bool bAppendToEnd = bAppendToSegment && SegmentIndex > (this->GetNumDistinctSegments() - 1);
-		bool bPrependToStart = !bAppendToSegment && SegmentIndex == 0;
-	
-		// Insert knots based on append/prepend to the segment
-		if (NumSegments == 0) // First segment
-		{
-			this->ResetKnotVector();
-			this->InsertKnot( FKnot(0.0f, this->bClampEnds ? this->GetDegree() + 1 : 1));
-			this->InsertKnot( FKnot(1.f, this->bClampEnds ? this->GetDegree() + 1 : 1));
-		}
-		else if (bAppendToEnd)
-		{
-			const float LastParameter = Base::PairKnots.Last().Value;
-			const float CandidateByLength = LastParameter + SegmentLength;
-			const float CandidateByStep   = UE::Geometry::Spline::Param::NextDistinct(LastParameter);
-			const float NextParameter = FMath::Max(CandidateByLength, CandidateByStep);
-			this->InsertKnot(FKnot(NextParameter, this->bClampEnds ? this->GetDegree() : 1));
-		}
-		else if (bPrependToStart)
+		const int32 NumAnchors  = GetNumAnchorPoints();
+	    const int32 NumSegments = GetNumberOfSegments();
+
+	    // Clamp requested slot into legal insertion range
+	    SegmentIndex = FMath::Clamp(SegmentIndex, 0, NumSegments);
+
+		auto FixKnots = [this](EParameterizationPolicy Policy)
 		{
-			const float FirstParameter = Base::PairKnots[0].Value;
-			const float CandidateByLength = FirstParameter - SegmentLength;
-			const float CandidateByStep   = UE::Geometry::Spline::Param::PrevDistinct(FirstParameter);
-			const float NextParameter = FMath::Min(CandidateByLength, CandidateByStep);
-			this->InsertKnot( FKnot(NextParameter, this->bClampEnds ? this->GetDegree() : 1));
-		}
-	
-		this->ApplyClampedKnotsMultiplicity();
+			this->ApplyClampedKnotsMultiplicity();
+			ApplyInternalKnotsMultiplicity();
 
-		// make sure the internal knots have the correct multiplicity as one of the end knots
-		// might have become an internal knot
-		this->ApplyInternalKnotsMultiplicity();
-		Base::Dump();
-		return SegmentIndex;
-	}
+			if (IsActuallyClosed())
+			{
+				RefreshClosingKnot(Policy);
+			}
+		};
 
-	/** Returns the length of the parameter range of the closing segment. The returned value is unset if the spline is open. */
-	TOptional<float> GetLoopKnotDelta() const
-	{
-		return Base::IsClosedLoop()
-			? Base::PairKnots[Base::PairKnots.Num() - 1].Value - Base::PairKnots[Base::PairKnots.Num() - 2].Value
-			: TOptional<float>();
+	    // Measure for spacing when we do need to place a new boundary knot
+	    const float SegLen = ComputeSegmentLength(P0, P3, ParameterizationPolicy);
+
+	    // ===== CASE A: first-ever segment (no anchors yet) =====
+	    if (NumAnchors == 0)
+	    {
+	        // Seed A0 from P0.
+	        // IMPORTANT: make In0 a zero-handle by setting it equal to A0.
+	        // (We cannot infer a real incoming handle for A0 from {P0,P1,P2,P3} of the *first* segment.)
+	        AppendAnchor(/*In0*/ P0, /*A0*/ P0, /*Out0*/ P1);
+
+	        // Append A1 from {P2,P3}. Out1 is a zero-handle placeholder (equal to A1)
+	        AppendAnchor(/*In1*/ P2, /*A1*/ P3, /*Out1*/ P3);
+
+	        // Build a minimal clamped knot vector for 1 segment (open):
+	        this->ResetKnotVector();
+	        const int32 EndMult = Degree + 1;
+	        this->InsertKnot(FKnot(0.0f, EndMult));
+	        this->InsertKnot(FKnot(SegLen, EndMult));
+	    	FixKnots(ParameterizationPolicy);
+
+	        return 0; // the first (and only) segment index
+	    }
+
+	    // ===== CASE B: append at tail =====
+	    const bool bAppendAtTail = bAppendToSegment && (SegmentIndex >= NumSegments);
+	    if (bAppendAtTail)
+	    {
+	        const int32 Last = NumAnchors - 1;
+
+	        // Set Out_last from P1
+	        SetHandleOutOfAnchor(Last, P1);
+
+	        // Add new tail anchor from {P2,P3}; Out_new is zero-handle until more segments arrive
+	        AppendAnchor(/*In_n*/ P2, /*A_n*/ P3, /*Out_n*/ P3);
+
+	        // Insert a new boundary knot at the end (incremental)
+	        const float Kprev         = this->PairKnots.IsEmpty() ? 0.f : this->PairKnots.Last().Value;
+	        const float NextByLength  = Kprev + SegLen;
+	        const float NextByStep    = UE::Geometry::Spline::Param::NextDistinct(Kprev);
+	        const float Knew          = FMath::Max(NextByLength, NextByStep);
+	        this->InsertKnot(FKnot(Knew, Degree));
+
+	    	FixKnots(ParameterizationPolicy);
+	        return NumSegments; // newly appended segment index
+	    }
+
+	    // ===== CASE C: prepend at head =====
+	    const bool bPrependAtHead = (!bAppendToSegment && SegmentIndex == 0);
+	    if (bPrependAtHead)
+	    {
+	        // Insert new head anchor built from P0,P1
+	        InsertAnchorAt(/*At*/ 0, /*In_-1*/ P0, /*A_-1*/ P0, /*Out_-1*/ P1);
+
+	        // Old head is now anchor 1 — set its In_1 from P2
+	        SetHandleInOfAnchor(/*i=1*/ 1, P2);
+
+	        // Insert a boundary knot at the start (decremental)
+	        const float Kfirst        = this->PairKnots.IsEmpty() ? 0.f : this->PairKnots[0].Value;
+	        const float PrevByLength  = Kfirst - SegLen;
+	        const float PrevByStep    = UE::Geometry::Spline::Param::PrevDistinct(Kfirst);
+	        const float Knew          = FMath::Min(PrevByLength, PrevByStep);
+	        this->InsertKnot(FKnot(Knew, Degree));
+
+	    	FixKnots(ParameterizationPolicy);
+	        return 0;
+	    }
+
+	    // ===== CASE D: insert in the middle (between A_s and A_{s+1}) =====
+	    //   We don’t touch knots here; caller will InsertKnot(TargetParam).
+	    {
+	        const int32 s = FMath::Clamp(SegmentIndex, 0, FMath::Max(0, NumSegments - 1));
+
+	        // Left anchor’s Out_s comes from P1
+	        SetHandleOutOfAnchor(s, P1);
+
+	        // Insert new anchor at s+1 from {P2,P3}; Out_new is zero-handle
+	        InsertAnchorAt(/*At*/ s+1, /*In*/ P2, /*A*/ P3, /*Out*/ P3);
+
+			FixKnots(ParameterizationPolicy);
+	        return s;
+	    }
 	}
 
 private:
@@ -2005,6 +2633,9 @@ private:
 	
 	/** Optional override for FindWindow*/
 	FFindWindowOverride FindWindowOverrideFn = nullptr;
+
+	/** Length of the closing segment’s parameter interval. Only meaningful when IsActuallyClosed() == true. */
+	float LoopKnotDelta = 1.f;
 };
 
 
@@ -2027,4 +2658,4 @@ using FPolyBezierSpline2d = TPolyBezierSpline<FVector2d>;
 using FPolyBezierSpline3d = TPolyBezierSpline<FVector3d>;
 } // end namespace UE::Geometry::Spline
 } // end namespace UE::Geometry
-} // end namespace UE
\ No newline at end of file
+} // end namespace UE
diff --git a/Engine/Source/Runtime/GeometryCore/Public/Splines/TangentBezierSpline.h b/Engine/Source/Runtime/GeometryCore/Public/Splines/TangentBezierSpline.h
--- a/Engine/Source/Runtime/GeometryCore/Public/Splines/TangentBezierSpline.h
+++ b/Engine/Source/Runtime/GeometryCore/Public/Splines/TangentBezierSpline.h
@@ -265,15 +265,6 @@ public:
 			TangentModes.SetNum(NumTangentModes);
 
 			InstallFindWindowOverride();
-			
-			// It was once valid for a spline to have a single bezier point, but we now expect 3 for this case.
-			// It is never valid to have 1 or 2 points.
-			switch (InternalSpline.NumKeys())
-			{
-			case 1: InternalSpline.AddValue(ValueType());	// fallthrough
-			case 2: InternalSpline.AddValue(ValueType());	// fallthrough
-			default: break;
-			}
 		}
         
 		for (int32 i = 0; i < NumTangentModes; ++i)
@@ -451,72 +442,6 @@ public:
     
 		return Clone;
 	}
-
-	virtual bool IsClosedLoop() const override
-	{
-		return InternalSpline.IsClosedLoop();
-	}
-
-	virtual void SetClosedLoop(bool bInClosedLoop) override
-	{
-        // Skip if state isn't changing
-	    if (bInClosedLoop == InternalSpline.IsClosedLoop())
-	        return;
-	        
-	    if (bInClosedLoop)
-	    {
-	        const int32 NumPoints = GetNumPoints();
-	        if (NumPoints >= 2)
-	        {
-	            // Get first and last points
-	            const ValueType FirstPos = GetValue(0);
-	            const ValueType LastPos = GetValue(NumPoints - 1);
-	            
-	            // Compute tangents for the closure segment
-	            ValueType LastOutTangent, FirstInTangent;
-	            
-	            if (IsAutoTangent(NumPoints - 1) && IsAutoTangent(0))
-	            {
-	                // Auto tangents for closure
-	                ValueType ClosingTangent = (FirstPos - LastPos) * (1.0f - Tension);
-	                LastOutTangent = ClosingTangent;
-	                FirstInTangent = -ClosingTangent;
-	            }
-	            else
-	            {
-	                // Use existing tangents
-	                LastOutTangent = GetTangentOut(NumPoints - 1);
-	                FirstInTangent = GetTangentIn(0);
-	            }
-	            
-	            // Calculate control points for the closing segment
-	            ValueType P0 = LastPos;
-	            ValueType P3 = FirstPos;
-	        	ValueType P1 = TangentOps::TangentOutToP1(P0, LastOutTangent);
-	        	ValueType P2 = TangentOps::TangentInToP2(P3, FirstInTangent);
-	        	
-	            // Add the closing segment to the internal spline
-	            InternalSpline.AppendBezierSegment(P1, P2, P3, EParameterizationPolicy::Uniform);
-	        	
-	        	InternalSpline.SetClosedLoopFlag(bInClosedLoop);
-	        }
-	    }
-	    else
-	    {
-	    	// For an open spline, remove the closing segment first
-	    	// then update the internal flag
-	    	const int32 NumSegments = GetNumberOfSegments();
-	    	if (NumSegments > 0)
-	    	{
-	    		// Remove the last segment
-	    		InternalSpline.RemoveSegment(NumSegments - 1);
-	    	}
-        
-	    }
-
-		// update flag on the internal spline
-		InternalSpline.SetClosedLoopFlag(bInClosedLoop);
-    }
     
 	/**
 	 * Replaces all points in the spline with the provided points.
@@ -813,93 +738,56 @@ public:
 	 */
 	void PrependPoint(const FTangentBezierControlPoint& ControlPoint)
 	{
-		bool bWasClosed = IsClosedLoop();
-		SetClosedLoop(false);
-		ON_SCOPE_EXIT
-		{
-			if (bWasClosed)
-			{
-				SetClosedLoop(true);
-			}
-		};
-		
-    	int32 OutNewIndex;
+    	int32 NewIndex;
 
-    	// Special case: handle first two points
-    	if (InitializeFirstSegment(ControlPoint, OutNewIndex))
+		// Special case: handle first two points
+    	if (InitializeFirstSegment(ControlPoint, NewIndex))
     	{
     		return;
     	}
 		        
 		// Calculate control points for the cubic Bezier segment
 		ValueType P0 = ControlPoint.Position;
-        
-		// Handle tangent computation based on mode
-		ValueType OutTangent;
-        
-		if (ControlPoint.TangentMode == ETangentMode::User || 
-			ControlPoint.TangentMode == ETangentMode::Broken)
-		{
-			// Use the provided tangent for user-specified mode
-			OutTangent = ControlPoint.TangentOut;
-		}
-		else
-		{
-			// For auto modes, use a reasonable default that will be updated later
-			OutTangent = GetValue(0) - ControlPoint.Position;
-		}
-        
-		ValueType P1, P2;
-		P2 = TangentOps::TangentInToP2(GetValue(0), GetTangentIn(0));
-		P1 = TangentOps::TangentOutToP1(P0, OutTangent);
-		
-		// Let InternalSpline handle the prepending efficiently
+		ValueType P1 = TangentOps::TangentOutToP1(P0, ControlPoint.TangentOut);
+		ValueType P2 = TangentOps::TangentInToP2(GetValue(0), GetTangentIn(0));
+
+		// Add the segment
 		InternalSpline.PrependBezierSegment(P0, P1, P2, EParameterizationPolicy::Uniform);
         
+		// Fixup the in tangent of the new anchor point
+		SetTangentIn(0, ControlPoint.TangentIn);
+		
+		// Set tangent mode
 		TangentModes.InsertDefaulted(0);
 		SetPointTangentMode(0, ControlPoint.TangentMode);
 	}
+
     /** 
 	 * Appends a control point to the end of the spline.
 	 * @param ControlPoint Structure containing position and tangent information
-     */
+	 */
     void AppendPoint(const FTangentBezierControlPoint& ControlPoint)
     {
-		bool bWasClosed = IsClosedLoop();
-		SetClosedLoop(false);
-		ON_SCOPE_EXIT
-		{
-			if (bWasClosed)
-			{
-				SetClosedLoop(true);
-			}
-		};
-		
-    	int32 NewIndex = GetNumPoints();
+		int32 NewIndex = GetNumPoints();
 
-    	// Special case: handle first two points
+		// Special case: handle first two points
     	if (InitializeFirstSegment(ControlPoint, NewIndex))
     	{
     		return;
     	}
     	
-	    // Normal case - add segment to existing spline
-        ValueType P0 = GetValue(NewIndex - 1);
+		// Calculate control points for the cubic Bezier segment
 		ValueType P3 = ControlPoint.Position;
-    
-		// Get its outgoing tangent (may have been updated above)
-		ValueType OutTangentPrev = GetTangentOut(NewIndex - 1);
-    
-		// Calculate control points for the new segment
-		// For the incoming tangent of the new point
-		ValueType P1, P2;
-		ValueType InTangent = ControlPoint.TangentIn;
-		P2 = TangentOps::TangentInToP2(P3, InTangent);
-		P1 = TangentOps::TangentOutToP1(P0, OutTangentPrev);
+		ValueType P2 = TangentOps::TangentInToP2(P3, ControlPoint.TangentIn);
+		ValueType P1 = TangentOps::TangentOutToP1(GetValue(NewIndex - 1), GetTangentOut(NewIndex - 1));
     
 		// Add the segment
 		InternalSpline.AppendBezierSegment(P1, P2, P3, EParameterizationPolicy::Uniform);
 		
+		// Fixup the out tangent of the new anchor point
+		SetTangentOut(NewIndex, ControlPoint.TangentOut);
+
+		// Set tangent mode
 		TangentModes.InsertDefaulted(NewIndex);
 		SetPointTangentMode(NewIndex, ControlPoint.TangentMode);
     }
@@ -953,27 +841,26 @@ public:
 			TangentModes.Reset();
 			return;
 		}
-		else if (NumPoints == 2 && !IsClosedLoop())
+		else if (NumPoints == 2 && !this->IsClosedLoop())
 		{
-			ValueType PadValue = Index == 0 ? InternalSpline.GetValue(2) : InternalSpline.GetValue(1);
-		
-			// Removing the last point and tangents
-			// For CP = 0, we remove P0, P1, P2
-			// For CP = 1, we remove P1, P2, P3
-			InternalSpline.RemoveValue(Index+2);
-			InternalSpline.RemoveValue(Index+1);
-			InternalSpline.RemoveValue(Index);
-			TangentModes.RemoveAt(Index);
-			
-			// If we remove point 0:
-			// Values = < p3, p2, p2 >
-			
-			// Ff we remove point 1:
-			// Values = < p0, p1, p1 >
-
-			InternalSpline.AddValue(PadValue);
-			InternalSpline.AddValue(PadValue);
-			
+			// Figure out which point remains (0 or 1)
+			const int32 RemainingIndex = (Index == 0) ? 1 : 0;
+	
+			// Build a control point for the remaining anchor
+			FTangentBezierControlPoint RemainingCP;
+			RemainingCP.Position    = GetValue(RemainingIndex);
+			RemainingCP.TangentIn   = GetTangentIn(RemainingIndex);
+			RemainingCP.TangentOut  = GetTangentOut(RemainingIndex);
+			RemainingCP.TangentMode = TangentModes.IsValidIndex(RemainingIndex)
+										? TangentModes[RemainingIndex]
+										: ETangentMode::Auto;
+	
+			InternalSpline.Clear();
+			TangentModes.Reset();
+	
+			int32 TmpIdx = INDEX_NONE;
+			InitializeFirstSegment(RemainingCP, TmpIdx);
+	
 			return;
 		}
     
@@ -999,44 +886,16 @@ public:
 	 */
 	void ModifyPoint(int32 Index, const FTangentBezierControlPoint& ControlPoint)
 	{
-       const int32 NumPoints = GetNumPoints();
+        const int32 NumPoints = GetNumPoints();
         if (Index < 0 || Index >= NumPoints)
         {
             return;
         }
 
-        const ValueType OldPosition = GetValue(Index);
-        bool bPositionChanged = (OldPosition != ControlPoint.Position);
-		const bool bTangentModeChanged = (TangentModes[Index] != ControlPoint.TangentMode);
-
         SetPointTangentMode(Index, ControlPoint.TangentMode);
 
         // Update position
-        if (Index == 0)
-        {
-            // First point
-            UpdateBezierControlPoint(Index * 4, ControlPoint.Position);
-
-        	if (IsClosedLoop())
-        	{
-        		UpdateBezierControlPoint((NumPoints - 1) * 4 + 3, ControlPoint.Position);
-        	}
-        }
-        else if (Index == NumPoints - 1)
-        {
-            // Last point
-            UpdateBezierControlPoint(Index * 4 - 1, ControlPoint.Position);
-        	if (IsClosedLoop())
-        	{
-        		UpdateBezierControlPoint(Index * 4, ControlPoint.Position);
-        	}
-        }
-        else
-        {
-            // Middle point - affects two segments
-            UpdateBezierControlPoint((Index - 1) * 4 + 3, ControlPoint.Position);
-            UpdateBezierControlPoint(Index * 4, ControlPoint.Position);
-        }
+		InternalSpline.SetAnchorPoint(Index, ControlPoint.Position);
         
         // If not auto tangents, apply the specified tangents
         if (ControlPoint.TangentMode != ETangentMode::Auto && 
@@ -1070,33 +929,12 @@ public:
     {
     	const int32 NumPoints = GetNumPoints();
 
-    	if (NumPoints == 0 || Index < 0 || Index > GetNumPoints())
-    		return ValueType();
-    
-    	if (NumPoints == 1)
-    	{
-    		// Special case when we only have one point
-    		if (Index == 0)
-    			return InternalSpline.GetValue(0);
-    		
-    		return ValueType();
-    	}
-
-    	// Regular case - convert from point index to bezier control point index
-    	if (Index == 0)
-    	{
-    		// First point is P0 of first segment
-    		return InternalSpline.GetValue(0);
-    	}
-    	else
-    	{
-    		// Other points are P3 of previous segments
-    		const int32 ControlPointIndex = (Index - 1) * 4 + 3;
-    		if (ControlPointIndex < InternalSpline.NumKeys())
-    			return InternalSpline.GetValue(ControlPointIndex);
-    	}
+		if (NumPoints == 0 || Index < 0 || Index > GetNumPoints())
+		{
+			return ValueType();
+		}
 
-    	return ValueType();
+		return InternalSpline.GetAnchorPoint(Index);
     }
 
     /** Gets tangent at parameter in spline space */
@@ -1110,30 +948,13 @@ public:
     {
 		const int32 NumPoints = GetNumPoints();
 
-		if (NumPoints <= 1)
+		if (Index < 0 || Index >= NumPoints)
 		{
 			return ValueType();
 		}
 		
-        if (Index == 0)
-        {
-            if (IsClosedLoop())
-            {
-            	// The point at NumPoints when the spline is closed is the point that is coincident with point 0.
-                return GetTangentIn(NumPoints);
-            }
-            else
-            {
-            	// If spline is unclosed, in tangent == out tangent at first point.
-				return GetTangentOut(Index);
-            }
-        }
-		const int32 p2Idx = P2IndexForIncomingOfPoint(Index);
-		if (p2Idx < 0)
-			return ValueType();
-
 		const ValueType P3 = GetValue(Index);
-		const ValueType P2 = InternalSpline.GetValue(p2Idx);
+		const ValueType P2 = InternalSpline.GetHandleInOfAnchor(Index);
 		return TangentOps::P2ToTangentIn(P3, P2);
     }
     
@@ -1141,17 +962,12 @@ public:
 	ValueType GetTangentOut(int32 Index) const
 	{
 		const int32 NumPoints = GetNumPoints();
-		if (Index == NumPoints - 1 && !IsClosedLoop())
+		if (Index < 0 || Index >= NumPoints)
 		{
-			// If spline is unclosed, out tangent == in tangent at final point.
-			return GetTangentIn(Index);
+			return ValueType();
 		}
-		
-		const int32 p1Idx = P1IndexForOutgoingOfPoint(Index);
-		if (p1Idx < 0) return ValueType();
-
 		const ValueType P0 = GetValue(Index);
-		const ValueType P1 = InternalSpline.GetValue(p1Idx);
+		const ValueType P1 = InternalSpline.GetHandleOutOfAnchor(Index);
 
 		return TangentOps::P1ToTangentOut(P0, P1);
 	}
@@ -1163,57 +979,31 @@ public:
 		{
 			return;
 		}
-		// No incoming segment? (first point on open spline)
-		const int32 P2Index = P2IndexForIncomingOfPoint(Index);
-		if (P2Index < 0) return;
 
-		ValueType P3 = GetValue(Index);
-		ValueType Tin = NewTangent;
-		
-        ValueType P2 = TangentOps::TangentInToP2(P3, Tin);
-        UpdateBezierControlPoint(P2Index, P2);
-        
+		ValueType P3 = GetValue(Index);		
+        ValueType P2 = TangentOps::TangentInToP2(P3, NewTangent);
+
+        InternalSpline.SetHandleInOfAnchor(Index, P2);   
     }
     
     /** Sets tangent out for the specified point */
     void SetTangentOut(int32 Index, const ValueType& NewTangent)
     {
-		const int32 NumPoints = GetNumPoints();
-		if (Index < 0 || (Index >= NumPoints && !IsClosedLoop()))
+		if (Index < 0 || Index >= GetNumPoints())
 		{
 		    return;
 		}
 
-		// No outgoing segment? (last point on open spline)
-		const int32 P1Index = P1IndexForOutgoingOfPoint(Index);
-		if (P1Index < 0) return;
-
-		ValueType P0 = GetValue(Index);
-		ValueType Tout = NewTangent;
-		
-		ValueType P1 = TangentOps::TangentOutToP1(P0, Tout);
-		UpdateBezierControlPoint(P1Index, P1);
+		ValueType P0 = GetValue(Index);		
+		ValueType P1 = TangentOps::TangentOutToP1(P0, NewTangent);
       
+     	InternalSpline.SetHandleOutOfAnchor(Index, P1);
     }
 
     /** Gets the number of points in the spline */
     int32 GetNumPoints() const
     {
-    	const int32 NumKeys = InternalSpline.NumKeys();
-    
-    	switch (NumKeys)
-    	{
-    	case 0:
-    		return 0;
-    	case 1:
-    		// intentional fallthrough!
-    	case 2:
-    		ensureAlwaysMsgf(false, TEXT("Invalid number of bezier points!"));
-    	case 3:
-			return 1;
-    	default:
-			return IsClosedLoop() ? GetNumberOfSegments() : GetNumberOfSegments() + 1;
-    	}
+		return InternalSpline.GetNumAnchorPoints();
     }
 	
 
@@ -1254,7 +1044,7 @@ public:
 		{
 			// If we get here, TangentModes was in a bad state. We will correct the number of TangentModes, defaulting new points to Auto, before proceeding with the operation.
 
-			UE_LOG(LogSpline, Warning, TEXT("Index %d is an invalid index into TangentModes for %s spline with %d points."), Index, IsClosedLoop() ? TEXT("a closed") : TEXT("an open"), NumPoints);
+			UE_LOG(LogSpline, Warning, TEXT("Index %d is an invalid index into TangentModes for %s spline with %d points."), Index, this->IsClosedLoop() ? TEXT("a closed") : TEXT("an open"), NumPoints);
 
 			const int32 OldTangentModesNum = TangentModes.Num();
 			TangentModes.SetNum(NumPoints);
@@ -1274,17 +1064,19 @@ public:
 		InternalSpline.Clear();
 		TangentModes.Empty();
 	}
-	
+
+	/** Gets the parameter value for the specified anchor point index */
 	float GetParameter(int32 Index) const
 	{
-		const int32 ControlPointIndex = (Index == 0) ? 0 : ((Index - 1) * 4 + 3);
-		return InternalSpline.GetParameter(ControlPointIndex);
+		const int32 BezierControlPointIdx = (Index == 0) ? 0 : ((Index - 1) * 4 + 3);
+		return InternalSpline.GetParameter(BezierControlPointIdx);
 	}
 
+	/** Sets the parameter value for the specified anchor point index */
 	int32 SetParameter(int32 Index, float NewParameter)
     {
-		const int32 ControlPointIndex = (Index == 0) ? 0 : ((Index - 1) * 4 + 3);
-		const int32 NewIndex = (InternalSpline.SetParameter(ControlPointIndex, NewParameter) + 1) / 4;
+		const int32 BezierControlPointIdx = (Index == 0) ? 0 : ((Index - 1) * 4 + 3);
+		const int32 NewIndex = (InternalSpline.SetParameter(BezierControlPointIdx, NewParameter) + 1) / 4;
 
 		return NewIndex;
     }
@@ -1315,7 +1107,8 @@ public:
 	    {
 	        return;
 	    }
-
+		check(TangentModes.Num() == NumPoints);
+		
 		// Update tangents for each point according to its mode
 		for (int32 i = 0; i < NumPoints; ++i)
 		{
@@ -1359,8 +1152,8 @@ public:
 		const ValueType& Current = GetValue(Index);
 
 		// Calculate indices handling closed loops
-		const int32 PrevIndex = (Index > 0) ? (Index - 1) : (IsClosedLoop() ? NumPoints - 1 : Index);
-		const int32 NextIndex = (Index < NumPoints - 1) ? (Index + 1) : (IsClosedLoop() ? 0 : Index);
+		const int32 PrevIndex = (Index > 0) ? (Index - 1) : (this->IsClosedLoop() ? NumPoints - 1 : Index);
+		const int32 NextIndex = (Index < NumPoints - 1) ? (Index + 1) : (this->IsClosedLoop() ? 0 : Index);
 		const ValueType& Prev = GetValue(PrevIndex);
 		const ValueType& Next = GetValue(NextIndex);
 
@@ -1372,7 +1165,7 @@ public:
 		// If we are closed and updating the first or last point, we need to do some special work 
 		// to get virtual PrevParam or NextParam. The 'virtual' parameter is not the actual parameter 
 		// of the previous or next point, but the current parameter +- the knot delta (closing segment length).
-		if (IsClosedLoop())
+		if (this->IsClosedLoop())
 		{
 			if (Index == 0 && PrevIndex == NumPoints - 1)
 			{
@@ -1401,7 +1194,7 @@ public:
 			{
 				const bool bIsStationaryEndpoint = bStationaryEndpoint && 
 										  (Index == 0 || Index == NumPoints - 1) && 
-										  !IsClosedLoop();
+										  !this->IsClosedLoop();
         
 				if (bIsStationaryEndpoint)
 				{
@@ -1445,7 +1238,7 @@ public:
 		case ETangentMode::LegacyAuto:
 		case ETangentMode::LegacyAutoClamped:
 			{
-				if (bStationaryEndpoint && (Index == 0 || Index == NumPoints - 1) && !IsClosedLoop())
+				if (bStationaryEndpoint && (Index == 0 || Index == NumPoints - 1) && !this->IsClosedLoop())
 				{
 					// Stationary endpoints get zero tangent
 					InTangent  = ValueType(0);
@@ -1533,7 +1326,7 @@ public:
 		InternalSpline.SetKnotVector(InKnots);
 	}
 	
-	const TArray<FKnot>& GetKnotVector() const
+	TArray<FKnot> GetKnotVector() const
 	{
 		return InternalSpline.GetKnotVector();
 	}
@@ -1586,16 +1379,8 @@ public:
 	float FindNearestOnSegment(const ValueType& Point, int32 SegmentIndex, float& OutSquaredDistance) const
 		{ return InternalSpline.FindNearestOnSegment(Point, SegmentIndex, OutSquaredDistance); }
 
-protected:
-	/** Updates a Bezier control point in the internal spline */
-	void UpdateBezierControlPoint(int32 Index, const ValueType& NewValue)
-	{
-		if (Index < InternalSpline.NumKeys())
-		{
-			InternalSpline.SetValue(Index, NewValue);
-		}
-	}
 private:
+
 	void InstallFindWindowOverride()
 	{
 		InternalSpline.SetFindWindowOverride(&RescaleNormalizedHandlesForSegment);
@@ -1634,31 +1419,6 @@ private:
 			   Mode == ETangentMode::User ||
 			   Mode == ETangentMode::Broken;
 	}
-
-	/** Returns base Bezier index for the given segment */
-	static int32 SegmentBaseIndex(int32 Segment) { return 4 * Segment; }
-
-	/** Returns Bezier index for the outgoing P1 control of point i; -1 if none */
-	int32 P1IndexForOutgoingOfPoint(int32 i) const
-	{
-		const int32 N = GetNumPoints();
-		const int32 S = i;
-		if (GetNumberOfSegments() == 0) return -1;
-		if (!IsClosedLoop() && i >= N - 1) return -1; // no outgoing at last point
-		const int32 base = SegmentBaseIndex(S);
-		return base + 1; // P1
-	}
-
-	/** Returns Bezier index for the incoming P2 control of point i; -1 if none */
-	int32 P2IndexForIncomingOfPoint(int32 i) const
-	{
-		const int32 N = GetNumPoints();
-		if (GetNumberOfSegments() == 0) return -1;
-		if (!IsClosedLoop() && i == 0) return -1; // no incoming at first point
-		const int32 S = (IsClosedLoop() && i == 0) ? (GetNumberOfSegments() - 1) : (i - 1);
-		const int32 base = SegmentBaseIndex(S);
-		return base + 2; // P2
-	}
 	
 	/**
 	 * Helper method to add the first segment to the spline
@@ -1671,59 +1431,42 @@ private:
 		const int32 NumPoints = GetNumPoints();
     
     	if (NumPoints == 0)
-    	{
-    		// First point - just store it directly with both tangents, these are not used until the first segment is completed
-    		InternalSpline.AddValue(ControlPoint.Position);
-    		InternalSpline.AddValue(TangentOps::TangentOutToP1(ControlPoint.Position, ControlPoint.TangentOut));
-    		InternalSpline.AddValue(TangentOps::TangentInToP2(ControlPoint.Position, ControlPoint.TangentIn));
-    		
-    		// update tangent mode
-            if (TangentModes.IsEmpty())
-            {
-	            TangentModes.Add(ControlPoint.TangentMode);
-            }
-    		else
-    		{
-    			TangentModes[0] = ControlPoint.TangentMode;
-    		}
-    		
-    		OutPointIndex = 0;
-    		return true;
-    	}
-    	else if (NumPoints == 1)
-    	{
-    		constexpr bool bAppend = true;
-    		if (bAppend)
-    		{
-				// We have exactly one point stored - this is our second point
-				// Get the first point, then clear to start building segments
-				
-				// Calculate control points
-    			ValueType OutTangent = ControlPoint.TangentOut;
-    			ValueType InTangent = ControlPoint.TangentIn;
-    			
-				ValueType P0 = InternalSpline.GetValue(0);
-    			ValueType P1 = InternalSpline.GetValue(1);  // Already stored placeholder P1
-    			ValueType P3 = ControlPoint.Position;
-    			ValueType P2 = TangentOps::TangentInToP2(P3, ControlPoint.TangentIn);  // Second point's IN tangent
-
-				// Add the first segment
-				InternalSpline.SetControlPoints( {P0, P1, P2, P3}, EParameterizationPolicy::Uniform);
-				
-				// update tangent mode
-				TangentModes.AddDefaulted();
-				SetPointTangentMode(1, ControlPoint.TangentMode);
-				
-				OutPointIndex = 1;
-				return true;
-    		}
-    		else
-    		{
-    			// todo: prepend
-    		}
-    	}
-    	
-    	return false;
+		{
+			const ValueType A0   = ControlPoint.Position;
+			const ValueType Out0 = TangentOps::TangentOutToP1(A0, ControlPoint.TangentOut);
+			const ValueType In0  = TangentOps::TangentInToP2(A0, ControlPoint.TangentIn);
+	
+			InternalSpline.Clear();
+	
+			// Single anchor (no segments yet), but In/Out cache the tangents
+			InternalSpline.AppendAnchor(In0, A0, Out0);
+	
+			TangentModes.SetNum(1);
+			TangentModes[0] = ControlPoint.TangentMode;
+	
+			OutPointIndex = 0;
+			return true;
+		}
+		else if (NumPoints == 1)
+		{
+			// We already have anchor 0; this is the second point -> anchor 1
+	
+			const ValueType A1   = ControlPoint.Position;
+			const ValueType Out1 = TangentOps::TangentOutToP1(A1, ControlPoint.TangentOut);
+			const ValueType In1  = TangentOps::TangentInToP2(A1, ControlPoint.TangentIn);
+	
+			InternalSpline.AppendAnchor(In1, A1, Out1);
+	
+			TangentModes.Add(ControlPoint.TangentMode);
+	
+			// Rebuild parameterization if needed (Uniform or whatever policy you use)
+			InternalSpline.Reparameterize(EParameterizationPolicy::Uniform);
+	
+			OutPointIndex = 1;
+			return true;
+		}
+	
+		return false;
     }
 	void AssignInternalSpline(const FPolyBezierSpline& NewSpline)
 	{
