FDynamicMeshEditor: pre-allocate data in AppendTriangles to improve performance

#rnx
#Synced-CL 49723777

diff --git a/Engine/Source/Runtime/GeometryCore/Private/DynamicMeshEditor.cpp b/Engine/Source/Runtime/GeometryCore/Private/DynamicMeshEditor.cpp
--- a/Engine/Source/Runtime/GeometryCore/Private/DynamicMeshEditor.cpp
+++ b/Engine/Source/Runtime/GeometryCore/Private/DynamicMeshEditor.cpp
@@ -2507,22 +2507,66 @@ void FDynamicMeshEditor::AppendTriangles(const FDynamicMesh3* SourceMesh, const
 	ResultOut.Reset();
 	IndexMaps.Initialize(Mesh);
 
-	int DefaultGroupID = FDynamicMesh3::InvalidID;
+	int32 NumTriangles = 0;
 	for (int SourceTriangleID : SourceTriangles)
 	{
-		check(SourceMesh->IsTriangle(SourceTriangleID));
-		if (SourceMesh->IsTriangle(SourceTriangleID) == false)
+		if (SourceMesh->IsTriangle(SourceTriangleID))
 		{
-			continue;	// ignore missing triangles
+			NumTriangles++;
 		}
+	}
+	check(NumTriangles == SourceTriangles.Num());
+	
+	if (NumTriangles != SourceTriangles.Num())
+	{
+		return;
+	}
+
+	// reserve data
+	{
+		if (SourceMesh->HasAttributes() && Mesh->HasAttributes())
+		{
+			const int32 MaxNumTriangles = FMath::Max(Mesh->TriangleCount() + NumTriangles, Mesh->MaxTriangleID());
+			
+			const FDynamicMeshAttributeSet* SrcAttributes = SourceMesh->Attributes();
+			FDynamicMeshAttributeSet* Attributes = Mesh->Attributes();
+		
+			IndexMaps.GetTriangleMap().Reserve(NumTriangles);
+			ResultOut.NewTriangles.Reserve(NumTriangles);
+			
+			const int NumUVLayers = FMath::Min(SrcAttributes->NumUVLayers(), Attributes->NumUVLayers());
+			for (int UVLayerIndex = 0; UVLayerIndex < NumUVLayers; UVLayerIndex++)
+			{
+				IndexMaps.GetUVMap(UVLayerIndex).Reserve(NumTriangles);
+				Attributes->GetUVLayer(UVLayerIndex)->InitializeTriangles(MaxNumTriangles);
+			}
 
-		FIndex3i Tri = SourceMesh->GetTriangle(SourceTriangleID);
+			const int NumNormalLayers = FMath::Min(SrcAttributes->NumNormalLayers(), Attributes->NumNormalLayers());
+			for (int NormalLayerIndex = 0; NormalLayerIndex < NumNormalLayers; NormalLayerIndex++)
+			{
+				IndexMaps.GetNormalMap(NormalLayerIndex).Reserve(NumTriangles);
+				Attributes->GetNormalLayer(NormalLayerIndex)->InitializeTriangles(MaxNumTriangles);
+			}
 
+			if (Attributes->HasPrimaryColors())
+			{
+				IndexMaps.GetColorMap().Reserve(NumTriangles);
+				Attributes->PrimaryColors()->InitializeTriangles(MaxNumTriangles);
+			}
+		}
+	}
+
+	const bool bHasTriangleGroups = Mesh->HasTriangleGroups();
+	const bool bSourceHasTriangleGroups = SourceMesh->HasTriangleGroups();
+	
+	int DefaultGroupID = FDynamicMesh3::InvalidID;
+	for (int SourceTriangleID : SourceTriangles)
+	{
 		// FindOrCreateDuplicateGroup
 		int NewGroupID = FDynamicMesh3::InvalidID;
-		if (Mesh->HasTriangleGroups())
+		if (bHasTriangleGroups)
 		{
-			if (SourceMesh->HasTriangleGroups())
+			if (bSourceHasTriangleGroups)
 			{
 				int SourceGroupID = SourceMesh->GetTriangleGroup(SourceTriangleID);
 				if (SourceGroupID >= 0)
@@ -2551,6 +2595,7 @@ void FDynamicMeshEditor::AppendTriangles(const FDynamicMesh3* SourceMesh, const
 
 		// FindOrCreateDuplicateVertex
 		FIndex3i NewTri;
+		const FIndex3i& Tri = SourceMesh->GetTriangleRef(SourceTriangleID);
 		for (int j = 0; j < 3; ++j)
 		{
 			int SourceVertexID = Tri[j];
@@ -2577,7 +2622,6 @@ void FDynamicMeshEditor::AppendTriangles(const FDynamicMesh3* SourceMesh, const
 	}
  
 	AppendVertexAttributes(SourceMesh, Mesh, IndexMaps);
-	
 }
 
 
diff --git a/Engine/Source/Runtime/GeometryCore/Public/DynamicSubmesh3.h b/Engine/Source/Runtime/GeometryCore/Public/DynamicSubmesh3.h
--- a/Engine/Source/Runtime/GeometryCore/Public/DynamicSubmesh3.h
+++ b/Engine/Source/Runtime/GeometryCore/Public/DynamicSubmesh3.h
@@ -256,8 +256,6 @@ public:
 			Submesh.Attributes()->EnableMatchingAttributes(*BaseMesh->Attributes());
 		}
 
-		int EstVerts = SubTriangles.Num() / 2;
-
 		FDynamicMeshEditor Editor(&Submesh);
 		FDynamicMeshEditResult ResultOut;
 		Editor.AppendTriangles(BaseMesh, SubTriangles, Mappings, ResultOut, bComputeTriMaps);
