Add a safer/clearer version of FDynamicMeshEditor::SplitMesh that (1) takes an array of uniqueptr, and (2) makes the caller explicitly opt-in to the no-split-needed single mesh optimization. Migrate some unsafe usages to the new SplitMesh method.

(Old version of SplitMesh is left in for now, as there are a number of other usages that need to be updated.)

#Synced-CL 49770727

diff --git a/Engine/Source/Runtime/GeometryCore/Private/DynamicMeshEditor.cpp b/Engine/Source/Runtime/GeometryCore/Private/DynamicMeshEditor.cpp
--- a/Engine/Source/Runtime/GeometryCore/Private/DynamicMeshEditor.cpp
+++ b/Engine/Source/Runtime/GeometryCore/Private/DynamicMeshEditor.cpp
@@ -2624,7 +2624,7 @@ void FDynamicMeshEditor::AppendTriangles(const FDynamicMesh3* SourceMesh, const
 	AppendVertexAttributes(SourceMesh, Mesh, IndexMaps);
 }
 
-
+// Note: This is the legacy / direct TArray<FDynamicMesh3> version of SplitMesh
 bool FDynamicMeshEditor::SplitMesh(const FDynamicMesh3* SourceMesh, TArray<FDynamicMesh3>& SplitMeshes, TFunctionRef<int(int)> TriIDToMeshID, int DeleteMeshID, TArray<int32>* MeshIDsOut, bool bSortByID)
 {
 	using namespace UE::DynamicMeshEditorInternals;
@@ -2684,7 +2684,7 @@ bool FDynamicMeshEditor::SplitMesh(const FDynamicMesh3* SourceMesh, TArray<FDyna
 		}
 	}
 
-	if (NumMeshes == 0) // full delete case, just leave the empty mesh
+	if (NumMeshes == 0) // full delete case, just leave the empty SplitMeshes array
 	{
 		return true;
 	}
@@ -2779,6 +2779,168 @@ bool FDynamicMeshEditor::SplitMesh(const FDynamicMesh3* SourceMesh, TArray<FDyna
 	return true;
 }
 
+bool FDynamicMeshEditor::SplitMesh(const FDynamicMesh3* SourceMesh, TArray<TUniquePtr<FDynamicMesh3>>& SplitMeshes, bool bSplitIfSingle, TFunctionRef<int(int)> TriIDToMeshID, int DeleteMeshID, TArray<int32>* MeshIDsOut, bool bSortByID)
+{
+	using namespace UE::DynamicMeshEditorInternals;
+
+	TMap<int, int> MeshIDToIndex;
+	int NumMeshes = 0;
+	bool bAlsoDelete = false;
+	for (int TID : SourceMesh->TriangleIndicesItr())
+	{
+		int MeshID = TriIDToMeshID(TID);
+		if (MeshID == DeleteMeshID)
+		{
+			bAlsoDelete = true;
+			continue;
+		}
+		if (!MeshIDToIndex.Contains(MeshID))
+		{
+			MeshIDToIndex.Add(MeshID, NumMeshes++);
+		}
+	}
+
+	if (!bAlsoDelete && NumMeshes < 2 && !bSplitIfSingle)
+	{
+		return false; // nothing to do, so don't bother filling the split meshes array
+	}
+
+	if (bSortByID)
+	{
+		TArray<int32> SortedMeshIDs;
+		MeshIDToIndex.GenerateKeyArray(SortedMeshIDs);
+		SortedMeshIDs.Sort();
+		for (int32 Idx = 0; Idx < SortedMeshIDs.Num(); ++Idx)
+		{
+			MeshIDToIndex[SortedMeshIDs[Idx]] = Idx;
+		}
+	}
+
+	if (MeshIDsOut)
+	{
+		MeshIDsOut->SetNum(NumMeshes);
+		for (TPair<int32, int32> MeshIDIndexPair : MeshIDToIndex)
+		{
+			(*MeshIDsOut)[MeshIDIndexPair.Value] = MeshIDIndexPair.Key;
+		}
+	}
+
+	SplitMeshes.Reset();
+	SplitMeshes.SetNum(NumMeshes);
+
+	if (!bAlsoDelete && NumMeshes == 1)
+	{
+		SplitMeshes[0] = MakeUnique<FDynamicMesh3>(*SourceMesh);
+		return true;
+	}
+
+	// enable matching attributes
+	for (int32 Idx = 0; Idx < SplitMeshes.Num(); ++Idx)
+	{
+		SplitMeshes[Idx] = MakeUnique<FDynamicMesh3>();
+		FDynamicMesh3& M = *SplitMeshes[Idx];
+		M.EnableMeshComponents(SourceMesh->GetComponentsFlags());
+		if (SourceMesh->HasAttributes())
+		{
+			M.EnableAttributes();
+			M.Attributes()->EnableMatchingAttributes(*SourceMesh->Attributes());
+		}
+	}
+
+	if (NumMeshes == 0) // full delete case, just leave the empty SplitMeshes array
+	{
+		return true;
+	}
+
+	TArray<FMeshIndexMappings> Mappings; Mappings.Reserve(NumMeshes);
+	FDynamicMeshEditResult UnusedInvalidResultAccumulator; // only here because some functions require it
+	for (int Idx = 0; Idx < NumMeshes; Idx++)
+	{
+		FMeshIndexMappings& Map = Mappings.Emplace_GetRef();
+		Map.Initialize(SplitMeshes[Idx].Get());
+	}
+
+	for (int SourceTID : SourceMesh->TriangleIndicesItr())
+	{
+		int MeshID = TriIDToMeshID(SourceTID);
+		if (MeshID == DeleteMeshID)
+		{
+			continue; // just skip triangles w/ the Delete Mesh ID
+		}
+		int MeshIndex = MeshIDToIndex[MeshID];
+		FDynamicMesh3& Mesh = *SplitMeshes[MeshIndex];
+		FMeshIndexMappings& IndexMaps = Mappings[MeshIndex];
+
+		FIndex3i Tri = SourceMesh->GetTriangle(SourceTID);
+
+		// Find or create corresponding triangle group
+		int NewGID = FDynamicMesh3::InvalidID;
+		if (SourceMesh->HasTriangleGroups())
+		{
+			int SourceGroupID = SourceMesh->GetTriangleGroup(SourceTID);
+			if (SourceGroupID >= 0)
+			{
+				NewGID = IndexMaps.GetNewGroup(SourceGroupID);
+				if (NewGID == IndexMaps.InvalidID())
+				{
+					NewGID = Mesh.AllocateTriangleGroup();
+					IndexMaps.SetGroup(SourceGroupID, NewGID);
+				}
+			}
+		}
+
+		bool bCreatedNewVertex[3] = { false, false, false };
+		FIndex3i NewTri;
+		for (int j = 0; j < 3; ++j)
+		{
+			int SourceVID = Tri[j];
+			int NewVID = IndexMaps.GetNewVertex(SourceVID);
+			if (NewVID == IndexMaps.InvalidID())
+			{
+				bCreatedNewVertex[j] = true;
+				NewVID = Mesh.AppendVertex(*SourceMesh, SourceVID);
+				IndexMaps.SetVertex(SourceVID, NewVID);
+			}
+			NewTri[j] = NewVID;
+		}
+
+		int NewTID = Mesh.AppendTriangle(NewTri, NewGID);
+
+		// conceivably this should never happen, but it did occur due to other mesh issues,
+		// and it can be handled here without much effort
+		if (NewTID < 0)
+		{
+			// append failed, try creating separate new vertices
+			for (int j = 0; j < 3; ++j)
+			{
+				if (bCreatedNewVertex[j] == false)
+				{
+					int SourceVID = Tri[j];
+					NewTri[j] = Mesh.AppendVertex(*SourceMesh, SourceVID);
+				}
+			}
+			NewTID = Mesh.AppendTriangle(NewTri, NewGID);
+		}
+
+		if (NewTID >= 0)
+		{
+			IndexMaps.SetTriangle(SourceTID, NewTID);
+			AppendTriangleAttributes(SourceMesh, SourceTID, &Mesh, NewTID, IndexMaps, UnusedInvalidResultAccumulator);
+		}
+		else
+		{
+			checkSlow(false);
+			// something has gone very wrong, skip this triangle
+		}
+	}
+
+	for (int Idx = 0; Idx < NumMeshes; Idx++)
+	{
+		AppendVertexAttributes(SourceMesh, SplitMeshes[Idx].Get(), Mappings[Idx]);
+	}
+
+	return true;
+}
 
 template <typename RealType, int ElementSize>
 void FDynamicMeshEditor::AppendElementSubset(
diff --git a/Engine/Source/Runtime/GeometryCore/Public/DynamicMeshEditor.h b/Engine/Source/Runtime/GeometryCore/Public/DynamicMeshEditor.h
--- a/Engine/Source/Runtime/GeometryCore/Public/DynamicMeshEditor.h
+++ b/Engine/Source/Runtime/GeometryCore/Public/DynamicMeshEditor.h
@@ -624,6 +624,10 @@ public:
 	GEOMETRYCORE_API void AppendTriangles(const FDynamicMesh3* SourceMesh, const TArrayView<const int>& SourceTriangles, FMeshIndexMappings& IndexMaps, FDynamicMeshEditResult& ResultOut, bool bComputeTriangleMap = true);
 
 	/**
+	 * NOTE: Please use the TArray<TUniquePtr<FDynamicMesh3>> version of this method, below, instead of this one, because storing FDynamicMesh3 in TArray directly is not safe.
+	 * This method will not directly cause memory corruption, but any subsequent manipulation of the array will be unsafe.
+	 * TODO: Deprecate this method.
+	 * 
 	 * Create multiple meshes out of the source mesh by splitting triangles out.
 	 * Static because it creates multiple output meshes, so doesn't quite fit in the FDynamicMeshEditor model of operating on a single mesh
 	 *
@@ -637,6 +641,21 @@ public:
 	 */
 	static GEOMETRYCORE_API bool SplitMesh(const FDynamicMesh3* SourceMesh, TArray<FDynamicMesh3>& SplitMeshes, TFunctionRef<int(int)> TriIDToMeshID, int DeleteMeshID = -1, TArray<int32>* MeshIDPerSplitMesh = nullptr, bool bSortByMeshID = false);
 
+	/**
+	 * Create multiple meshes out of the source mesh by splitting triangles out.
+	 * Static because it creates multiple output meshes, so doesn't quite fit in the FDynamicMeshEditor model of operating on a single mesh
+	 *
+	 * @param SourceMesh The mesh to split
+	 * @param SplitMeshes An array with one mesh per unique Mesh ID found in SourceMesh (optionally excluding the DeleteMeshID)
+	 * @param bSplitIfSingle Whether to still split if the split would just create a single mesh with no deletion.
+	 * @param TriIDToMeshID A function mapping Triangle IDs to Mesh IDs
+	 * @param DeleteMeshID If non-negative, triangles with this Mesh ID will be removed rather than being placed in a split mesh
+	 * @param MeshIDPerSplitMesh If non-null, will be filled with the Mesh ID for each of the SplitMeshes generated
+	 * @param bSortByMeshID If true, SplitMeshes will be sorted in order of increasing Mesh ID
+	 * @return true if needed split, false if there were not multiple mesh ids, and bSplitIfSingle was false. If false, the SplitMeshes array will not be updated.
+	 */
+	static GEOMETRYCORE_API bool SplitMesh(const FDynamicMesh3* SourceMesh, TArray<TUniquePtr<FDynamicMesh3>>& SplitMeshes, bool bSplitIfSingle, TFunctionRef<int(int)> TriIDToMeshID, int DeleteMeshID = -1, TArray<int32>* MeshIDPerSplitMesh = nullptr, bool bSortByMeshID = false);
+
 	/**
 	 * Merge any seams in the given 2D attribute Overlay along the given mesh edge IDs
 	 *
