PolyBezier: Post anchor refactor cleanup

Since our refactor, we conceptually store anchor points. We evaluate bezier segments and are conceptually storing segments implicitly via their endpoints, but the methods for adding/removing segments don't make much sense anymore.

#Synced-CL 49803947
#Synced-CL 49803947

diff --git a/Engine/Source/Runtime/Engine/Classes/Curves/Splines/TangentSpline.h b/Engine/Source/Runtime/Engine/Classes/Curves/Splines/TangentSpline.h
--- a/Engine/Source/Runtime/Engine/Classes/Curves/Splines/TangentSpline.h
+++ b/Engine/Source/Runtime/Engine/Classes/Curves/Splines/TangentSpline.h
@@ -276,7 +276,7 @@ int32 FTangentSpline::AddAttributeValue(FName Name, const AttrType& Value, float
 		// Special case for empty spline
 		if (NumPoints == 0)
 		{
-			Channel->AppendPoint(ControlPoint);
+			Channel->AppendPoint(ControlPoint, 0.f);
 			SetAttributeChannelRange(Name, FInterval1f(ParentSpaceParameter, ParentSpaceParameter), EMappingRangeSpace::Parent);
 			return 0;
 		}
@@ -285,15 +285,13 @@ int32 FTangentSpline::AddAttributeValue(FName Name, const AttrType& Value, float
 		{
 			if (ParentSpaceParameter > MappedChildSpace.Min)
 			{
-				Channel->AppendPoint(ControlPoint);
-				Channel->Reparameterize(UE::Geometry::Spline::EParameterizationPolicy::Uniform);
+				Channel->AppendPoint(ControlPoint, UE::Geometry::Spline::EParameterizationPolicy::Uniform);
 				SetAttributeChannelRange(Name, FInterval1f(MappedChildSpace.Min, ParentSpaceParameter), EMappingRangeSpace::Parent);
 				return 1;
 			}
 			else
 			{
-				Channel->PrependPoint(ControlPoint);
-				Channel->Reparameterize(UE::Geometry::Spline::EParameterizationPolicy::Uniform);
+				Channel->PrependPoint(ControlPoint, UE::Geometry::Spline::EParameterizationPolicy::Uniform);
 				SetAttributeChannelRange(Name, FInterval1f(ParentSpaceParameter, MappedChildSpace.Min), EMappingRangeSpace::Parent);
 				return 0;
 			}
@@ -304,10 +302,9 @@ int32 FTangentSpline::AddAttributeValue(FName Name, const AttrType& Value, float
 		{
 			// It is important to compute ChildSpaceParameter before the AppendPoint call, otherwise it will not be correct.
 			float ChildSpaceParameter = MapParameterToChildSpace(Name, ParentSpaceParameter);
-			Channel->AppendPoint(ControlPoint);
+			Channel->AppendPoint(ControlPoint, ChildSpaceParameter);
 
 			// By growing the child space and the mapped parent range proportionally, we keep the internal points stable in parent space.
-			Channel->SetParameter(Channel->GetNumPoints() - 1, ChildSpaceParameter);
 			SetAttributeChannelRange(Name, FInterval1f(MappedChildSpace.Min, ParentSpaceParameter), EMappingRangeSpace::Parent);
 
 			return Channel->GetNumPoints() - 1;
@@ -318,20 +315,16 @@ int32 FTangentSpline::AddAttributeValue(FName Name, const AttrType& Value, float
 		{
 			// It is important to compute ChildSpaceParameter before the AppendPoint call, otherwise it will not be correct.
 			float ChildSpaceParameter = MapParameterToChildSpace(Name, ParentSpaceParameter);
-			Channel->PrependPoint(ControlPoint);
+			Channel->PrependPoint(ControlPoint, ChildSpaceParameter);
 
 			// By growing the child space and the mapped parent range proportionally, we keep the internal points stable in parent space.
-			Channel->SetParameter(0, ChildSpaceParameter);
 			SetAttributeChannelRange(Name, FInterval1f(ParentSpaceParameter, MappedChildSpace.Max), EMappingRangeSpace::Parent);
 
 			return 0;
 		}
 
 		float ChildSpaceParameter = MapParameterToChildSpace(Name, ParentSpaceParameter);
-		float LocalT;
-		int32 SegmentIndex = Channel->FindSegmentIndex(ChildSpaceParameter, LocalT);
-
-		int NewPointIndex = Channel->InsertPointAtSegmentParam(SegmentIndex, LocalT, ControlPoint);
+		int NewPointIndex = Channel->InsertPoint(ControlPoint, ChildSpaceParameter);
 		return NewPointIndex;
 	}
 
diff --git a/Engine/Source/Runtime/Engine/Private/Curves/Splines/TangentSpline.cpp b/Engine/Source/Runtime/Engine/Private/Curves/Splines/TangentSpline.cpp
--- a/Engine/Source/Runtime/Engine/Private/Curves/Splines/TangentSpline.cpp
+++ b/Engine/Source/Runtime/Engine/Private/Curves/Splines/TangentSpline.cpp
@@ -164,7 +164,7 @@ FTangentSpline::FTangentSpline(const FLegacyTangentSpline& Other)
 	}
 
 	// initialize spline
-	GetSpline().SetControlPoints(Points);
+	GetSpline().SetPoints(Points, UE::Geometry::Spline::EParameterizationPolicy::Uniform);
 	ResetRotation();
 	ResetScale();
 
@@ -220,7 +220,7 @@ FTangentSpline::FTangentSpline(const FSplineCurves& Other)
 	}
 
 	// initialize spline
-	GetSpline().SetControlPoints(Points);
+	GetSpline().SetPoints(Points, UE::Geometry::Spline::EParameterizationPolicy::Uniform);
 	ResetRotation();
 	ResetScale();
 
@@ -578,6 +578,8 @@ void FTangentSpline::ResetScale()
 
 void FTangentSpline::AddPoint(const FSplinePoint& Point)
 {
+	using namespace UE::Geometry::Spline;
+
 	ValidateRotScale();
 	ON_SCOPE_EXIT
 	{
@@ -595,17 +597,17 @@ void FTangentSpline::AddPoint(const FSplinePoint& Point)
 
 	if (NumPoints == 0 || Index >= NumPoints)
 	{
-		GetSpline().AppendPoint(ControlPoint);
+		GetSpline().AppendPoint(ControlPoint, EParameterizationPolicy::Uniform);
 
-		if (UE::Geometry::Spline::TTangentBezierSpline<FVector>* ScaleChild = GetTypedAttributeChannel<UE::Geometry::Spline::TTangentBezierSpline<FVector>>(ScaleAttrName))
+		if (TTangentBezierSpline<FVector>* ScaleChild = GetTypedAttributeChannel<TTangentBezierSpline<FVector>>(ScaleAttrName))
 		{
-			ScaleChild->AppendPoint(UE::Geometry::Spline::TTangentBezierSpline<FVector>::FTangentBezierControlPoint(Point.Scale));
+			ScaleChild->AppendPoint(TTangentBezierSpline<FVector>::FTangentBezierControlPoint(Point.Scale), EParameterizationPolicy::Uniform);
 			ScaleChild->SetKnotVector(GetSpline().GetKnotVector());
 		}
 
-		if (UE::Geometry::Spline::TTangentBezierSpline<FQuat>* RotChild = GetTypedAttributeChannel<UE::Geometry::Spline::TTangentBezierSpline<FQuat>>(RotationAttrName))
+		if (TTangentBezierSpline<FQuat>* RotChild = GetTypedAttributeChannel<TTangentBezierSpline<FQuat>>(RotationAttrName))
 		{
-			RotChild->AppendPoint(UE::Geometry::Spline::TTangentBezierSpline<FQuat>::FTangentBezierControlPoint(Point.Rotation.Quaternion()));
+			RotChild->AppendPoint(TTangentBezierSpline<FQuat>::FTangentBezierControlPoint(Point.Rotation.Quaternion()), EParameterizationPolicy::Uniform);
 			RotChild->SetKnotVector(GetSpline().GetKnotVector());
 		}
 
@@ -614,17 +616,17 @@ void FTangentSpline::AddPoint(const FSplinePoint& Point)
 	}
 	else if (Index == 0)
 	{
-		GetSpline().PrependPoint(ControlPoint);
+		GetSpline().PrependPoint(ControlPoint, EParameterizationPolicy::Uniform);
 
-		if (UE::Geometry::Spline::TTangentBezierSpline<FVector>* ScaleChild = GetTypedAttributeChannel<UE::Geometry::Spline::TTangentBezierSpline<FVector>>(ScaleAttrName))
+		if (TTangentBezierSpline<FVector>* ScaleChild = GetTypedAttributeChannel<TTangentBezierSpline<FVector>>(ScaleAttrName))
 		{
-			ScaleChild->PrependPoint(UE::Geometry::Spline::TTangentBezierSpline<FVector>::FTangentBezierControlPoint(Point.Scale));
+			ScaleChild->PrependPoint(TTangentBezierSpline<FVector>::FTangentBezierControlPoint(Point.Scale), EParameterizationPolicy::Uniform);
 			ScaleChild->SetKnotVector(GetSpline().GetKnotVector());
 		}
 
-		if (UE::Geometry::Spline::TTangentBezierSpline<FQuat>* RotChild = GetTypedAttributeChannel<UE::Geometry::Spline::TTangentBezierSpline<FQuat>>(RotationAttrName))
+		if (TTangentBezierSpline<FQuat>* RotChild = GetTypedAttributeChannel<TTangentBezierSpline<FQuat>>(RotationAttrName))
 		{
-			RotChild->PrependPoint(UE::Geometry::Spline::TTangentBezierSpline<FQuat>::FTangentBezierControlPoint(Point.Rotation.Quaternion()));
+			RotChild->PrependPoint(TTangentBezierSpline<FQuat>::FTangentBezierControlPoint(Point.Rotation.Quaternion()), EParameterizationPolicy::Uniform);
 			RotChild->SetKnotVector(GetSpline().GetKnotVector());
 		}
 
@@ -633,20 +635,22 @@ void FTangentSpline::AddPoint(const FSplinePoint& Point)
 	}
 
 	// Insert at calculated parameter
-	int32 InsertedIndex = GetSpline().InsertPointAtSegmentParam(Index, Fraction, ControlPoint);
-	float NewPointInternalParam = GetSpline().GetParameter(InsertedIndex);
+	const int32 SegmentIndex = Index - 1;
+	const float Parameter = GetSpline().GetSegmentParameterRange(SegmentIndex).Interpolate(Fraction);
 
-	float ScaleParam = MapParameterToChildSpace(ScaleAttrName, NewPointInternalParam);
-	if (UE::Geometry::Spline::TTangentBezierSpline<FVector>* ScaleChild = GetTypedAttributeChannel<UE::Geometry::Spline::TTangentBezierSpline<FVector>>(ScaleAttrName))
+	int32 InsertedIndex = GetSpline().InsertPoint(ControlPoint, Parameter);
+
+	float ScaleParam = MapParameterToChildSpace(ScaleAttrName, Parameter);
+	if (TTangentBezierSpline<FVector>* ScaleChild = GetTypedAttributeChannel<TTangentBezierSpline<FVector>>(ScaleAttrName))
 	{
-		ScaleChild->InsertPointAtGlobalParam(ScaleParam, UE::Geometry::Spline::TTangentBezierSpline<FVector>::FTangentBezierControlPoint(Point.Scale));
+		ScaleChild->InsertPoint(TTangentBezierSpline<FVector>::FTangentBezierControlPoint(Point.Scale), ScaleParam);
 		ScaleChild->SetKnotVector(GetSpline().GetKnotVector());
 	}
 
-	float RotParam = MapParameterToChildSpace(RotationAttrName, NewPointInternalParam);
-	if (UE::Geometry::Spline::TTangentBezierSpline<FQuat>* RotChild = GetTypedAttributeChannel<UE::Geometry::Spline::TTangentBezierSpline<FQuat>>(RotationAttrName))
+	float RotParam = MapParameterToChildSpace(RotationAttrName, Parameter);
+	if (TTangentBezierSpline<FQuat>* RotChild = GetTypedAttributeChannel<TTangentBezierSpline<FQuat>>(RotationAttrName))
 	{
-		RotChild->InsertPointAtGlobalParam(RotParam, UE::Geometry::Spline::TTangentBezierSpline<FQuat>::FTangentBezierControlPoint(Point.Rotation.Quaternion()));
+		RotChild->InsertPoint(TTangentBezierSpline<FQuat>::FTangentBezierControlPoint(Point.Rotation.Quaternion()), RotParam);
 		RotChild->SetKnotVector(GetSpline().GetKnotVector());
 	}
 
@@ -655,6 +659,8 @@ void FTangentSpline::AddPoint(const FSplinePoint& Point)
 
 void FTangentSpline::InsertPoint(const FSplinePoint& Point, int32 Index)
 {
+	using namespace UE::Geometry::Spline;
+
 	ValidateRotScale();
 	ON_SCOPE_EXIT
 	{
@@ -669,17 +675,17 @@ void FTangentSpline::InsertPoint(const FSplinePoint& Point, int32 Index)
 	// Special case for empty spline
 	if (NumPoints == 0 || Index >= NumPoints)
 	{
-		GetSpline().AppendPoint(ControlPoint);
+		GetSpline().AppendPoint(ControlPoint, EParameterizationPolicy::Uniform);
 
-		if (UE::Geometry::Spline::TTangentBezierSpline<FVector>* ScaleChild = GetTypedAttributeChannel<UE::Geometry::Spline::TTangentBezierSpline<FVector>>(ScaleAttrName))
+		if (TTangentBezierSpline<FVector>* ScaleChild = GetTypedAttributeChannel<TTangentBezierSpline<FVector>>(ScaleAttrName))
 		{
-			ScaleChild->AppendPoint(UE::Geometry::Spline::TTangentBezierSpline<FVector>::FTangentBezierControlPoint(Point.Scale));
+			ScaleChild->AppendPoint(TTangentBezierSpline<FVector>::FTangentBezierControlPoint(Point.Scale), EParameterizationPolicy::Uniform);
 			ScaleChild->SetKnotVector(GetSpline().GetKnotVector());
 		}
 
-		if (UE::Geometry::Spline::TTangentBezierSpline<FQuat>* RotChild = GetTypedAttributeChannel<UE::Geometry::Spline::TTangentBezierSpline<FQuat>>(RotationAttrName))
+		if (TTangentBezierSpline<FQuat>* RotChild = GetTypedAttributeChannel<TTangentBezierSpline<FQuat>>(RotationAttrName))
 		{
-			RotChild->AppendPoint(UE::Geometry::Spline::TTangentBezierSpline<FQuat>::FTangentBezierControlPoint(Point.Rotation.Quaternion()));
+			RotChild->AppendPoint(TTangentBezierSpline<FQuat>::FTangentBezierControlPoint(Point.Rotation.Quaternion()), EParameterizationPolicy::Uniform);
 			RotChild->SetKnotVector(GetSpline().GetKnotVector());
 		}
 
@@ -688,17 +694,17 @@ void FTangentSpline::InsertPoint(const FSplinePoint& Point, int32 Index)
 	}
 	else if (Index == 0)
 	{
-		GetSpline().PrependPoint(ControlPoint);
+		GetSpline().PrependPoint(ControlPoint, EParameterizationPolicy::Uniform);
 
-		if (UE::Geometry::Spline::TTangentBezierSpline<FVector>* ScaleChild = GetTypedAttributeChannel<UE::Geometry::Spline::TTangentBezierSpline<FVector>>(ScaleAttrName))
+		if (TTangentBezierSpline<FVector>* ScaleChild = GetTypedAttributeChannel<TTangentBezierSpline<FVector>>(ScaleAttrName))
 		{
-			ScaleChild->PrependPoint(UE::Geometry::Spline::TTangentBezierSpline<FVector>::FTangentBezierControlPoint(Point.Scale));
+			ScaleChild->PrependPoint(TTangentBezierSpline<FVector>::FTangentBezierControlPoint(Point.Scale), EParameterizationPolicy::Uniform);
 			ScaleChild->SetKnotVector(GetSpline().GetKnotVector());
 		}
 
-		if (UE::Geometry::Spline::TTangentBezierSpline<FQuat>* RotChild = GetTypedAttributeChannel<UE::Geometry::Spline::TTangentBezierSpline<FQuat>>(RotationAttrName))
+		if (TTangentBezierSpline<FQuat>* RotChild = GetTypedAttributeChannel<TTangentBezierSpline<FQuat>>(RotationAttrName))
 		{
-			RotChild->PrependPoint(UE::Geometry::Spline::TTangentBezierSpline<FQuat>::FTangentBezierControlPoint(Point.Rotation.Quaternion()));
+			RotChild->PrependPoint(TTangentBezierSpline<FQuat>::FTangentBezierControlPoint(Point.Rotation.Quaternion()), EParameterizationPolicy::Uniform);
 			RotChild->SetKnotVector(GetSpline().GetKnotVector());
 		}
 
@@ -706,22 +712,26 @@ void FTangentSpline::InsertPoint(const FSplinePoint& Point, int32 Index)
 		return;
 	}
 
-	Index = FMath::Clamp(Index, 0, GetNumberOfPoints());
+	// We handle 0 and NumPoints cases above
+	Index = FMath::Clamp(Index, 1, GetNumberOfPoints() - 1);
 
-	int32 InsertedIndex = GetSpline().InsertPointAtPosition(Index, ControlPoint);
-	float NewPointInternalParam = GetSpline().GetParameter(InsertedIndex);
+	const int32 SegmentIndex = Index - 1;
+	float SegmentParam;
+	const float Parameter = GetSpline().FindNearestOnSegment(Point.Position, SegmentIndex, SegmentParam);
 
-	float ScaleParam = MapParameterToChildSpace(ScaleAttrName, NewPointInternalParam);
-	if (UE::Geometry::Spline::TTangentBezierSpline<FVector>* ScaleChild = GetTypedAttributeChannel<UE::Geometry::Spline::TTangentBezierSpline<FVector>>(ScaleAttrName))
+	int32 InsertedIndex = GetSpline().InsertPoint(ControlPoint, Parameter);
+
+	float ScaleParam = MapParameterToChildSpace(ScaleAttrName, Parameter);
+	if (TTangentBezierSpline<FVector>* ScaleChild = GetTypedAttributeChannel<TTangentBezierSpline<FVector>>(ScaleAttrName))
 	{
-		ScaleChild->InsertPointAtGlobalParam(ScaleParam, UE::Geometry::Spline::TTangentBezierSpline<FVector>::FTangentBezierControlPoint(Point.Scale));
+		ScaleChild->InsertPoint(TTangentBezierSpline<FVector>::FTangentBezierControlPoint(Point.Scale), ScaleParam);
 		ScaleChild->SetKnotVector(GetSpline().GetKnotVector());
 	}
 
-	float RotParam = MapParameterToChildSpace(RotationAttrName, NewPointInternalParam);
-	if (UE::Geometry::Spline::TTangentBezierSpline<FQuat>* RotChild = GetTypedAttributeChannel<UE::Geometry::Spline::TTangentBezierSpline<FQuat>>(RotationAttrName))
+	float RotParam = MapParameterToChildSpace(RotationAttrName, Parameter);
+	if (TTangentBezierSpline<FQuat>* RotChild = GetTypedAttributeChannel<TTangentBezierSpline<FQuat>>(RotationAttrName))
 	{
-		RotChild->InsertPointAtGlobalParam(RotParam, UE::Geometry::Spline::TTangentBezierSpline<FQuat>::FTangentBezierControlPoint(Point.Rotation.Quaternion()));
+		RotChild->InsertPoint(TTangentBezierSpline<FQuat>::FTangentBezierControlPoint(Point.Rotation.Quaternion()), RotParam);
 		RotChild->SetKnotVector(GetSpline().GetKnotVector());
 	}
 
@@ -752,6 +762,8 @@ FSplinePoint FTangentSpline::GetPoint(int32 Index) const
 /** Removes a point from the spline */
 void FTangentSpline::RemovePoint(int32 Index)
 {
+	using namespace UE::Geometry::Spline;
+
 	ValidateRotScale();
 	ON_SCOPE_EXIT
 	{
@@ -768,13 +780,13 @@ void FTangentSpline::RemovePoint(int32 Index)
 
 	GetSpline().RemovePoint(Index);
 
-	if (UE::Geometry::Spline::TTangentBezierSpline<FVector>* ScaleChild = GetTypedAttributeChannel<UE::Geometry::Spline::TTangentBezierSpline<FVector>>(ScaleAttrName))
+	if (TTangentBezierSpline<FVector>* ScaleChild = GetTypedAttributeChannel<TTangentBezierSpline<FVector>>(ScaleAttrName))
 	{
 		ScaleChild->RemovePoint(Index);
 		ScaleChild->SetKnotVector(GetSpline().GetKnotVector());
 	}
 
-	if (UE::Geometry::Spline::TTangentBezierSpline<FQuat>* RotChild = GetTypedAttributeChannel<UE::Geometry::Spline::TTangentBezierSpline<FQuat>>(RotationAttrName))
+	if (TTangentBezierSpline<FQuat>* RotChild = GetTypedAttributeChannel<TTangentBezierSpline<FQuat>>(RotationAttrName))
 	{
 		RotChild->RemovePoint(Index);
 		RotChild->SetKnotVector(GetSpline().GetKnotVector());
@@ -1098,6 +1110,8 @@ void FTangentSpline::UpdateSpline()
 
 void FTangentSpline::UpdateSpline(const FUpdateSplineParams& Params)
 {
+	using namespace UE::Geometry::Spline;
+
 	TRACE_CPUPROFILER_EVENT_SCOPE(FTangentSpline::UpdateSpline);
 
 	// todo: handled unhandled params
@@ -1134,12 +1148,12 @@ void FTangentSpline::UpdateSpline(const FUpdateSplineParams& Params)
 				{
 					GetSpline().SetParameter(GetNumberOfPoints(), LoopKey);
 
-					if (UE::Geometry::Spline::TTangentBezierSpline<FVector>* ScaleChild = GetTypedAttributeChannel<UE::Geometry::Spline::TTangentBezierSpline<FVector>>(ScaleAttrName))
+					if (TTangentBezierSpline<FVector>* ScaleChild = GetTypedAttributeChannel<TTangentBezierSpline<FVector>>(ScaleAttrName))
 					{
 						ScaleChild->SetParameter(GetNumberOfPoints(), LoopKey);
 					}
 
-					if (UE::Geometry::Spline::TTangentBezierSpline<FQuat>* RotChild = GetTypedAttributeChannel<UE::Geometry::Spline::TTangentBezierSpline<FQuat>>(RotationAttrName))
+					if (TTangentBezierSpline<FQuat>* RotChild = GetTypedAttributeChannel<TTangentBezierSpline<FQuat>>(RotationAttrName))
 					{
 						RotChild->SetParameter(GetNumberOfPoints(), LoopKey);
 					}
@@ -1153,6 +1167,8 @@ void FTangentSpline::UpdateSpline(const FUpdateSplineParams& Params)
 
 void FTangentSpline::Reparameterize()
 {
+	using namespace UE::Geometry::Spline;
+
 	TRACE_CPUPROFILER_EVENT_SCOPE(FTangentSpline::Reparameterize);
 
 	// Caching segment space attribute parameters so we can restore them (necessary because we store the mapping in parent space which are about to be invalidated).
@@ -1198,19 +1214,19 @@ void FTangentSpline::Reparameterize()
 		}
 
 		RestoreChannelKnots(ChannelName, CachedChannelParameterizations[ChannelName].Value);
-		if (UE::Geometry::Spline::TTangentBezierSpline<float>* Channel = GetTypedAttributeChannel<UE::Geometry::Spline::TTangentBezierSpline<float>>(ChannelName))
+		if (TTangentBezierSpline<float>* Channel = GetTypedAttributeChannel<TTangentBezierSpline<float>>(ChannelName))
 		{
 			Channel->UpdateTangents();
 		}
 	}
 
-	if (UE::Geometry::Spline::TTangentBezierSpline<FVector>* ScaleSpline = GetTypedAttributeChannel<UE::Geometry::Spline::TTangentBezierSpline<FVector>>(ScaleAttrName))
+	if (TTangentBezierSpline<FVector>* ScaleSpline = GetTypedAttributeChannel<TTangentBezierSpline<FVector>>(ScaleAttrName))
 	{
 		ScaleSpline->SetKnotVector(GetSpline().GetKnotVector());
 		ScaleSpline->UpdateTangents();
 	}
 
-	if (UE::Geometry::Spline::TTangentBezierSpline<FQuat>* RotSpline = GetTypedAttributeChannel<UE::Geometry::Spline::TTangentBezierSpline<FQuat>>(RotationAttrName))
+	if (TTangentBezierSpline<FQuat>* RotSpline = GetTypedAttributeChannel<TTangentBezierSpline<FQuat>>(RotationAttrName))
 	{
 		RotSpline->SetKnotVector(GetSpline().GetKnotVector());
 		RotSpline->UpdateTangents();
diff --git a/Engine/Source/Runtime/GeometryCore/Public/Splines/BSpline.h b/Engine/Source/Runtime/GeometryCore/Public/Splines/BSpline.h
--- a/Engine/Source/Runtime/GeometryCore/Public/Splines/BSpline.h
+++ b/Engine/Source/Runtime/GeometryCore/Public/Splines/BSpline.h
@@ -16,7 +16,7 @@ namespace Spline
 {
 
 struct UE_EXPERIMENTAL(5.7, "New spline APIs are experimental.") FKnot;
-template <typename VALUETYPE, int32 DEGREE> class UE_EXPERIMENTAL(5.7, "New spline APIs are experimental.") TBSpline;
+template <typename ValueType, int32 Degree> class UE_EXPERIMENTAL(5.7, "New spline APIs are experimental.") TBSpline;
 	
 enum class UE_EXPERIMENTAL(5.7, "New spline APIs are experimental.") EParameterizationPolicy
 {
@@ -119,22 +119,22 @@ struct BSplineNameSelector<3>
 	static constexpr const TCHAR* Name = TEXT("BSpline3");
 };
 
-template <typename VALUETYPE, int32 DEGREE>
+template <typename InValueType, int32 InDegree>
 class TBSpline
-	: public TSplineInterface<VALUETYPE>,
-	  private TSelfRegisteringSpline<TBSpline<VALUETYPE, DEGREE>, VALUETYPE>
+	: public TSplineInterface<InValueType>,
+	  private TSelfRegisteringSpline<TBSpline<InValueType, InDegree>, InValueType>
 {
 public:
-	static constexpr int32 Degree = DEGREE;
-	using ValueType = typename TSplineInterface<VALUETYPE>::ValueType;
+	using ValueType = InValueType;
 
+	static constexpr int32 Degree = InDegree;
 	static constexpr int32 WindowSize = Degree + 1;
 	using FWindow = TStaticArray<const ValueType*, WindowSize>;
 	using FWindowStore = TStaticArray<ValueType, WindowSize>;
 	
 	DECLARE_SPLINE_TYPE_ID(
-		BSplineNameSelector<DEGREE>::Name,
-		*TSplineValueTypeTraits<VALUETYPE>::Name
+		BSplineNameSelector<Degree>::Name,
+		*TSplineValueTypeTraits<ValueType>::Name
 	);
 
 	TBSpline() = default;
@@ -238,7 +238,7 @@ public:
 
 	virtual TUniquePtr<ISplineInterface> Clone() const override
 	{
-		TUniquePtr<TBSpline<VALUETYPE, DEGREE>> Clone = MakeUnique<TBSpline<VALUETYPE, DEGREE>>();
+		TUniquePtr<TBSpline<ValueType, Degree>> Clone = MakeUnique<TBSpline<ValueType, Degree>>();
 
 		// Copy all member variables
 		Clone->Values = this->Values;
@@ -641,9 +641,9 @@ protected:
 	/**
 	 * Inserts a knot at the specified parameter value
 	 * @param InKnot - Knot to insert
-	 * @return true if knot insertion was successful, false otherwise.
+	 * @return index of new knot if knot insertion was successful, INDEX_NONE otherwise.
 	 */
-	bool InsertKnot(FKnot InKnot)
+	int32 InsertKnot(FKnot InKnot)
 	{
 		// Binary search the knot vector:
 
@@ -657,7 +657,7 @@ protected:
 			if (PairKnots[Mid].Value == InKnot.Value)
 			{
 				UE_LOG(LogSpline, Warning, TEXT("Knot insertion failed, knot already exists!"));
-				return false;
+				return INDEX_NONE;
 			}
 			else if (PairKnots[Mid].Value < InKnot.Value)
 			{
@@ -669,10 +669,8 @@ protected:
 			}
 		}
 
-		PairKnots.Insert(InKnot, Low);
 		MarkFlatKnotsCacheDirty();
-
-		return true;
+		return PairKnots.Insert(InKnot, Low);
 	}
 
 	struct FValidKnotSearchParams
diff --git a/Engine/Source/Runtime/GeometryCore/Public/Splines/PolyBezierSpline.h b/Engine/Source/Runtime/GeometryCore/Public/Splines/PolyBezierSpline.h
--- a/Engine/Source/Runtime/GeometryCore/Public/Splines/PolyBezierSpline.h
+++ b/Engine/Source/Runtime/GeometryCore/Public/Splines/PolyBezierSpline.h
@@ -17,59 +17,65 @@ namespace Private
 	GEOMETRYCORE_API bool UseFindNearestUseSplineCulling();
 }
 
-template<typename VALUETYPE> class UE_EXPERIMENTAL(5.7, "New spline APIs are experimental.") TPolyBezierSpline;
+template<typename ValueType> class UE_EXPERIMENTAL(5.7, "New spline APIs are experimental.") TPolyBezierSpline;
 	
 /* Augments B-Spline interface */
-template<typename VALUETYPE>
+template<typename InValueType>
 class TPolyBezierSpline :
-	public TBSpline<VALUETYPE, 3>,
-	private TSelfRegisteringSpline<TPolyBezierSpline<VALUETYPE>, VALUETYPE>
+	public TBSpline<InValueType, 3>,
+	private TSelfRegisteringSpline<TPolyBezierSpline<InValueType>, InValueType>
 {
+
 public:
-	using Base = TBSpline<VALUETYPE, 3>;
+	
+	using ValueType = InValueType;
+	using Base = TBSpline<ValueType, 3>;
 	using Base::Degree;
-	typedef typename Base::ValueType ValueType;
+
 	using FWindow = typename Base::FWindow;
 	using FWindowStore = typename Base::FWindowStore;
 	using FFindWindowOverride = void(*)(FWindow& /*InOut*/,
-										int32 /*SegmentIndex*/,
-										const FInterval1f& /*SegmentRange*/,
-										FWindowStore& /*Scratch*/);
+								int32 /*SegmentIndex*/,
+								const FInterval1f& /*SegmentRange*/,
+								FWindowStore& /*Scratch*/);
 	
+private:
+
+	using typename Base::FKnotSearchCache;
+	using typename Base::FValidKnotSearchParams;
+	using Base::ApplyClampedKnotsMultiplicity;
+
+public:
+
 	// Generate compile-time type ID for PolyBezier
 	DECLARE_SPLINE_TYPE_ID(
 		TEXT("PolyBezier"),
-		*TSplineValueTypeTraits<VALUETYPE>::Name
+		*TSplineValueTypeTraits<ValueType>::Name
 	);
 	
     TPolyBezierSpline() = default;
     virtual ~TPolyBezierSpline() override = default;
 	
+	/**
+	 * TPolyBezierSpline consists of an arbitrary number of anchor points.
+	 * Adjacent anchor points implicitly define a cubic bezier segment.
+	 * The cubic bezier segment defined by points A and B is <A_Anchor, A_OutHandle, B_InHandle, B_Anchor>.
+	 */
+	struct FAnchor
+	{
+		ValueType InHandle;
+		ValueType Anchor;
+		ValueType OutHandle;
+	};
+
 	TPolyBezierSpline(
-		const ValueType& P0, 
-		const ValueType& P1, 
-		const ValueType& P2, 
-		const ValueType& P3,
+		const FAnchor& Start,
+		const FAnchor& End,
 		EParameterizationPolicy Parameterization = EParameterizationPolicy::Uniform)
 	: Base()
 	{
-		// Add the initial segment
-		AddBezierSegmentInternal( 0, P0, P1, P2, P3, Parameterization);
-		TPolyBezierSpline<VALUETYPE>::Reparameterize(Parameterization);
-	}
-
-    /**
-	 * Factory method to create an empty spline with a default segment 
-	 */
-	static TPolyBezierSpline CreateDefault()
-	{
-		// Create a minimal valid default segment for serialization
-		ValueType P0(0, 0, 0);
-		ValueType P1(0.33f, 0, 0);
-		ValueType P2(0.66f, 0, 0);
-		ValueType P3(1, 0, 0);
-		
-		return TPolyBezierSpline(P0, P1, P2, P3, EParameterizationPolicy::Uniform);
+		AppendAnchor(Start, Parameterization);
+		AppendAnchor(End, Parameterization);
 	}
 
 	// Static factory methods for common shapes    
@@ -86,11 +92,15 @@ public:
         ValueType P1 = Start + Direction * 0.33f;
         ValueType P2 = Start + Direction * 0.66f;
         
-        TPolyBezierSpline<ValueType> Result(P0, P1, P2, P3);
-        Result.Reparameterize(EParameterizationPolicy::Uniform);
+		TPolyBezierSpline<ValueType> Result(
+			FAnchor{P0, P0, P1},
+			FAnchor{P2, P3, P3},
+			EParameterizationPolicy::Uniform
+		);
+
         return Result;
     }
-    
+
     // Create a circular arc with given center, radius, and angle range
     static TPolyBezierSpline<ValueType> CreateCircleArc(
         const ValueType& Center,
@@ -99,84 +109,89 @@ public:
         float EndAngle,
         int32 NumSegments = 4)
     {
-           // Ensure angles are in the right order
-           if (EndAngle < StartAngle)
-           {
-               float Temp = StartAngle;
-               StartAngle = EndAngle;
-               EndAngle = Temp;
-           }
-           
-           // Ensure reasonable segment count
-           NumSegments = FMath::Max(1, NumSegments);
+		// Ensure angles are in the right order
+		if (EndAngle < StartAngle)
+		{
+			float Temp = StartAngle;
+			StartAngle = EndAngle;
+			EndAngle = Temp;
+		}
            
-           // Calculate angle per segment
-           float AngleSpan = EndAngle - StartAngle;
-           float AnglePerSegment = AngleSpan / static_cast<float>(NumSegments);
+		// Ensure reasonable segment count
+		NumSegments = FMath::Max(1, NumSegments);
            
-           // For large angles per segment, increase segment count for better approximation
-           if (AnglePerSegment > UE_PI/4)
-           {
-               NumSegments = FMath::CeilToInt(AngleSpan / (UE_PI/4));
-               AnglePerSegment = AngleSpan / static_cast<float>(NumSegments);
-           }
+		// Calculate angle per segment
+		float AngleSpan = EndAngle - StartAngle;
+		float AnglePerSegment = AngleSpan / static_cast<float>(NumSegments);
            
-           // Pre-compute all points along the arc for exact positioning
-           TArray<ValueType> Points;
-           Points.SetNum(NumSegments + 1);
+		// For large angles per segment, increase segment count for better approximation
+		if (AnglePerSegment > UE_PI/4)
+		{
+			NumSegments = FMath::CeilToInt(AngleSpan / (UE_PI/4));
+			AnglePerSegment = AngleSpan / static_cast<float>(NumSegments);
+		}
            
-           for (int32 i = 0; i <= NumSegments; i++)
-           {
-               float Angle = StartAngle + (static_cast<float>(i) * AnglePerSegment);
-               Points[i] = Center + ValueType(
-                   Radius * FMath::Cos(Angle),
-                   Radius * FMath::Sin(Angle),
-                   0);
-           }
+		// Pre-compute all points along the arc for exact positioning
+		TArray<ValueType> Points;
+		Points.SetNum(NumSegments + 1);
            
-           // Calculate first segment with precise tangents
-           ValueType P0 = Points[0];
-           ValueType P3 = Points[1];
+		for (int32 i = 0; i <= NumSegments; i++)
+		{
+			float Angle = StartAngle + (static_cast<float>(i) * AnglePerSegment);
+			Points[i] = Center + ValueType(
+				Radius * FMath::Cos(Angle),
+				Radius * FMath::Sin(Angle),
+				0);
+		}
            
-           // Get normalized direction vectors from center to points
-           ValueType Dir0 = Math::GetSafeNormal(P0 - Center);
-           ValueType Dir3 = Math::GetSafeNormal(P3 - Center);
+		// Calculate first segment with precise tangents
+		ValueType P0 = Points[0];
+		ValueType P3 = Points[1];
            
-           // Calculate perpendicular vectors (normalized)
-           ValueType Perp0 = ValueType(-Dir0.Y, Dir0.X, 0);
-           ValueType Perp3 = ValueType(-Dir3.Y, Dir3.X, 0);
+		// Get normalized direction vectors from center to points
+		ValueType Dir0 = Math::GetSafeNormal(P0 - Center);
+		ValueType Dir3 = Math::GetSafeNormal(P3 - Center);
            
-           // Calculate precise tangent scale
-           float TangentScale = Radius * (4.0f / 3.0f) * 
-                                    FMath::Tan(AnglePerSegment / 4.0f);
+		// Calculate perpendicular vectors (normalized)
+		ValueType Perp0 = ValueType(-Dir0.Y, Dir0.X, 0);
+		ValueType Perp3 = ValueType(-Dir3.Y, Dir3.X, 0);
            
-           // Create control points with exact perpendicular tangents
-           ValueType P1 = P0 + Perp0 * TangentScale;
-           ValueType P2 = P3 - Perp3 * TangentScale;
+		// Calculate precise tangent scale
+		float TangentScale = Radius * (4.0f / 3.0f) * 
+									FMath::Tan(AnglePerSegment / 4.0f);
            
-           // Create spline with first segment
-           TPolyBezierSpline<ValueType> Result(P0, P1, P2, P3);
+		// Create control points with exact perpendicular tangents
+		ValueType P1 = P0 + Perp0 * TangentScale;
+		ValueType P2 = P3 - Perp3 * TangentScale;
+
+		TPolyBezierSpline<ValueType> Result(
+			FAnchor{P0, P0, P1},
+			FAnchor{P2, P3, P3},
+			EParameterizationPolicy::Uniform
+		);
            
-           // Add remaining segments with carefully calculated control points
-           for (int32 i = 1; i < NumSegments; i++)
-           {
-               P0 = Points[i];
-               P3 = Points[i+1];
+		// Add remaining segments with carefully calculated control points
+		for (int32 i = 1; i < NumSegments; i++)
+		{
+			P0 = Points[i];
+			P3 = Points[i+1];
                
-               // Calculate exact perpendicular vectors for this segment
-               Dir0 = Math::GetSafeNormal(P0 - Center);
-               Dir3 = Math::GetSafeNormal(P3 - Center);
-               Perp0 = ValueType(-Dir0.Y, Dir0.X, 0);
-               Perp3 = ValueType(-Dir3.Y, Dir3.X, 0);
+			// Calculate exact perpendicular vectors for this segment
+			Dir0 = Math::GetSafeNormal(P0 - Center);
+			Dir3 = Math::GetSafeNormal(P3 - Center);
+			Perp0 = ValueType(-Dir0.Y, Dir0.X, 0);
+			Perp3 = ValueType(-Dir3.Y, Dir3.X, 0);
                
-               // Create control points with exact perpendicular tangents
-               P1 = P0 + Perp0 * TangentScale;
-               P2 = P3 - Perp3 * TangentScale;
+			// Create control points with exact perpendicular tangents
+			P1 = P0 + Perp0 * TangentScale;
+			P2 = P3 - Perp3 * TangentScale;
                
-               Result.AppendBezierSegment(P1, P2, P3);
-           }
+			const int32 NewAnchorIdx = Result.AppendAnchor(FAnchor { P2, P3, P3 }, EParameterizationPolicy::Uniform);
+			const int32 OldAnchorIdx = NewAnchorIdx - 1;
+			Result.SetHandleOutOfAnchor(OldAnchorIdx, P1);
+		}
            
-           return Result;
+		return Result;
     }
     
     // Create a complete circle
@@ -235,7 +250,11 @@ public:
 	    ValueType P2 = P3 - Tangent3;
 	    
 	    // Create the spline with the first segment
-	    TPolyBezierSpline<ValueType> Result(P0, P1, P2, P3);
+		TPolyBezierSpline<ValueType> Result(
+			FAnchor{P0, P0, P1},
+			FAnchor{P2, P3, P3},
+			EParameterizationPolicy::Uniform
+		);
 	    
 	    // Add remaining segments
 	    for (int32 i = 1; i < NumSegments; ++i)
@@ -257,7 +276,9 @@ public:
 	        P2 = P3 - Tangent3;
 	        
 	        // Add the segment
-	        Result.AppendBezierSegment(P1, P2, P3);
+			const int32 NewAnchorIdx = Result.AppendAnchor(P2, P3, P3);
+			const int32 OldAnchorIdx = NewAnchorIdx - 1;
+			Result.SetHandleOutOfAnchor(OldAnchorIdx, P1);
 	    }
 	    
 	    // Set as closed loop
@@ -298,13 +319,13 @@ public:
         TUniquePtr<TPolyBezierSpline> Clone = MakeUnique<TPolyBezierSpline>();
         
         // Copy base class members
-        Clone->Values = this->Values;
-        Clone->PairKnots = this->PairKnots;
-        Clone->bIsClosedLoop = this->bIsClosedLoop;
+        Clone->Values = Base::Values;
+        Clone->PairKnots = Base::PairKnots;
+        Clone->bIsClosedLoop = Base::bIsClosedLoop;
         
         // Copy infinity modes
-        Clone->PreInfinityMode = this->PreInfinityMode;
-        Clone->PostInfinityMode = this->PostInfinityMode;
+        Clone->PreInfinityMode = Base::PreInfinityMode;
+        Clone->PostInfinityMode = Base::PostInfinityMode;
         
         return Clone;
     }
@@ -317,7 +338,6 @@ public:
     		return 0.0f;
     	}
 
-
 		// Get Bezier control points for this segment
 		TStaticArray<ValueType, 4> Coeffs = {};
 
@@ -467,9 +487,10 @@ public:
 		FWindow Window;
 		FWindowStore Scratch;
 		FindWindow(Parameter, Window, Scratch);
+
 		// If FindWindow fails, it will return an array of nullptr. InterpolateWindow assumes validity of elements.
-		if (!Math::FSplineValidation::IsValidWindow<ValueType, Base::WindowSize>(Window) || // Cubic Bezier needs 4 points
-				!Math::FSplineValidation::IsValidParameter(Parameter))
+		if (!Math::FSplineValidation::IsValidWindow<ValueType, Base::WindowSize>(Window) ||
+			!Math::FSplineValidation::IsValidParameter(Parameter))
 		{
 			return Math::FSplineValidation::GetDefaultValue<ValueType>();
 		}
@@ -481,9 +502,10 @@ public:
 	    }
     
     	// Convert global parameter to segment information
-    	int32 SegmentIndex;
-    	float LocalT;
-		if (!MapGlobalParameterToLocalSegment(Parameter, SegmentIndex, LocalT))
+		float LocalT;
+		int32 SegmentIndex = FindIndexForParameter(Parameter, LocalT);
+
+		if (SegmentIndex == INDEX_NONE)
 		{
 			return Math::FSplineValidation::GetDefaultValue<ValueType>();
 		}
@@ -498,692 +520,479 @@ public:
 	/**
 	 * MANIPULATION METHODS
 	 */
+	
+	/** Returns the number of Bezier segments in the spline */
+	virtual int32 GetNumberOfSegments() const override
+	{
+		const int32 NumAnchors = GetNumAnchorPoints();
+		return HasClosingSegment() ? NumAnchors : FMath::Max(0, NumAnchors - 1);
+	}
 
 	/**
-	 * Replaces all control points in the spline with the provided points.
-	 * For a cubic Bezier spline, the number of points must be a multiple of 4 (or 4N-1 for open splines).
-	 * 
-	 * @param Points - Array of new control points
-	 * @param ParameterizationPolicy - How to distribute Knots for the new segments
-	 * @return True if the points were successfully set
+	 * Maps a segment index to its parameter range
+	 * @param SegmentIndex - Index of the segment (0-based)
+	 * @return True if the segment index is valid and mapping succeeded
 	 */
-	bool SetControlPoints(
-		const TArray<ValueType>& Points,
-		EParameterizationPolicy ParameterizationPolicy)
+	virtual FInterval1f GetSegmentParameterRange(int32 SegmentIndex) const override
 	{
-	   const int32 NumPoints = Points.Num();
-
-	    // We expect Bezier quartets as input: P0,P1,P2,P3 per segment.
-	    if (NumPoints < 4 || (NumPoints % 4) != 0)
-	    {
-	        UE_LOG(LogSpline, Warning,
-	            TEXT("SetControlPoints expects 4*k Bezier control points (got %d)."), NumPoints);
-	        return false;
-	    }
-
-	    // Remember desired closed state, then reset storage.
-	    const bool bWasClosed = Base::IsClosedLoop();
-	    Clear();
+		FInterval1f SegmentRange;
 
-	    const int32 NumSegments = NumPoints / 4;
-		
-	    // First segment bootstraps two anchors
-	    {
-	        const ValueType& P0 = Points[0];
-	        const ValueType& P1 = Points[1];
-	        const ValueType& P2 = Points[2];
-	        const ValueType& P3 = Points[3];
-
-	        // A0
-	        AppendAnchor(/*In0*/ P0, /*A0*/ P0, /*Out0*/ P1);
-	        // A1
-	        AppendAnchor(/*In1*/ P2, /*A1*/ P3, /*Out1*/ P3);
-	    }
+		const int32 NumSegments = GetNumberOfSegments();
+		if (SegmentIndex < 0 || SegmentIndex >= NumSegments || Base::PairKnots.Num() == 0 || !Base::PairKnots.IsValidIndex(SegmentIndex))
+		{
+			return SegmentRange;
+		}
 
-	    // Remaining segments
-	    for (int32 s = 1; s < NumSegments; ++s)
-	    {
-	        const int32 b = 4 * s;
-	        const ValueType& P0 = Points[b + 0];
-	        const ValueType& P1 = Points[b + 1];
-	        const ValueType& P2 = Points[b + 2];
-	        const ValueType& P3 = Points[b + 3];
-
-	        // Optionally warn if P0 doesn't match previous A (only a sanity message).
-	        const ValueType& PrevA = this->Values[AnchorBase(s) + 1];
-	        if (Math::SizeSquared(P0 - PrevA) > UE_KINDA_SMALL_NUMBER)
-	        {
-	            UE_LOG(LogSpline, Verbose,
-	                TEXT("SetControlPoints: segment %d P0 differs from previous A (len^2=%g)."),
-	                s, (double)Math::SizeSquared(P0 - PrevA));
-	        }
-
-	        // Out_s = P1
-	        this->SetValue(AnchorBase(s) + 2, P1);
-
-	        // Append A_{s+1} = P3 with In_{s+1}=P2, Out_{s+1}=A_{s+1} (placeholder until another seg arrives)
-	        AppendAnchor(/*In*/ P2, /*A*/ P3, /*Out*/ P3);
-	    }
+		SegmentRange.Min = Base::PairKnots[SegmentIndex].Value;
 
-	    // Apply desired loop state:
-	    // - We DO NOT materialize a physical closing quartet anymore.
-	    // - Closing segment is virtual: {A_last, Out_last, In_0, A_0}.
-	    SetClosedLoop(bWasClosed);
+		if (!HasClosingSegment())
+		{
+			// Open: knots are standard boundaries, so we still use i+1.
+			if (SegmentIndex + 1 < Base::PairKnots.Num())
+			{
+				SegmentRange.Max = Base::PairKnots[SegmentIndex + 1].Value;
+			}
+			return SegmentRange;
+		}
 
-	    // Build / rebuild knot vector once from scratch using the policy
-	    // (You already do this today; keep it — it re-spaces and sets multiplicities.)
-	    Reparameterize(ParameterizationPolicy);
+		// Closed: PairKnots holds starts of each segment; the final segment
+		// closes the loop at Min + LoopKnotDelta.
+		if (SegmentIndex < NumSegments - 1)
+		{
+			SegmentRange.Max = Base::PairKnots[SegmentIndex + 1].Value;
+		}
+		else
+		{
+			SegmentRange.Max = SegmentRange.Min + LoopKnotDelta;
+		}
 
-	    return true;
+		return SegmentRange;
 	}
 
+	// Parameterization methods
 	/**
-	 * Adds multiple Bezier segments to the spline.
-	 * 
-	 * @param Storage - Storage to add the segments to
-	 * @param Points - Array of control points defining the new segments
-	 * @param bAppend - If true, append to the end of spline; if false, prepend to the beginning
-	 * @param ParameterizationPolicy - How to distribute Knots for the new segments
-	 * @return True if the segments were successfully added
+	 * Reparameterizes the spline based on the provided points and mode
+	 * @param Mode - New parameterization mode
+	 * @param Points - Array of control points to use for reparameterization
+	 * @return 
 	 */
-	bool AddBezierSegments(
-		const TArray<ValueType>& Points,
-		bool bAppend,
-		EParameterizationPolicy ParameterizationPolicy)
+	virtual void Reparameterize(EParameterizationPolicy Mode = EParameterizationPolicy::Centripetal) override
 	{
-	    const int32 NumPoints = Points.Num();
-	    
-		// Need at least 2 points to add a valid segment
-		if (NumPoints < 4)
-		{
-			UE_LOG(LogSpline, Warning, TEXT("AddBezierSegments requires at least 4 points to add a valid segment. Got %d points."), NumPoints);
-			return false;
-		}
-	    
-	    // If the spline is currently empty (even though it shouldn't be according to design), this is equivalent to SetControlPoints
-	    if (this->Base::NumKeys() == 0)
-	    {
-	        return SetControlPoints(Points, ParameterizationPolicy);
-	    }
-	    
-	    // Check points format based on append/prepend direction
-	    if (bAppend)
-	    {
-	        // For appending, we expect 3N+1 points: P0 (shared with last existing point), then (P1,P2,P3) per segment
-	        if (NumPoints % 3 != 1)
-	        {
-	            UE_LOG(LogSpline, Warning, TEXT("Invalid point count for appending segments. Expected 3N+1, got %d."), NumPoints);
-	            return false;
-	        }
-	    }
-	    else
-	    {
-	        // For prepending, we expect 3N+1 points: (P0,P1,P2) per segment, then P3 (shared with first existing point)
-	        if (NumPoints % 3 != 1)
-	        {
-	            UE_LOG(LogSpline, Warning, TEXT("Invalid point count for prepending segments. Expected 3N+1, got %d."), NumPoints);
-	            return false;
-	        }
-	    }
-		const int32 NumAnchors = GetNumAnchorPoints();
+		const TArray<ValueType>& Points = Base::Values;
+		const int32 NumValues = Points.Num();
 
-	    // Check connection point
-	    if (bAppend)
-	    {
-	        // First point of new segments should match the last existing anchor
-	        const ValueType& FirstNewPoint = Points[0];
-	    	const ValueType& LastExistingAnchor = this->Values[AnchorBase(NumAnchors - 1) + 1];
-	        
-	        // Check if the connection point is reasonably close
-	        const double DistanceSquared = Math::SizeSquared(LastExistingAnchor, FirstNewPoint);
-	        if (DistanceSquared > UE_KINDA_SMALL_NUMBER)
-	        {
-	            UE_LOG(LogSpline, Warning, TEXT("Connection point mismatch when appending segments. Distance: %f"),
-	                FMath::Sqrt(static_cast<float>(DistanceSquared)));
-	            // Continue anyway but log warning
-	        }
-	    }
-	    else
-	    {
-	        // Last point of new segments should match the first point of existing spline
-	        const ValueType& LastNewPoint = Points.Last();
-	    	const ValueType& FirstExistingAnchor = this->Values[AnchorBase(0) + 1];
-	        
-	        // Check if the connection point is reasonably close
-	        const double DistanceSquared = Math::SizeSquared(LastNewPoint, FirstExistingAnchor);
-	        if (DistanceSquared > UE_KINDA_SMALL_NUMBER)
-	        {
-	            UE_LOG(LogSpline, Warning, TEXT("Connection point mismatch when prepending segments. Distance: %f"),
-	                FMath::Sqrt(static_cast<float>(DistanceSquared)));
-	            // Continue anyway but log warning
-	        }
-	    }
-	    
-	    if (bAppend)
-	    {
-	        // Append segments to the end of the spline
-	        const int32 NumSegmentsToAdd = (NumPoints - 1) / 3;
-	        
-	        // Add each segment
-	        for (int32 SegmentIndex = 0; SegmentIndex < NumSegmentsToAdd; ++SegmentIndex)
-	        {
-	            const int32 BaseIndex = 1 + SegmentIndex * 3; // Skip the first point (P0) for first segment
-	            
-	            // For each segment, we need P1, P2, P3 (P0 is the last point of existing spline)
-	            const ValueType& P1 = Points[BaseIndex];
-	            const ValueType& P2 = Points[BaseIndex + 1];
-	            const ValueType& P3 = Points[BaseIndex + 2];
-	            
-	            // Use the existing AppendBezierSegment method
-	            this->AppendBezierSegment(P1, P2, P3, ParameterizationPolicy);
-	        }
-	    }
-	    else
-	    {
-	        // Prepend segments to the beginning of the spline
-	        const int32 NumSegmentsToPrepend = (NumPoints - 1) / 3;
-	        
-	        // Prepend segments in order starting from the first segment
-	        // For prepending, we need to work backwards since each PrependBezierSegment
-	        // adds to the beginning
-	        for (int32 i = NumSegmentsToPrepend - 1; i >= 0; --i)
-	        {
-	            const int32 BaseIndex = i * 3;
-	            
-	            // For each segment, we need P0, P1, P2 (P3 will connect to the existing spline)
-	            const ValueType& P0 = Points[BaseIndex];
-	            const ValueType& P1 = Points[BaseIndex + 1];
-	            const ValueType& P2 = Points[BaseIndex + 2];
-	            
-	            // Use the existing PrependBezierSegment method
-	            this->PrependBezierSegment(P0, P1, P2, ParameterizationPolicy);
-	        }
-	    }
-	    
-	    // Reparameterize the spline to ensure consistent parameterization
-    	Reparameterize(ParameterizationPolicy);
-	    
-	    return true;
-	}
-    /**
-     * Appends a Bezier segment using only 3 points (P1, P2, P3)
-     * The start point (P0) is automatically inferred from the last point of the previous segment.
-     * 
-     * @param P1 First control point
-     * @param P2 Second control point
-     * @param P3 End point
-     * @param ParameterizationPolicy How to distribute Knots
-     * @return Index of the new segment
-     */
-    int32 AppendBezierSegment(
-        const ValueType& P1, 
-        const ValueType& P2, 
-        const ValueType& P3, 
-        EParameterizationPolicy ParameterizationPolicy = EParameterizationPolicy::Centripetal)
-    {
-		const int32 NumAnchors = GetNumAnchorPoints();
-		const int32 NumSegments = GetNumberOfSegments();
+		GenerateDistanceKnotsForBezier(Mode);
+    	
+		UE_LOG(LogSpline, Verbose, TEXT("Set knot vector (Reparameterize) - Points: %d, Knots: %d, Mode: %d"),
+					   NumValues,  Base::PairKnots.Num(), static_cast<int32>(Mode));
 		
-		// P0 = last anchor position A_{n-1}
-		const int32 LastAnchor = NumAnchors - 1;
-		const ValueType& P0 = this->Values[AnchorBase(LastAnchor) + 1];
+		Base::PrintKnotVector();
+	}
+	
+	virtual FInterval1f GetParameterSpace() const override
+	{
+		if (Base::PairKnots.Num() == 0)
+		{
+			return FInterval1f::Empty();
+		}
 
-		// Insert new segment after the current last one
-		//   SegmentIndex == NumSegments means "append at tail"
-		return AddBezierSegmentInternal(NumSegments, P0, P1, P2, P3, ParameterizationPolicy, true);
-    }
+		const float Min = Base::PairKnots[0].Value;
+		const float LastKnot = Base::PairKnots.Last().Value;
 
-    /**
-     * Prepends a Bezier segment using only 3 points (P0, P1, P2)
-     * The end point (P3) is automatically inferred from the first point of the existing spline.
-     * 
-     * @param P0 First control point
-     * @param P1 Second control point
-     * @param P2 Third control point
-     * @param ParameterizationPolicy How to distribute Knots
-     * @return Index of the new segment
-     */
-    int32 PrependBezierSegment(
-		const ValueType& P0, 
-		const ValueType& P1, 
-		const ValueType& P2, 
-		EParameterizationPolicy ParameterizationPolicy = EParameterizationPolicy::Centripetal)
-	{
-		const int32 NumAnchors = GetNumAnchorPoints();
-		const int32 NumSegments = GetNumberOfSegments();
-		
-		// P3 = first anchor position A_0
-		const ValueType& P3 = this->Values[AnchorBase(0) + 1];
+		if (!HasClosingSegment())
+		{
+			return FInterval1f(Min, LastKnot);
+		}
 
-		// Insert new segment before segment 0 / at the head
-		return AddBezierSegmentInternal(0, P0, P1, P2, P3,
-										ParameterizationPolicy,
-										false); // Prepend to first segment
+		// Closed: closing boundary is last real knot + LoopKnotDelta
+		return FInterval1f(Min, LastKnot + LoopKnotDelta);
 	}
 
 	/**
-	 * Inserts a point at the closest location on a segment to the given position.
-	 * 
-	 * @param SegmentIndex Index of the segment to split
-	 * @param Position Position for the new point
-	 * @param ParameterizationPolicy How to distribute Knots
-	 * @return The index of the newly created point
+	 * Gets the parameter value for a specific anchor point or the virtual closing point.
+	 * The virtual closing point exists when there are at least 2 anchor points and the spline is closed. The index is GetNumAnchorPoints().
+	 * @param Index - Index of the anchor point
+	 * @return Parameter value at the specified point
 	 */
-	int32 InsertPointAtPosition(
-		int32 SegmentIndex,
-		const ValueType& Position,
-		EParameterizationPolicy ParameterizationPolicy = EParameterizationPolicy::Centripetal)
+	virtual float GetParameter(int32 Index) const override
 	{
-		// Find the closest point on the segment to the given position
-		float SquaredDistance;
-		float LocalT = FindLocalParameterNearestToPosition(SegmentIndex, Position, SquaredDistance);
-            
-		// Insert at that local parameter
-		return this->InsertPointAtSegmentParam(SegmentIndex, LocalT, Position, ParameterizationPolicy);
-	}
+		const int32 NumAnchors = GetNumAnchorPoints();
 
-	/**
-	 * Inserts a point at a specific segment with local parameter.
-	 * This splits the segment at the specified location and places the new point
-	 * at the provided position, which may change the spline shape.
-	 * 
-	 * @param SegmentIndex Index of the segment to split
-	 * @param LocalT Local parameter [0,1] within the segment
-	 * @param Position Position for the new point
-	 * @param ParameterizationPolicy How to distribute Knots
-	 * @return The index of the newly created point
-	 */
-	int32 InsertPointAtSegmentParam(
-	    int32 SegmentIndex,
-	    float LocalT,
-	    const ValueType& Position,
-	    EParameterizationPolicy ParameterizationPolicy = EParameterizationPolicy::Centripetal)
-	{
-    	int32 NumSegments = GetNumberOfSegments();
-	    // Validate inputs
-	    if (NumSegments == 0)
-	    {
-	        // Special case for empty spline
-	        int32 NewPointIndex;
-	        if (InitializeFirstSegment(Position, NewPointIndex))
-	            return NewPointIndex;
-	    }
-	    
-	    SegmentIndex = FMath::Clamp(SegmentIndex, 0, NumSegments - 1);
-	    
-	    // Get the segment's parameter range
-	    FInterval1f SegmentRange = GetSegmentParameterRange(SegmentIndex);
-	    
-	    // Convert local parameter to global parameter
-	    float Parameter = SegmentRange.Min + LocalT * (SegmentRange.Max - SegmentRange.Min);
-	    
-	    return InsertPoint(Parameter, Position, ParameterizationPolicy);
-	}
-		
-	/**
-	 * Inserts a point at the specified parameter along the spline.
-	 * This splits the segment at the parameter value and places the new point
-	 * at the specified position, which may change the spline shape.
-	 * 
-	 * @param Storage - Storage to add the point to
-	 * @param Parameter Parameter value where to split the spline
-	 * @param Position Position for the new point (doesn't need to be on the original curve)
-	 * @param ParameterizationPolicy How to distribute Knots
-	 * @return The index of the newly created point
-	 */
-    int32 InsertPoint(
-        float Parameter, 
-        const ValueType& Position,
-        EParameterizationPolicy ParameterizationPolicy = EParameterizationPolicy::Centripetal)
-    {
-	    if (int32 NewPointIndex;
-	    	InitializeFirstSegment(Position, NewPointIndex))
+		if (NumAnchors <= 0)
 		{
-			return NewPointIndex;
+			// todo: fail
+			return 0.f;
 		}
 
-		typename Base::FKnotSearchCache SearchCache{ this->PairKnots };
-
-		// Make sure we have a valid parameter
-		Parameter = Base::GetNearestAvailableKnotValue(Parameter, SearchCache);
-		
-        FInterval1f ParameterRange = GetParameterSpace();
-		
-		// Handle out-of-range Knots
-		if (Parameter <= ParameterRange.Min)
-		{
-			// Prepend
-			return PrependPoint(Position, ParameterizationPolicy);
-		}
-    
-		if (Parameter >= ParameterRange.Max)
+		// Getting the virtual closing knot.
+		if (Index == NumAnchors && HasClosingSegment())
 		{
-			return AppendPoint(Position, ParameterizationPolicy);
+			return Base::PairKnots.Last().Value + LoopKnotDelta;
 		}
 
-		int32 SegmentIndex;
-		float LocalT;
-		if (!MapGlobalParameterToLocalSegment(Parameter, SegmentIndex, LocalT))
+		if (Index < 0 || Index >= NumAnchors)
 		{
-			return INDEX_NONE;
+			// todo: fail
+			return 0.f;
 		}
 
-		FWindow Window;
-		if (!BuildRawWindowForSegment(SegmentIndex, Window))
+		return Base::PairKnots[Index].Value;
+	}
+
+	/**
+	 * Sets the parameter value for a specific anchor control point or the virtual closing point.
+	 * This function will reorder the knot vector to ensure that it is always monotonically increasing.
+	 * If the knot vector is reordered, anchor points are reordered such that they are always associated with the same knot.
+	 * Setting the parameter of the virtual closing point will not cause any reordering. If reordering would occur, the function fails.
+	 * @param Index - Index of the anchor point
+	 * @param NewParameter - New parameter value to set
+	 * @return The index of the point that was actually updated, or INDEX_NONE if failed
+	 */ 
+	virtual int32 SetParameter(int32 Index, float NewParameter) override
+	{
+		const int32 NumAnchors = GetNumAnchorPoints();
+		const int32 NumSegments = GetNumberOfSegments();
+
+		// Cannot set parameters for empty splines, fail.
+		if (NumAnchors <= 0)
 		{
 			return INDEX_NONE;
 		}
-		ValueType P0 = *Window[0];
-		ValueType P1 = *Window[1];
-		ValueType P2 = *Window[2];
-		ValueType P3 = *Window[3];
-		
-        // Use De Casteljau algorithm to split the curve
-        const float t = LocalT;
-        const float mt = 1.0f - t;
-        
-        // Calculate split control points
-        ValueType Q0 = P0;
-        ValueType Q1 = P0 * mt + P1 * t;
-        ValueType Q2 = (P0 * mt + P1 * t) * mt + (P1 * mt + P2 * t) * t;
-        ValueType Q3 = Position;  // Use provided position for the split point
-        
-        ValueType R0 = Position;  // Use provided position for the split point
-        ValueType R1 = (P1 * mt + P2 * t) * mt + (P2 * mt + P3 * t) * t;
-        ValueType R2 = P2 * mt + P3 * t;
-        ValueType R3 = P3;
-        
-        FInterval1f SegmentRange = GetSegmentParameterRange(SegmentIndex);
-		float TargetSplitParam = SegmentRange.Interpolate(LocalT);
 
-		if (TargetSplitParam == SegmentRange.Min)
+		// Trivial case, no possible reordering or numerical issues.
+		if (NumAnchors == 1)
 		{
-			TargetSplitParam = Param::NextDistinct(TargetSplitParam);
+			if (Index == 0)
+			{
+				Base::PairKnots[Index].Value = NewParameter;
+				return 0;
+			}
+			else
+			{
+				return INDEX_NONE;
+			}
 		}
 
-		if (TargetSplitParam == SegmentRange.Max)
-		{
-			TargetSplitParam = Param::PrevDistinct(TargetSplitParam);
-		}
-		
-		// Bail out if we lost usable range
-		if (TargetSplitParam <= SegmentRange.Min || TargetSplitParam >= SegmentRange.Max)
+		// Setting the virtual closing knot.
+		if (Index == NumAnchors && HasClosingSegment())
 		{
-			return INDEX_NONE;
+			// No reordering when manipulating the closing knot, fail.
+			if (NewParameter <= Base::PairKnots.Last().Value)
+			{
+				return INDEX_NONE;
+			}
+
+			LoopKnotDelta = NewParameter - Base::PairKnots.Last().Value;
+			return Index;
 		}
 
-		// --- Anchor/handle surgery in 3-tuple layout ---
-		const int32 NumAnchors = GetNumAnchorPoints();
-		if (NumAnchors < 2)
+		// Because we special case the closing point above, we can trivially bounds check here.
+		if (Index < 0 || Index >= NumAnchors)
 		{
 			return INDEX_NONE;
 		}
 
-		// For segment s in an *open* spline, its anchors are A_s and A_{s+1}.
-		// For closed, the closing segment (s == A-1) is between A_{A-1} and A_0.
-		const bool bClosed = IsActuallyClosed();
-		const int32 LeftAnchor = SegmentIndex;
-		const int32 RightAnchor =
-			(!bClosed)
-			? (SegmentIndex + 1)
-			: ((SegmentIndex + 1) % NumAnchors);
-
-		// 1) Left side: Out_left = Q1
-		SetHandleOutOfAnchor(LeftAnchor, Q1);
-
-		// 2) Insert new anchor between left and right with (In_new, A_new, Out_new) = (Q2, Q3, R1)
-		//    For open, this is simply at index RightAnchor.
-		//    For closed, inserting at index RightAnchor still produces a valid cyclic order;
-		//    the notion of "start anchor" is arbitrary on a loop.
-		InsertAnchorAt(RightAnchor, /*In*/ Q2, /*A*/ Q3, /*Out*/ R1);
-
-		// 3) Right side: old right anchor has shifted by +1 after the insert
-		const int32 OldRightAnchorIndex =
-			(!bClosed)
-			? (RightAnchor + 1)
-			: ((RightAnchor + 1) % (NumAnchors + 1)); // one more anchor now in closed loop
+		const int32 SegmentIndex = Index == 0 ? 0 : Index - 1;
+		const int32 LocalPointIndex = Index == 0 ? 0 : 3;
 
-		SetHandleInOfAnchor(OldRightAnchorIndex, R2);
-
-		// --- Insert a knot at the split location (incremental) ---
-
-		this->InsertKnot(FKnot(TargetSplitParam, Degree));
+		float OldParameter = Base::PairKnots[Index].Value;
+    
+		// Early exit if parameter isn't changing
+		if (OldParameter == NewParameter)
+		{
+			return Index;
+		}
+		
+		UE_LOG(LogSpline, Verbose, TEXT("\t"));
+		UE_LOG(LogSpline, Verbose, TEXT("Before SetParameter(%d, %f):"), Index, NewParameter);
+		Base::Dump();
+		
+		// Prevent reuse of an existing knot.
+		const float DesiredNewParameter = NewParameter;
+		FKnotSearchCache SearchCache{ Base::PairKnots };
+		const float ActualNewParameter = Base::GetNearestAvailableKnotValue(DesiredNewParameter, SearchCache);
 
-		this->ApplyClampedKnotsMultiplicity();
-		ApplyInternalKnotsMultiplicity();
+		const float DesiredToOldDistance = FMath::Abs(DesiredNewParameter - OldParameter);
+		const float DesiredToActualNewDistance = FMath::Abs(DesiredNewParameter - ActualNewParameter);
 
-		if (bClosed)
+		if (DesiredToActualNewDistance > DesiredToOldDistance)
 		{
-			RefreshClosingKnot(ParameterizationPolicy); // optional helper to ensure closing boundary is consistent
+			// If our nearest valid knot is further away from the desired value than our current value is, we need to just keep our current value
+			return Index;
 		}
+		
+		int32 ReturnIndex = Index;
+		Base::SetKnot(Index, NewParameter);
 
-		// The new anchor is at index RightAnchor
-		return RightAnchor;
-    }
+		// Fix the knot vector by flipping segments which have inverted parameter ranges until it is valid again.
+		if (NewParameter < OldParameter)
+		{
+			int32 LeftSegmentIndex = Index - 1;
 
-	/**
-	* Removes a Bezier segment from the spline
-	*
-	* @param Storage - Storage to remove the segment from
-	* @param SegmentIndex - Segment index to remove
-	* @return true if successfully removed
-	*/
-    bool RemoveSegment(const int32 SegmentIndex)
-    {
-    	const int32 NumSegments = GetNumberOfSegments();
-	    if (SegmentIndex < 0 || SegmentIndex >= NumSegments)
-	    {
-	        return false;
-	    }
-
-    	// Calculate storage indices
-    	const bool bIsLastSegment = (SegmentIndex == NumSegments - 1);
-    
-    	UE_LOG(LogSpline, Verbose, TEXT("Removing segment at index %d"), SegmentIndex);
-    
-        if (bIsLastSegment)
-	    {
-	    	return RemovePoint(GetNumberOfSegments());
-	    }
-	    else
-	    {
-	    	// in all other cases, removing the segment would be the same as removing the point at that index
-	        return RemovePoint(SegmentIndex);
-	    }
-    }
+			// Repeated left flips
+			while (LeftSegmentIndex >= 0 && GetSegmentParameterRange(LeftSegmentIndex).Min > GetSegmentParameterRange(LeftSegmentIndex).Max)
+			{
+				FlipSegment(LeftSegmentIndex);
+				LeftSegmentIndex--;
+				ReturnIndex--;
+			}
+		}
+		else if (NewParameter > OldParameter)
+		{
+			int32 RightSegmentIndex = Index;
 
-	/**
-	 * Removes a point from the spline.
-	 * @param PointIndex - Index of the point to remove
-	 * @return true if successfully removed
-	 */
-	bool RemovePoint(int32 PointIndex)
+			while (RightSegmentIndex <= NumSegments - 1 && !IsClosingSegment(RightSegmentIndex) && GetSegmentParameterRange(RightSegmentIndex).Min > GetSegmentParameterRange(RightSegmentIndex).Max)
+			{
+				FlipSegment(RightSegmentIndex);
+				RightSegmentIndex++;
+				ReturnIndex++;
+			}
+		}
+			
+		UE_LOG(LogSpline, Verbose, TEXT("\t"));
+		UE_LOG(LogSpline, Verbose, TEXT("After SetParameter(%d, %f):"), Index, NewParameter);
+		Base::Dump();
+		
+		return ReturnIndex;
+	}
+	
+	void FlipSegment(int32 Segment)
 	{
+		const int32 NumSegments = GetNumberOfSegments();
 		const int32 NumAnchors  = GetNumAnchorPoints();
-	    const int32 NumSegments = GetNumberOfSegments();
+		if (Segment < 0 || Segment >= NumSegments || NumAnchors < 2)
+		{
+			return;
+		}
+
+		UE_LOG(LogSpline, Verbose, TEXT("\t"));
+		UE_LOG(LogSpline, Verbose, TEXT("Before FlipSegment(%d):"), Segment);
+		Base::Dump();
+		
+		// 1) Flip this segment's Bezier window: P0<->P3, P1<->P2
+		{
+			int32 Idx0, Idx1, Idx2, Idx3;
+			GetBezierIndices(Segment, Idx0, Idx1, Idx2, Idx3);
+
+			ValueType P0 = Base::GetValue(Idx0);
+			ValueType P1 = Base::GetValue(Idx1);
+			ValueType P2 = Base::GetValue(Idx2);
+			ValueType P3 = Base::GetValue(Idx3);
+
+			Base::SetValue(Idx0, P3);
+			Base::SetValue(Idx1, P2);
+			Base::SetValue(Idx2, P1);
+			Base::SetValue(Idx3, P0);
+		}
 
-	    if (NumAnchors == 0)
+		// 2) Fix left neighbor (Segment-1): keep its out-tangent but move its end point to match
+	    if (Segment > 0 || (HasClosingSegment() && Segment == 0))
 	    {
-	        return false;
+	    	const int32 LeftSeg =
+			 (Segment > 0) ? (Segment - 1) : (NumSegments - 1); // wrap in closed case
+
+	    	int32 ThisIdx0, ThisIdx1, ThisIdx2, ThisIdx3;
+	    	int32 LeftIdx0, LeftIdx1, LeftIdx2, LeftIdx3;
+
+	    	GetBezierIndices(Segment, ThisIdx0, ThisIdx1, ThisIdx2, ThisIdx3);
+	    	GetBezierIndices(LeftSeg, LeftIdx0, LeftIdx1, LeftIdx2, LeftIdx3);
+
+	    	const ValueType ThisP0 = Base::GetValue(ThisIdx0);
+	    	const ValueType LeftP2 = Base::GetValue(LeftIdx2);
+	    	const ValueType LeftP3 = Base::GetValue(LeftIdx3);
+
+	    	const ValueType Delta = LeftP3 - LeftP2; // original out-tangent
+	    	const ValueType NewEnd = ThisP0;
+	    	const ValueType NewIn  = NewEnd - Delta;
+
+	    	Base::SetValue(LeftIdx3, NewEnd);
+	    	Base::SetValue(LeftIdx2, NewIn);
 	    }
 
-	    // In this model, "points" == anchors
-	    if (PointIndex < 0 || PointIndex >= NumAnchors)
+	    // 3) Fix right neighbor (Segment+1): keep its in-tangent but move its start point to match
+	    if (Segment < NumSegments - 1 || IsClosingSegment(Segment))
 	    {
-	        return false;
+	    	const int32 RightSeg =
+			 (Segment < NumSegments - 1) ? (Segment + 1) : 0; // wrap in closed case
+
+	    	int32 ThisIdx0, ThisIdx1, ThisIdx2, ThisIdx3;
+	    	int32 RightIdx0, RightIdx1, RightIdx2, RightIdx3;
+
+	    	GetBezierIndices(Segment, ThisIdx0, ThisIdx1, ThisIdx2, ThisIdx3);
+	    	GetBezierIndices(RightSeg, RightIdx0, RightIdx1, RightIdx2, RightIdx3);
+
+	    	const ValueType ThisP3  = Base::GetValue(ThisIdx3);
+	    	const ValueType RightP0 = Base::GetValue(RightIdx0);
+	    	const ValueType RightP1 = Base::GetValue(RightIdx1);
+
+	    	const ValueType Delta   = RightP0 - RightP1; // original in-tangent
+	    	const ValueType NewStart = ThisP3;
+	    	const ValueType NewOut   = NewStart - Delta;
+
+	    	Base::SetValue(RightIdx0, NewStart);
+	    	Base::SetValue(RightIdx1, NewOut);
 	    }
 
-		auto RemoveAnchorTriple = [this](int32 AnchorIndex)
+	    // 4) Swap knot intervals for the flipped segment & its neighbor interval
+	    Base::SwapKnots(Segment, Segment + 1);
+
+	    UE_LOG(LogSpline, Verbose, TEXT("\t"));
+	    UE_LOG(LogSpline, Verbose, TEXT("After FlipSegment(%d):"), Segment);
+	    Base::Dump();
+	}
+	
+	virtual int32 FindIndexForParameter(float Parameter, float& OutLocalParam) const override
+	{
+		const int32 NumSegments = GetNumberOfSegments();
+		const UE::Geometry::FInterval1f ParameterSpace = GetParameterSpace();
+
+		if (ParameterSpace.IsEmpty())
 		{
-			const int32 BaseIdx = AnchorBase(AnchorIndex);
-			// Remove in reverse index order so earlier removes don't shift later indices
-			this->RemoveValue(BaseIdx + 2); // Out_i
-			this->RemoveValue(BaseIdx + 1); // A_i
-			this->RemoveValue(BaseIdx + 0); // In_i
-		};
+			return INDEX_NONE;
+		}
 
-		if (NumAnchors == 1)		// Special case: only 1 anchors → removing one leaves empty spline.
+		// Left end: snap onto first segment start
+		if (Parameter <= ParameterSpace.Min)
 		{
-			Base::Values.Reset();
-	        this->PairKnots.Reset();
-	        return true;
+			OutLocalParam = 0.0f;
+			return 0;
 		}
-	    else if (NumAnchors == 2)	// Special case: only 2 anchors → removing one leaves 1 anchor, no segments.
-	    {
-	        RemoveAnchorTriple(PointIndex);
-	        // One anchor left, no segments: clear knots
-	        this->PairKnots.Reset();
-	        return true;
-	    }
 
-	    // ---- OPEN SPLINE CASE ----
-	    if (!IsActuallyClosed())
-	    {
-	        // For open curves:
-	        //  - anchors = NumSegments + 1
-	        //  - boundary knots are indices [0 .. NumSegments]
-	        //  - interior boundaries (between anchors j and j+1) correspond to knot index j+1?
-
-	        // A) remove first point (anchor 0): drop first segment and knot 0
-	        if (PointIndex == 0)
-	        {
-	            RemoveAnchorTriple(0);
-
-	            // Remove start boundary knot
-	            if (this->PairKnots.Num() > 0)
-	            {
-	                this->RemoveKnot(0);
-	            }
-
-	            this->ApplyClampedKnotsMultiplicity();
-	            ApplyInternalKnotsMultiplicity();
-	            return true;
-	        }
-
-	        // B) remove last point (anchor N-1): drop last segment and last boundary knot
-	        if (PointIndex == NumAnchors - 1)
-	        {
-	            RemoveAnchorTriple(NumAnchors - 1);
-
-	            // Remove end boundary knot (index == NumSegments)
-	            if (this->PairKnots.Num() > 0)
-	            {
-	                const int32 EndKnotIndex = FMath::Clamp(NumSegments, 0, this->PairKnots.Num() - 1);
-	                this->RemoveKnot(EndKnotIndex);
-	            }
-
-	            this->ApplyClampedKnotsMultiplicity();
-	            ApplyInternalKnotsMultiplicity();
-	            return true;
-	        }
-
-	        // C) remove internal point (0 < j < N-1): merge segments (j-1)->j and j->(j+1)
-	        //    into one (j-1)->(j+1). Shape-wise we keep Out_{j-1} and In_{j+1};
-	        //    dropping A_j's triple is enough.
-	        RemoveAnchorTriple(PointIndex);
-
-	        // Interior knot index aligned with the removed anchor index in the old scheme:
-	        // - anchor j corresponds to a boundary between segments (j-1) and j
-	        // - that boundary’s knot index is j
-	        if (this->PairKnots.Num() > PointIndex)
-	        {
-	            this->RemoveKnot(PointIndex);
-	        }
-
-	        this->ApplyClampedKnotsMultiplicity();
-	        ApplyInternalKnotsMultiplicity();
-	        return true;
-	    }
+		// Right end: snap onto last segment end
+		if (Parameter >= ParameterSpace.Max)
+		{
+			OutLocalParam = 1.0f;
+			return NumSegments - 1;
+		}
 
-	    // ---- CLOSED SPLINE CASE ----
-	    // For closed curves:
-	    //  - anchors == segments
-	    //  - all boundaries are "interior" (no clamped ends)
-	    //  - removing anchor j merges two segments (j-1)->j and j->(j+1)
-	    //    into one (j-1)->(j+1); we keep Out_{j-1} and In_{j+1} and just drop A_j.
-	    {
-	        const int32 RemoveAnchorIdx = PointIndex % NumAnchors;
-	        RemoveAnchorTriple(RemoveAnchorIdx);
-
-	        // In the old closed implementation, a point removal deleted the
-	        // corresponding boundary knot at the same index.
-	        if (this->PairKnots.Num() > RemoveAnchorIdx)
-	        {
-	            this->RemoveKnot(RemoveAnchorIdx);
-	        }
-
-	        // No clamped ends in closed curve: all remaining knots should be interior-style.
-	        ApplyInternalKnotsMultiplicity();
-	        return true;
-	    }
-	}
+		const float LastKnot = Base::PairKnots.Last().Value;
 
-    /**
-	 * Updates a single control point within a Bezier segment
-	 * @param Storage - Storage to update the point in
-	 * @param SegmentIndex - Which cubic Bezier segment (each has 4 points)
-	 * @param PointIndex - Which point in the segment (0=P0, 1=P1, 2=P2, 3=P3)
-	 * @param NewValue - New position for the control point
-	 * @return true if successfully updated
-	 */
-    bool UpdateSegmentPoint(
-    	const int32 SegmentIndex,
-    	const int32 PointIndex,
-    	const ValueType& NewValue)
-    {
-		const int32 NumSegments = GetNumberOfSegments();
-		if (SegmentIndex < 0 || SegmentIndex >= NumSegments ||
-			PointIndex  < 0 || PointIndex  > 3)
+		// We can now assume Parameter is contained by ParameterSpace, no need for clamping.
+
+		// Closed loop: handle the synthetic closing segment [LastKnot, LastKnot + LoopKnotDelta]
+		if (Parameter >= LastKnot && ensure(HasClosingSegment()))
 		{
-			return false;
+			const float t = (Parameter - LastKnot) / FMath::Max(ParameterSpace.Max - LastKnot, UE_SMALL_NUMBER);
+			OutLocalParam = FMath::Clamp(t, 0.0f, 1.0f);
+			return NumSegments - 1;
 		}
+        
+		int32 SegmentIndex = INDEX_NONE;
 
-		int32 BezierIndices[4];
-		GetBezierIndices(SegmentIndex, BezierIndices[0], BezierIndices[1], BezierIndices[2], BezierIndices[3]);
-		int32 ValueIndex = BezierIndices[PointIndex];
-		this->SetValue(ValueIndex, NewValue);
-		return true;
-    }
+		// todo: make CVar
+		constexpr bool bEnableSegmentCaching = true;
+		if (bEnableSegmentCaching && LastEvaluatedSegment != INDEX_NONE && LastEvaluatedSegment < Base::PairKnots.Num() - 1)
+		{
+			if (Parameter < Base::PairKnots[LastEvaluatedSegment + 1].Value && Parameter >= Base::PairKnots[LastEvaluatedSegment].Value)
+			{
+				// cache hit!
+				SegmentIndex = LastEvaluatedSegment;
+			}
+		}
+		
+		if (SegmentIndex == INDEX_NONE)
+		{
+			// cache miss!
+			int32 SegmentLow = 0;
+			int32 SegmentHigh = Base::PairKnots.Num() - 2;
+			SegmentIndex = SegmentLow + (SegmentHigh - SegmentLow) / 2;
 
-    /**
-     * Updates all control points of a Bezier segment at once
-     * @param SegmentIndex - Which cubic Bezier segment to update
-     * @param P0 - Start point of segment
-     * @param P1 - First control point
-     * @param P2 - Second control point
-     * @param P3 - End point of segment
-     * @return true if successfully updated
-     */
-    bool UpdateSegment(
-    	int32 SegmentIndex,
-    	const ValueType& P0,
-    	const ValueType& P1, 
-        const ValueType& P2,
-        const ValueType& P3)
-    {
-        const int32 NumSegments = GetNumberOfSegments();
-        if (SegmentIndex < 0 || SegmentIndex >= NumSegments)
-        {
-            return false;
-        }
+			while (SegmentLow <= SegmentHigh)
+			{
+				int32 Mid = SegmentLow + (SegmentHigh - SegmentLow) / 2;
+				const float& StartParam = Base::PairKnots[Mid].Value;
+				const float& EndParam   = Base::PairKnots[Mid + 1].Value;
 
-		int32 Idx0, Idx1, Idx2, Idx3;
-		GetBezierIndices(SegmentIndex, Idx0, Idx1, Idx2, Idx3);
+				if (Parameter >= StartParam && Parameter < EndParam)
+				{
+					SegmentIndex = Mid;
+					break;
+				}
+				else if (Parameter < StartParam)
+				{
+					SegmentHigh = Mid - 1;
+				}
+				else
+				{
+					SegmentLow = Mid + 1;
+				}
+			}
+		}
 		
-		this->SetValue(Idx0, P0);
-		this->SetValue(Idx1, P1);
-		this->SetValue(Idx2, P2);
-		this->SetValue(Idx3, P3);
+		// Update the cache.
+		LastEvaluatedSegment = SegmentIndex;
+        
+		// Calculate local parameter
+		float SegmentLength = Base::PairKnots[SegmentIndex + 1].Value - Base::PairKnots[SegmentIndex].Value;
+		if (SegmentLength > UE_SMALL_NUMBER)
+		{
+			OutLocalParam = (Parameter - Base::PairKnots[SegmentIndex].Value) / SegmentLength;
+		}
+		else
+		{
+			OutLocalParam = 0.0f;
+		}
+        
+		return SegmentIndex;
+	}
+	
+	virtual void SetClosedLoop(bool bShouldClose) override
+	{
+		// Skip if state isn't changing
+		if (bShouldClose == Base::IsClosedLoop())
+		{
+			return;
+		}
+        
+		// Just flip the topology flag; do not add/remove any control points.
+		Base::SetClosedLoop(bShouldClose);
 
-		return true;
-    }
+		// Update knot multiplicities.
+		ApplyClampedKnotsMultiplicity();
+		ApplyInternalKnotsMultiplicity();
+	}
+
+	/**
+	 * Maps a local segment parameter [0,1] to global parameter space
+	 * @param SegmentIndex - Index of the segment
+	 * @param LocalParam - Local parameter within segment [0,1]
+	 * @return Parameter value in global space, or 0 if segment is invalid
+	 */
+	float MapLocalSegmentParameterToGlobal(int32 SegmentIndex, float LocalParam) const
+	{
+		FInterval1f SegmentRange = GetSegmentParameterRange(SegmentIndex);
+		return SegmentRange.Interpolate(LocalParam);
+	}
 
+	void GetBezierIndices(int32 SegmentIndex, int32& OutP0Idx, int32& OutP1Idx, int32& OutP2Idx, int32& OutP3Idx) const
+	{
+		if (!IsClosingSegment(SegmentIndex))
+		{
+			OutP0Idx = P0Idx(SegmentIndex);
+			OutP1Idx = P1Idx(SegmentIndex);
+			OutP2Idx = P2Idx(SegmentIndex);
+			OutP3Idx = P3Idx(SegmentIndex);
+		}
+		else
+		{
+			const int32 NumAnchors = GetNumAnchorPoints();
+
+			OutP0Idx = Close_P0Idx(NumAnchors);
+			OutP1Idx = Close_P1Idx(NumAnchors);
+			OutP2Idx = Close_P2Idx();
+			OutP3Idx = Close_P3Idx();
+		}
+	}
+	
 	/**
 	 * Anchor points are unique segment endpoints.
 	 * All segments have 2 corresponding anchor points.
 	 * For a closed spline, the final segment's anchor points are the first and last points.
 	 */
-	int32 GetNumAnchorPoints() const { return Base::Values.Num() / 3; }
+	int32 GetNumAnchorPoints() const
+	{
+		return Base::Values.Num() / 3;
+	}
 
 	void SetAnchorPoint(int32 AnchorIndex, const ValueType& NewPosition)
 	{
@@ -1204,627 +1013,328 @@ public:
 		const int32 AnchorValueIndex = AnchorBase(AnchorIndex) + 1;
 		return Base::GetValue(AnchorValueIndex);
 	}
+
+	/** Per-anchor setters/getters */
+	void SetHandleOutOfAnchor(int32 AnchorIdx, const ValueType& Out)
+	{
+		Base::SetValue(AnchorBase(AnchorIdx) + 2, Out);
+	}
+
+	const ValueType& GetHandleOutOfAnchor(int32 AnchorIdx) const
+	{
+		return Base::GetValue(AnchorBase(AnchorIdx) + 2);
+	}
 	
-	/** Returns the number of Bezier segments in the spline */
-	virtual int32 GetNumberOfSegments() const override
+	void SetHandleInOfAnchor(int32 AnchorIdx, const ValueType& In)
 	{
-		const int32 NumAnchors = GetNumAnchorPoints();
-		return IsActuallyClosed() ? NumAnchors : FMath::Max(0, NumAnchors - 1);
+		Base::SetValue(AnchorBase(AnchorIdx) + 0, In);
 	}
 
-	/**
-	 * Maps a segment index to its parameter range
-	 * @param SegmentIndex - Index of the segment (0-based)
-	 * @return True if the segment index is valid and mapping succeeded
-	 */
-	virtual FInterval1f GetSegmentParameterRange(int32 SegmentIndex) const override
+	const ValueType& GetHandleInOfAnchor(int32 AnchorIdx) const
 	{
-		FInterval1f SegmentRange;
+		return Base::GetValue(AnchorBase(AnchorIdx) + 0);
+	}
 
-		const int32 NumSegments = GetNumberOfSegments();
-		if (SegmentIndex < 0 || SegmentIndex >= NumSegments || Base::PairKnots.Num() == 0 || !Base::PairKnots.IsValidIndex(SegmentIndex))
-		{
-			return SegmentRange;
-		}
+	void SetAnchors(const TArray<FAnchor>& InAnchors, EParameterizationPolicy ParameterizationPolicy = EParameterizationPolicy::Centripetal)
+	{
+		Clear();
 
-		SegmentRange.Min = Base::PairKnots[SegmentIndex].Value;
+		AppendAnchors(InAnchors, ParameterizationPolicy);
+	}
 
-		if (!IsActuallyClosed())
-		{
-			// Open: knots are standard boundaries, so we still use i+1.
-			if (SegmentIndex + 1 < Base::PairKnots.Num())
-			{
-				SegmentRange.Max = Base::PairKnots[SegmentIndex + 1].Value;
-			}
-			return SegmentRange;
-		}
+	void AppendAnchors(const TArray<FAnchor>& InAnchors, EParameterizationPolicy ParameterizationPolicy = EParameterizationPolicy::Centripetal)
+	{
+		const int32 ExpectedNumAnchors = GetNumAnchorPoints() + InAnchors.Num();
 
-		// Closed: PairKnots holds starts of each segment; the final segment
-		// closes the loop at Min + LoopKnotDelta.
-		if (SegmentIndex < NumSegments - 1)
+		Base::Values.Reserve(ExpectedNumAnchors * 3);
+		Base::PairKnots.Reserve(ExpectedNumAnchors * 3);
+
+		for (const FAnchor& Anchor : InAnchors)
 		{
-			SegmentRange.Max = Base::PairKnots[SegmentIndex + 1].Value;
+			AppendAnchor(Anchor, ParameterizationPolicy);
 		}
-		else
+	}
+
+	int32 PrependAnchor(const FAnchor& InAnchor, float Parameter)
+	{
+		// Cannot prepend if Parameter is not less than all of our knots.
+		if (Base::PairKnots.Num() > 0 && Parameter >= Base::PairKnots[0].Value)
 		{
-			SegmentRange.Max = SegmentRange.Min + LoopKnotDelta;
+			return INDEX_NONE;
 		}
 
-		return SegmentRange;
-	}
+		Base::InsertValue(0, InAnchor.InHandle);
+		Base::InsertValue(1, InAnchor.Anchor);
+		Base::InsertValue(2, InAnchor.OutHandle);
 
-	// Parameterization methods
-	/**
-	 * Reparameterizes the spline based on the provided points and mode
-	 * @param Mode - New parameterization mode
-	 * @param Points - Array of control points to use for reparameterization
-	 * @return 
-	 */
-	virtual void Reparameterize(EParameterizationPolicy Mode = EParameterizationPolicy::Centripetal) override
-	{
-		const TArray<ValueType>& Points = this->Values;
-		const int32 NumValues = Points.Num();
+		const int32 KnotIndex = Base::InsertKnot(FKnot(Parameter));
 
-		GenerateDistanceKnotsForBezier(Mode);
-    	
-		UE_LOG(LogSpline, Verbose, TEXT("Set knot vector (Reparameterize) - Points: %d, Knots: %d, Mode: %d"),
-					   NumValues,  this->PairKnots.Num(), static_cast<int32>(Mode));
-		
-		Base::PrintKnotVector();
+		ensureAlwaysMsgf(KnotIndex == 0, TEXT("Unexpected knot insertion index when prepending an anchor!"));
+
+		ApplyClampedKnotsMultiplicity();
+		ApplyInternalKnotsMultiplicity();
+
+		return 0;
 	}
 	
-	virtual FInterval1f GetParameterSpace() const override
+	int32 PrependAnchor(const FAnchor& InAnchor, EParameterizationPolicy ParameterizationPolicy = EParameterizationPolicy::Centripetal)
 	{
-		if (Base::PairKnots.Num() == 0)
+		float NewParameter;
+
+		if (GetNumAnchorPoints() <= 0)
 		{
-			return FInterval1f::Empty();
+			NewParameter = 0.f;
 		}
+		else
+		{
+			const ValueType OldAnchor = GetAnchorPoint(0);
+			const float OldParameter = GetParameter(0);
 
-		const float Min = Base::PairKnots[0].Value;
-		const float LastKnot = Base::PairKnots.Last().Value;
+			const float SegmentLength = ComputeSegmentLength(InAnchor.Anchor, OldAnchor, ParameterizationPolicy);
+			const float NewParamByLength = OldParameter - SegmentLength;
+			const float NewParamByStep = Param::PrevDistinct(OldParameter);
 
-		if (!IsActuallyClosed())
-		{
-			return FInterval1f(Min, LastKnot);
+			NewParameter = FMath::Min(NewParamByLength, NewParamByStep);
 		}
 
-		// Closed: closing boundary is last real knot + LoopKnotDelta
-		return FInterval1f(Min, LastKnot + LoopKnotDelta);
+		return PrependAnchor(InAnchor, NewParameter);
 	}
 
-	/**
-	 * Gets the parameter value for a specific Bezier control point
-	 * Index is a *logical* Bezier control index:
-	 * Index = SegmentIndex * 4 + LocalPointIndex (0..3)
-	 * @param Index - Index of the Bezier control point
-	 * @return Parameter value at the specified control point
-	 */
-	virtual float GetParameter(int32 Index) const override
+	int32 AppendAnchor(const FAnchor& InAnchor, float Parameter)
 	{
-		const int32 NumSegments = GetNumberOfSegments();
-		if (NumSegments <= 0)
+		// Cannot append if Parameter is not greater than all of our knots.
+		if (Base::PairKnots.Num() > 0 && Parameter <= Base::PairKnots.Last().Value)
 		{
-			return 0.f;
+			return INDEX_NONE;
 		}
 
-		const int32 MaxLogicalIndex = NumSegments * 4;
-		if (Index < 0 || Index > MaxLogicalIndex || Base::PairKnots.IsEmpty())
-		{
-			return 0.f;
-		}
-		
-        int32 SegmentIndex = Index / 4;
-        int32 LocalPointIndex = Index % 4;
-    
-		// Logical "one past last" control index – treat as the domain end.
-		if (SegmentIndex >= NumSegments)
-		{
-			const FInterval1f Domain = GetParameterSpace();
-			return Domain.Max;
-		}
+		const int32 InsertIndex = GetNumAnchorPoints();
 
-		// Use the segment parameter range helper, which already knows how to
-		// handle the closing segment for closed loops.
-		const FInterval1f Range = GetSegmentParameterRange(SegmentIndex);
+		Base::AddValue(InAnchor.InHandle);
+		Base::AddValue(InAnchor.Anchor);
+		Base::AddValue(InAnchor.OutHandle);
 
-        const float T = static_cast<float>(LocalPointIndex) / 3.0f;
-        return Range.Interpolate(T);
-	}
+		const int32 KnotIndex = Base::InsertKnot(FKnot(Parameter));
 
-	/**
-	* Sets the parameter value for a specific Bezier control point
-	* Index is a *logical* Bezier control index:
-	* Index = SegmentIndex * 4 + LocalPointIndex (0..3)
-	* @param Index - Index of the Bezier control point
-	* @param NewParameter - New parameter value to set
-	* @return The index of the bezier point that was actually updated, or INDEX_NONE if failed
-	*/ 
-	virtual int32 SetParameter(int32 Index, float NewParameter) override
-	{
-		const int32 NumSegments = GetNumberOfSegments();
-		if (NumSegments <= 0)
-		{
-			return INDEX_NONE;
-		}
+		ensureAlwaysMsgf(KnotIndex == InsertIndex, TEXT("Unexpected knot insertion index when appending an anchor!"));
 
-		// Max logical control index is 4 * NumSegments:
-		// For open: P0..P3 of each seg -> 4 * NumSegments controls, with
-		// the last P3 coincident with the last anchor.
-		const int32 MaxLogicalIndex = NumSegments * 4;
-		if (Index < 0 || Index > MaxLogicalIndex)
-		{
-			return INDEX_NONE;
-		}
-		
-		constexpr int32 PointsPerSegment = 4;
-		const int32 SegmentIndex = Index / PointsPerSegment;
-		const int32 LocalPointIndex = Index % PointsPerSegment;
+		ApplyClampedKnotsMultiplicity();
+		ApplyInternalKnotsMultiplicity();
 
-		const bool bClosingBoundary = IsClosingSegment(SegmentIndex) && (LocalPointIndex != 0); // P1/P2/P3 of last seg
+		return InsertIndex;
+	}
+	
+	int32 AppendAnchor(const FAnchor& InAnchor, EParameterizationPolicy ParameterizationPolicy = EParameterizationPolicy::Centripetal)
+	{
+		const int32 NumAnchors = GetNumAnchorPoints();
 
-		float OldParameter = 0.f;
-		int32 OldParameterIndex = INDEX_NONE;
+		float NewParameter;
 
-		if (!bClosingBoundary)
+		if (NumAnchors <= 0)
 		{
-			// Legacy behaviour: real knot index
-			OldParameterIndex = (LocalPointIndex == 0) ? SegmentIndex : SegmentIndex + 1;
-			if (!Base::PairKnots.IsValidIndex(OldParameterIndex))
-			{
-				return INDEX_NONE;
-			}
-			OldParameter = Base::PairKnots[OldParameterIndex].Value;
+			NewParameter = 0.f;
 		}
 		else
 		{
-			// Closing boundary: virtual "last + delta"
-			if (Base::PairKnots.Num() == 0)
-			{
-				return INDEX_NONE;
-			}
+			const ValueType OldAnchor = GetAnchorPoint(NumAnchors - 1);
+			const float OldParameter = GetParameter(NumAnchors - 1);
 
-			const float LastKnot = Base::PairKnots.Last().Value;
-			OldParameter = LastKnot + LoopKnotDelta;
-		}
-    
-		// Early exit if parameter isn't changing
-		if (OldParameter == NewParameter)
-		{
-			return Index;
+			const float SegmentLength = ComputeSegmentLength(InAnchor.Anchor, OldAnchor, ParameterizationPolicy);
+			const float NewParamByLength = OldParameter + SegmentLength;
+			const float NewParamByStep = Param::NextDistinct(OldParameter);
+
+			NewParameter = FMath::Max(NewParamByLength, NewParamByStep);
 		}
-		
-		UE_LOG(LogSpline, Verbose, TEXT("\t"));
-		UE_LOG(LogSpline, Verbose, TEXT("Before SetParameter(%d, %f):"), Index, NewParameter);
-		Base::Dump();
-		
-		// Prevent reuse of an existing knot.
-		const float DesiredNewParameter = NewParameter;
-		typename Base::FKnotSearchCache SearchCache{ this->PairKnots };
-		const float ActualNewParameter = this->GetNearestAvailableKnotValue(DesiredNewParameter, SearchCache);
 
-		const float DesiredToOldDistance = FMath::Abs(DesiredNewParameter - OldParameter);
-		const float DesiredToActualNewDistance = FMath::Abs(DesiredNewParameter - ActualNewParameter);
+		return AppendAnchor(InAnchor, NewParameter);
+	}
 
-		if (DesiredToActualNewDistance > DesiredToOldDistance)
+	/** Per-anchor structural inserts */
+	int32 InsertAnchor(const FAnchor& InAnchor, float Parameter)
+	{
+		if (GetNumAnchorPoints() <= 0)
 		{
-			// If our nearest valid knot is further away from the desired value than our current value is, we need to just keep our current value
-			return Index;
+			return PrependAnchor(InAnchor, Parameter);
 		}
-		
-		int32 ReturnIndex = INDEX_NONE;
-
-		// --- handle closing boundary separately ---------------------------------
-		if (bClosingBoundary)
-		{
-			// We are moving the virtual closing boundary of the last segment.
-			const float LastKnot = Base::PairKnots.Last().Value;
-			const float NewLoopDelta = NewParameter - LastKnot;
-
-			// Must remain a forward-advancing interval.
-			if (NewLoopDelta <= 0.f)
-			{
-				return Index;
-			}
-
-			LoopKnotDelta = NewLoopDelta;
-			Base::MarkFlatKnotsCacheDirty();
 
-			UE_LOG(LogSpline, Verbose, TEXT("\t"));
-			UE_LOG(LogSpline, Verbose, TEXT("After SetParameter(%d, %f) [closing]:"),
-				Index, NewParameter);
-			Base::Dump();
+		const FInterval1f ParameterSpace = GetParameterSpace();
 
-			return Index;
+		if (!ensureAlwaysMsgf(!ParameterSpace.IsEmpty(), TEXT("Parameter space is empty but we do not have an empty spline!")))
+		{
+			return INDEX_NONE;
 		}
-		// -----------------------------------------------------------------------------
-
-		if (LocalPointIndex == 0) // P0
+		else if (Parameter < ParameterSpace.Min)
 		{
-			int32 LeftSegmentIndex = SegmentIndex - 1;
-			int32 RightSegmentIndex = SegmentIndex;
-			
-			Base::SetKnot(OldParameterIndex, NewParameter);
-
-			int32 Segment = SegmentIndex;
-			if (NewParameter < OldParameter)
-			{
-				while (LeftSegmentIndex >= 0 && GetSegmentParameterRange(LeftSegmentIndex).Min > GetSegmentParameterRange(LeftSegmentIndex).Max)
-				{
-					FlipSegment(LeftSegmentIndex);
-					LeftSegmentIndex--;
-					RightSegmentIndex--;
-					Segment = RightSegmentIndex;
-				}
-			}
-			else if (NewParameter > OldParameter)
-			{
-				while (RightSegmentIndex <= NumSegments - 1 && GetSegmentParameterRange(RightSegmentIndex).Min > GetSegmentParameterRange(RightSegmentIndex).Max)
-				{
-					FlipSegment(RightSegmentIndex);
-					LeftSegmentIndex++;
-					RightSegmentIndex++;
-					Segment = RightSegmentIndex;
-				}
-			}
-			
-			ReturnIndex = Segment * PointsPerSegment + LocalPointIndex;
+			return PrependAnchor(InAnchor, Parameter);
 		}
-		else if (LocalPointIndex == 3) // P3
+		else if (Parameter > ParameterSpace.Max)
 		{
-			int32 LeftSegmentIndex = SegmentIndex;
-			int32 RightSegmentIndex = SegmentIndex + 1;
-			
-			Base::SetKnot(OldParameterIndex, NewParameter);
-
-			int32 Segment = SegmentIndex;
-			if (NewParameter < OldParameter)
-			{
-				while (LeftSegmentIndex >= 0 && GetSegmentParameterRange(LeftSegmentIndex).Min > GetSegmentParameterRange(LeftSegmentIndex).Max)
-				{
-					FlipSegment(LeftSegmentIndex);
-					LeftSegmentIndex--;
-					RightSegmentIndex--;
-					Segment = LeftSegmentIndex;
-				}
-			}
-			else if (NewParameter > OldParameter)
-			{
-				while (RightSegmentIndex <= NumSegments - 1 && GetSegmentParameterRange(RightSegmentIndex).Min > GetSegmentParameterRange(RightSegmentIndex).Max)
-				{
-					FlipSegment(RightSegmentIndex);
-					LeftSegmentIndex++;
-					RightSegmentIndex++;
-					Segment = LeftSegmentIndex;
-				}
-			}
-			
-			ReturnIndex = Segment * PointsPerSegment + LocalPointIndex;
+			return AppendAnchor(InAnchor, Parameter);
 		}
 
-		UE_LOG(LogSpline, Verbose, TEXT("\t"));
-		UE_LOG(LogSpline, Verbose, TEXT("After SetParameter(%d, %f):"), Index, NewParameter);
-		Base::Dump();
-		
-		return ReturnIndex;
-	}
-	
-	void FlipSegment(int32 Segment)
-	{
-		const int32 NumSegments = GetNumberOfSegments();
-		const int32 NumAnchors  = GetNumAnchorPoints();
-		if (Segment < 0 || Segment >= NumSegments || NumAnchors < 2)
+		float LocalT;
+		const int32 SegmentIndex = FindIndexForParameter(Parameter, LocalT);
+
+		if (SegmentIndex < 0 || SegmentIndex >= GetNumberOfSegments())
 		{
-			return;
+			// If Parameter is contained by ParameterSpace (confirmed by Prepend/Append special casing above)
+			// but the parameter maps to an invalid segment, something is wrong that we can't handle nicely.
+			return INDEX_NONE;
 		}
 
-		UE_LOG(LogSpline, Verbose, TEXT("\t"));
-		UE_LOG(LogSpline, Verbose, TEXT("Before FlipSegment(%d):"), Segment);
-		Base::Dump();
-		
-		// 1) Flip this segment's Bezier window: P0<->P3, P1<->P2
+		// Segment S is bounded by anchors S and S+1, except for the closing segment which is bounded by anchors S and 0.
+		// We will special case insertion on the closing segment as this is simply an append operation.
+		if (IsClosingSegment(SegmentIndex))
 		{
-			int32 Idx0, Idx1, Idx2, Idx3;
-			GetBezierIndices(Segment, Idx0, Idx1, Idx2, Idx3);
-
-			ValueType P0 = this->GetValue(Idx0);
-			ValueType P1 = this->GetValue(Idx1);
-			ValueType P2 = this->GetValue(Idx2);
-			ValueType P3 = this->GetValue(Idx3);
-
-			this->SetValue(Idx0, P3);
-			this->SetValue(Idx1, P2);
-			this->SetValue(Idx2, P1);
-			this->SetValue(Idx3, P0);
+			return AppendAnchor(InAnchor, Parameter);
 		}
 
-		// 2) Fix left neighbor (Segment-1): keep its out-tangent but move its end point to match
-	    if (Segment > 0 || (IsActuallyClosed() && Segment == 0))
-	    {
-	    	const int32 LeftSeg =
-			 (Segment > 0) ? (Segment - 1) : (NumSegments - 1); // wrap in closed case
-
-	    	int32 ThisIdx0, ThisIdx1, ThisIdx2, ThisIdx3;
-	    	int32 LeftIdx0, LeftIdx1, LeftIdx2, LeftIdx3;
-
-	    	GetBezierIndices(Segment, ThisIdx0, ThisIdx1, ThisIdx2, ThisIdx3);
-	    	GetBezierIndices(LeftSeg, LeftIdx0, LeftIdx1, LeftIdx2, LeftIdx3);
-
-	    	const ValueType ThisP0 = this->GetValue(ThisIdx0);
-	    	const ValueType LeftP2 = this->GetValue(LeftIdx2);
-	    	const ValueType LeftP3 = this->GetValue(LeftIdx3);
-
-	    	const ValueType Delta = LeftP3 - LeftP2; // original out-tangent
-	    	const ValueType NewEnd = ThisP0;
-	    	const ValueType NewIn  = NewEnd - Delta;
-
-	    	this->SetValue(LeftIdx3, NewEnd);
-	    	this->SetValue(LeftIdx2, NewIn);
-	    }
-
-	    // 3) Fix right neighbor (Segment+1): keep its in-tangent but move its start point to match
-	    if (Segment < NumSegments - 1 || IsClosingSegment(Segment))
-	    {
-	    	const int32 RightSeg =
-			 (Segment < NumSegments - 1) ? (Segment + 1) : 0; // wrap in closed case
-
-	    	int32 ThisIdx0, ThisIdx1, ThisIdx2, ThisIdx3;
-	    	int32 RightIdx0, RightIdx1, RightIdx2, RightIdx3;
-
-	    	GetBezierIndices(Segment, ThisIdx0, ThisIdx1, ThisIdx2, ThisIdx3);
-	    	GetBezierIndices(RightSeg, RightIdx0, RightIdx1, RightIdx2, RightIdx3);
-
-	    	const ValueType ThisP3  = this->GetValue(ThisIdx3);
-	    	const ValueType RightP0 = this->GetValue(RightIdx0);
-	    	const ValueType RightP1 = this->GetValue(RightIdx1);
-
-	    	const ValueType Delta   = RightP0 - RightP1; // original in-tangent
-	    	const ValueType NewStart = ThisP3;
-	    	const ValueType NewOut   = NewStart - Delta;
-
-	    	this->SetValue(RightIdx0, NewStart);
-	    	this->SetValue(RightIdx1, NewOut);
-	    }
-
-	    // 4) Swap knot intervals for the flipped segment & its neighbor interval
-	    Base::SwapKnots(Segment, Segment + 1);
+		// We may now assume anchors SegmentIndex and SegmentIndex + 1 are valid without bounds checking.
+		// In order to insert onto segment S, we must insert at the index S + 1 so that the new anchor is the new endpoint of segment S.
+		
+		FInterval1f SegmentRange = GetSegmentParameterRange(SegmentIndex);
 
-	    UE_LOG(LogSpline, Verbose, TEXT("\t"));
-	    UE_LOG(LogSpline, Verbose, TEXT("After FlipSegment(%d):"), Segment);
-	    Base::Dump();
-	}
-	
-	virtual int32 FindIndexForParameter(float Parameter, float& OutLocalParam) const override
-	{
-		const int32 NumSegments = GetNumberOfSegments();
-		if (NumSegments <= 0)
+		if (Parameter == SegmentRange.Min)
 		{
-			OutLocalParam = 0.0f;
-			return 0;
+			Parameter = Param::NextDistinct(Parameter);
 		}
-
-		int32 SegmentIndex = 0;
-		float LocalT       = 0.0f;
-		if (!MapGlobalParameterToLocalSegment(Parameter, SegmentIndex, LocalT))
+		else if (Parameter == SegmentRange.Max)
 		{
-			return INDEX_NONE;
+			Parameter = Param::PrevDistinct(Parameter);
+		}
+		
+		// Bail out if we lost usable range, this should not happen.
+		if (Parameter <= SegmentRange.Min || Parameter >= SegmentRange.Max)
+		{
+			ensureAlways(false);
+			return INDEX_NONE;
 		}
 
-		SegmentIndex = FMath::Clamp(SegmentIndex, 0, NumSegments - 1);
-		OutLocalParam = LocalT;
-
-		// This function historically returned a "control point index" = 4*segment,
-		// i.e. the (virtual) P0 index of that segment. We preserve that contract.
-		return SegmentIndex * 4;
-	}
-	
-	virtual void SetClosedLoop(bool bShouldClose) override
-	{
-		// Skip if state isn't changing
-		if (bShouldClose == this->IsClosedLoop())
-			return;
-        
-		// Just flip the topology flag; do not add/remove any control points.
-		Base::SetClosedLoop(bShouldClose);
-
-		this->ApplyClampedKnotsMultiplicity();
-		ApplyInternalKnotsMultiplicity();
-
-		// Note:
-		// - Knot topology (number and placement of knots) is rebuilt by Reparameterize()
-		//   when the owning spline chooses to (eg, FTangentSpline::Reparameterize).
-		// - We intentionally do NOT call AppendBezierSegment / RemoveSegment here anymore,
-		//   because the closing segment is now virtual: {A_last, Out_last, In_0, A_0}.
+		const int32 InsertIndex = SegmentIndex + 1;
+		Base::InsertValue(AnchorBase(InsertIndex) + 0, InAnchor.InHandle);
+		Base::InsertValue(AnchorBase(InsertIndex) + 1, InAnchor.Anchor);
+		Base::InsertValue(AnchorBase(InsertIndex) + 2, InAnchor.OutHandle);
 
-	}
+		// We always use Degree as the multiplicity because this is always an interior knot.
+		const int32 KnotIndex = Base::InsertKnot(FKnot(Parameter, Degree));
+		ensureAlways(InsertIndex == KnotIndex);
 
-	/**
-	 * Maps a local segment parameter [0,1] to global parameter space
-	 * @param SegmentIndex - Index of the segment
-	 * @param LocalParam - Local parameter within segment [0,1]
-	 * @return Parameter value in global space, or 0 if segment is invalid
-	 */
-	float MapLocalSegmentParameterToGlobal(int32 SegmentIndex, float LocalParam) const
-	{
-		FInterval1f SegmentRange = GetSegmentParameterRange(SegmentIndex);
-		return SegmentRange.Interpolate(LocalParam);
+		return InsertIndex;
 	}
 
 	/**
-	 * Maps a global parameter to a segment index and local parameter
-	 * @param GlobalParam - Parameter in global space
-	 * @param OutSegmentIndex - Output segment index
-	 * @param OutLocalParam - Output local parameter [0,1]
-	 * @return True if mapping succeeded
+	 * Removes a point from the spline.
+	 * @param PointIndex - Index of the point to remove
+	 * @return true if successfully removed
 	 */
-	bool MapGlobalParameterToLocalSegment(float GlobalParam, int32& OutSegmentIndex, float& OutLocalParam) const
+	bool RemoveAnchor(int32 PointIndex)
 	{
-		const int32 NumSegments = GetNumberOfSegments();
-		const UE::Geometry::FInterval1f ParameterSpace = GetParameterSpace();
-
-		if (ParameterSpace.IsEmpty())
+		if (PointIndex < 0 || PointIndex >= GetNumAnchorPoints())
 		{
 			return false;
 		}
 
-		// Left end: snap onto first segment start
-		if (GlobalParam <= ParameterSpace.Min)
-		{
-			OutSegmentIndex = 0;
-			OutLocalParam = 0.0f;
-			return true;
-		}
+		const int32 BaseIdx = AnchorBase(PointIndex);
+		Base::RemoveValue(BaseIdx + 2);
+		Base::RemoveValue(BaseIdx + 1);
+		Base::RemoveValue(BaseIdx + 0);
 
-		// Right end: snap onto last segment end
-		if (GlobalParam >= ParameterSpace.Max)
-		{
-			OutSegmentIndex = NumSegments - 1;
-			OutLocalParam = 1.0f;
-			return true;
-		}
+		Base::RemoveKnot(PointIndex);
 
-		const float LastKnot = Base::PairKnots.Last().Value;
+		ApplyClampedKnotsMultiplicity();
+		ApplyInternalKnotsMultiplicity();
 
-		// We can now assume GlobalParam is contained by ParameterSpace, no need for clamping.
+		return true;
+	}
 
-		// Closed loop: handle the synthetic closing segment [LastKnot, LastKnot + LoopKnotDelta]
-		if (GlobalParam >= LastKnot && ensure(IsActuallyClosed()))
+	// Split a segment using De Casteljau
+	bool SplitSegment(int32 SegmentIndex, float SegmentT)
+	{
+		if (SegmentIndex < 0 || SegmentIndex >= GetNumberOfSegments())
 		{
-			OutSegmentIndex = NumSegments - 1;
-			const float t = (GlobalParam - LastKnot) / FMath::Max(ParameterSpace.Max - LastKnot, UE_SMALL_NUMBER);
-			OutLocalParam = FMath::Clamp(t, 0.0f, 1.0f);
-			return true;
+			return false;
 		}
-        
-		int32 SegmentIndex = INDEX_NONE;
 
-		// todo: make CVar
-		constexpr bool bEnableSegmentCaching = true;
-		if (bEnableSegmentCaching && LastEvaluatedSegment != INDEX_NONE && LastEvaluatedSegment < Base::PairKnots.Num() - 1)
+		FInterval1f SegmentRange = GetSegmentParameterRange(SegmentIndex);
+		float Parameter = SegmentRange.Interpolate(SegmentT);
+
+		if (Parameter == SegmentRange.Min)
 		{
-			if (GlobalParam < Base::PairKnots[LastEvaluatedSegment + 1].Value && GlobalParam >= Base::PairKnots[LastEvaluatedSegment].Value)
-			{
-				// cache hit!
-				SegmentIndex = LastEvaluatedSegment;
-			}
+			Parameter = Param::NextDistinct(Parameter);
 		}
-		
-		if (SegmentIndex == INDEX_NONE)
+		else if (Parameter == SegmentRange.Max)
 		{
-			// cache miss!
-			int32 SegmentLow = 0;
-			int32 SegmentHigh = Base::PairKnots.Num() - 2;
-			SegmentIndex = SegmentLow + (SegmentHigh - SegmentLow) / 2;
-
-			while (SegmentLow <= SegmentHigh)
-			{
-				int32 Mid = SegmentLow + (SegmentHigh - SegmentLow) / 2;
-				const float& StartParam = Base::PairKnots[Mid].Value;
-				const float& EndParam   = Base::PairKnots[Mid + 1].Value;
-
-				if (GlobalParam >= StartParam && GlobalParam < EndParam)
-				{
-					SegmentIndex = Mid;
-					break;
-				}
-				else if (GlobalParam < StartParam)
-				{
-					SegmentHigh = Mid - 1;
-				}
-				else
-				{
-					SegmentLow = Mid + 1;
-				}
-			}
+			Parameter = Param::PrevDistinct(Parameter);
 		}
 		
-		OutSegmentIndex = SegmentIndex;
-		LastEvaluatedSegment = SegmentIndex;
-        
-		// Calculate local parameter
-		float SegmentLength = Base::PairKnots[SegmentIndex + 1].Value - Base::PairKnots[SegmentIndex].Value;
-		if (SegmentLength > UE_SMALL_NUMBER)
-		{
-			OutLocalParam = (GlobalParam - Base::PairKnots[SegmentIndex].Value) / SegmentLength;
-		}
-		else
+		// Bail out if we lost usable range
+		if (Parameter <= SegmentRange.Min || Parameter >= SegmentRange.Max)
 		{
-			OutLocalParam = 0.0f;
+			return false;
 		}
-        
-		return true;
-	}
 
-	int32 FindSegmentIndex(float Parameter, float &OutLocalParam) const
-	{
-		int32 ControlPointIdx = FindIndexForParameter(Parameter, OutLocalParam);
-		return ControlPointIdx / 4; // Each segment has 4 control points
-	}
+		ValueType Position = Base::Evaluate(Parameter);
 
-	void GetBezierIndices(int32 SegmentIndex, int32& OutP0Idx, int32& OutP1Idx, int32& OutP2Idx, int32& OutP3Idx) const
-	{
-		if (!IsClosingSegment(SegmentIndex))
+		FWindow Window;
+		if (!BuildRawWindowForSegment(SegmentIndex, Window))
 		{
-			OutP0Idx = P0Idx(SegmentIndex);
-			OutP1Idx = P1Idx(SegmentIndex);
-			OutP2Idx = P2Idx(SegmentIndex);
-			OutP3Idx = P3Idx(SegmentIndex);
+			// If we cannot build a window for segment index, something is wrong and we cannot handle it nicely.
+			return false;
 		}
-		else
-		{
-			const int32 NumAnchors = GetNumAnchorPoints();
 
-			OutP0Idx = Close_P0Idx(NumAnchors);
-			OutP1Idx = Close_P1Idx(NumAnchors);
-			OutP2Idx = Close_P2Idx();
-			OutP3Idx = Close_P3Idx();
-		}
-	}
-	
-	/** Per-anchor setters/getters */
-	void SetHandleOutOfAnchor(int32 AnchorIdx, const ValueType& Out)
-	{
-		this->SetValue(AnchorBase(AnchorIdx) + 2, Out);
-	}
-	const ValueType& GetHandleOutOfAnchor(int32 AnchorIdx) const
-	{
-		return this->GetValue(AnchorBase(AnchorIdx) + 2);
-	}
-	
-	void SetHandleInOfAnchor(int32 AnchorIdx, const ValueType& In)
-	{
-		this->SetValue(AnchorBase(AnchorIdx) + 0, In);
-	}
-	const ValueType& GetHandleInOfAnchor(int32 AnchorIdx) const
-	{
-		return this->GetValue(AnchorBase(AnchorIdx) + 0);
-	}
+		ValueType P0 = *Window[0];
+		ValueType P1 = *Window[1];
+		ValueType P2 = *Window[2];
+		ValueType P3 = *Window[3];
+		
+		// Use De Casteljau algorithm to split the curve
+		const float t = SegmentT;
+		const float mt = 1.0f - t;
+		
+		// Calculate split control points
+		ValueType Q0 = P0;
+		ValueType Q1 = P0 * mt + P1 * t;
+		ValueType Q2 = (P0 * mt + P1 * t) * mt + (P1 * mt + P2 * t) * t;
+		ValueType Q3 = Position;  // Use provided position for the split point
+		
+		ValueType R0 = Position;  // Use provided position for the split point
+		ValueType R1 = (P1 * mt + P2 * t) * mt + (P2 * mt + P3 * t) * t;
+		ValueType R2 = P2 * mt + P3 * t;
+		ValueType R3 = P3;
 
-	/** Per-anchor structural inserts */
-	void InsertAnchorAt(int32 At, const ValueType& In, const ValueType& A, const ValueType& Out)
-	{
-		this->InsertValue(AnchorBase(At) + 0, In);
-		this->InsertValue(AnchorBase(At) + 1, A);
-		this->InsertValue(AnchorBase(At) + 2, Out);
-	}
+		// Fixup endpoints of segment being split.
+		TPair<int32, int32> SegmentAnchors = GetSegmentAnchors(SegmentIndex);
+		SetHandleOutOfAnchor(SegmentAnchors.Key, Q1);
+		SetHandleInOfAnchor(SegmentAnchors.Value, R2);
 
-	/** Appends a new anchor point defined by its In, A, and Out values. */
-	void AppendAnchor(const ValueType& In, const ValueType& A, const ValueType& Out)
-	{
-		this->AddValue(In);
-		this->AddValue(A);
-		this->AddValue(Out);
+		// Insert anchor, splitting the segment.
+		InsertAnchor(FAnchor { Q2, Q3, R1 }, Parameter);
+
+		return true;
 	}
-	
+
 	virtual int32 GetExpectedNumKnots() const override
 	{
-		const int32 NumSegments = GetNumberOfSegments();	
+		const int32 NumAnchors = GetNumAnchorPoints();
+		const int32 NumSegments = GetNumberOfSegments();
 
-		if (NumSegments == 0)
+		if (NumAnchors == 0)
 		{
 			return 0;
 		}
+		else if (NumAnchors == 1)
+		{
+			if (Base::IsClosedLoop())
+			{
+				return 3;
+			}
+			else
+			{
+				return 4;
+			}
+		}
 
-		if (!IsActuallyClosed())
+		if (!HasClosingSegment())
 		{
 			// Open cubic poly-Bézier:
 			// start multiplicity = 4, end multiplicity = 4, each interior boundary = 3
@@ -1845,7 +1355,7 @@ public:
 	{
 		TArray<FKnot> KnotVector = Base::GetKnotVector();
 
-		if (IsActuallyClosed() && ensureAlwaysMsgf(KnotVector.Num() > 0, TEXT("Knot vector is empty, cannot compute the closing knot.")))
+		if (HasClosingSegment() && ensureAlwaysMsgf(KnotVector.Num() > 0, TEXT("Knot vector is empty, cannot compute the closing knot.")))
 		{
 			KnotVector.Add(FKnot(KnotVector.Last().Value + LoopKnotDelta, Degree));
 		}
@@ -1855,32 +1365,32 @@ public:
 
 	virtual bool SetKnotVector(const TArray<FKnot>& NewKnots) override
 	{
-		if (!this->ValidateKnotVector(NewKnots))
+		if (!Base::ValidateKnotVector(NewKnots))
 		{
 			return false;
 		}
 
-		this->PairKnots = NewKnots;
+		Base::PairKnots = NewKnots;
 
-		if (IsActuallyClosed())
+		if (HasClosingSegment())
 		{
-			const int32 LastIdx     = this->PairKnots.Num() - 1;
-			const float ClosingKnot = this->PairKnots[LastIdx].Value;
-			const float PrevKnot    = this->PairKnots[LastIdx - 1].Value;
+			const int32 LastIdx     = Base::PairKnots.Num() - 1;
+			const float ClosingKnot = Base::PairKnots[LastIdx].Value;
+			const float PrevKnot    = Base::PairKnots[LastIdx - 1].Value;
 
 			// Store the closing interval length and drop the closing knot.
 			LoopKnotDelta = ClosingKnot - PrevKnot;
 			ensure(LoopKnotDelta > 0.f);
 
-			this->PairKnots.RemoveAt(LastIdx, 1, EAllowShrinking::Yes);
+			Base::PairKnots.RemoveAt(LastIdx, 1, EAllowShrinking::Yes);
 		}
 		
-		this->MarkFlatKnotsCacheDirty();
+		Base::MarkFlatKnotsCacheDirty();
 
 		UE_LOG(LogSpline, Verbose, TEXT("Set knot vector (Direct) - Unique Knots: %d"),
-			   this->PairKnots.Num());
+			   Base::PairKnots.Num());
 
-		this->PrintKnotVector();
+		Base::PrintKnotVector();
 
 		return true;
 	}
@@ -1889,7 +1399,22 @@ public:
 	
 private:
 
-	
+	// Segment S is bounded by anchors S and S + 1. Closing segment is bounded by S and 0.
+	TPair<int32, int32> GetSegmentAnchors(int32 SegmentIndex)
+	{
+		if (SegmentIndex < 0 || SegmentIndex >= GetNumberOfSegments())
+		{
+			return { INDEX_NONE, INDEX_NONE };
+		}
+
+		if (IsClosingSegment(SegmentIndex))
+		{
+			return { SegmentIndex, 0 };
+		}
+
+		return { SegmentIndex, SegmentIndex + 1 };
+	}
+
 	static constexpr int32 AnchorBase(int32 i){ return 3*i; }
 
 	static constexpr int32 P0Idx(int32 s) { return AnchorBase(s) + 1; } // A_s
@@ -1923,7 +1448,7 @@ private:
 			return false;
 		}
 
-		auto I = [this](int32 idx){ return &this->Values[idx]; };
+		auto I = [this](int32 idx){ return &Base::Values[idx]; };
 
 		if (!IsClosingSegment(SegmentIndex))
 		{
@@ -1943,21 +1468,6 @@ private:
 		return Math::FSplineValidation::IsValidWindow<ValueType, Base::WindowSize>(OutWindow);
 	}
 
-	void RefreshClosingKnot(EParameterizationPolicy Policy)
-	{
-		if (!IsActuallyClosed())
-		{
-			LoopKnotDelta = 0.0f;
-			return;
-		}
-
-		const ValueType& P0 = Base::Values[ AnchorBase(GetNumAnchorPoints() - 1) + 1 ]; // A_last
-		const ValueType& P3 = Base::Values[ AnchorBase(0)                        + 1 ]; // A_0
-		LoopKnotDelta = ComputeSegmentLength(P0, P3, Policy);
-
-		Base::MarkFlatKnotsCacheDirty();
-	}
-
 	void MigrateLegacyBezierLayout()
 	{
 	    const int32 NumValues = Base::Values.Num();
@@ -2005,11 +1515,11 @@ private:
 	    checkf(NumValues % 4 == 0,
 	        TEXT("Legacy PolyBezier layout mismatch (NumValues = %d)"), NumValues);
 
-		const bool bClosed = this->IsClosedLoop();
+		const bool bClosed = Base::IsClosedLoop();
 
 	    const int32 NumSegments = NumValues / 4;
 
-	    const TArray<ValueType>& Legacy = this->Values;
+	    const TArray<ValueType>& Legacy = Base::Values;
 		// For open: anchors = NumSegments + 1
 		// For closed: anchors = NumSegments (closing segment is virtual)
 		const int32 NumAnchors = bClosed ? NumSegments : (NumSegments + 1);
@@ -2106,25 +1616,25 @@ private:
 	        NewValues[NewBase + 2] = Out;
 	    }
 
-	    this->Values = MoveTemp(NewValues);
+	    Base::Values = MoveTemp(NewValues);
 
 		// Fix up knot representation for closed loops.
 		// Legacy data uses an explicit closing knot in PairKnots; the new
 		// representation uses LoopKnotDelta instead.
-		if (bClosed && this->PairKnots.Num() > 2)
+		if (bClosed && Base::PairKnots.Num() > 2)
 		{
-			const int32 LastIdx     = this->PairKnots.Num() - 1;
-			const float ClosingKnot = this->PairKnots[LastIdx].Value;
-			const float PrevKnot    = this->PairKnots[LastIdx - 1].Value;
+			const int32 LastIdx     = Base::PairKnots.Num() - 1;
+			const float ClosingKnot = Base::PairKnots[LastIdx].Value;
+			const float PrevKnot    = Base::PairKnots[LastIdx - 1].Value;
 
 			// Store the closing interval length and drop the closing knot.
 			LoopKnotDelta = ClosingKnot - PrevKnot;
 
-			this->PairKnots.RemoveAt(LastIdx, 1, EAllowShrinking::Yes);
+			Base::PairKnots.RemoveAt(LastIdx, 1, EAllowShrinking::Yes);
 		}
 		
 		// Values changed (and potentially PairKnots layout), so flat-knot cache is dirty.
-	    this->MarkFlatKnotsCacheDirty();
+	    Base::MarkFlatKnotsCacheDirty();
 	}
 
 	/**
@@ -2142,15 +1652,24 @@ private:
 		// Values[3*i + 2] = Out_i
 
 		const int32 NumAnchors  = GetNumAnchorPoints();
-		const int32 NumSegments = GetNumberOfSegments();
 
-		// Need at least one segment
-		if (NumAnchors < 2 || NumSegments <= 0)
+		if (NumAnchors <= 0)
+		{
+			Base::PairKnots.Reset();
+			Base::MarkFlatKnotsCacheDirty();
+			return;
+		}
+		else if (NumAnchors == 1)
 		{
-			this->PairKnots.Reset();
-			this->MarkFlatKnotsCacheDirty();
+			// If we have a single anchor, our number of expected knots is itself the multiplicity of a single distinct knot.
+			Base::PairKnots = { FKnot(0.f, GetExpectedNumKnots()) };
+			Base::MarkFlatKnotsCacheDirty();
 			return;
 		}
+
+		// If we have > 1 anchors (guaranteed by the above branches), we ought to have at least 1 segment.
+		const int32 NumSegments = GetNumberOfSegments();
+		ensureAlways(NumSegments > 0);
 			
 		TArray<float> SegmentLengths;
 		SegmentLengths.Reserve(NumSegments);
@@ -2158,7 +1677,7 @@ private:
 		// For a Bezier spline, the control points don't all lie on the curve
 		// So we need to compute chord lengths between segment endpoints only.
 		// Segment i goes from A_i to A_{i+1} (open) or A_{(i+1)%N} (closed).
-		const bool bClosed = IsActuallyClosed();
+		const bool bClosed = HasClosingSegment();
 
 		for (int32 i = 0; i < NumSegments; ++i)
 		{
@@ -2167,8 +1686,8 @@ private:
 				? (i + 1) % NumAnchors
 				: (i + 1); // for open, NumSegments == NumAnchors - 1, so this is always valid
 	            
-			const ValueType& P0 = this->Values[AnchorBase(StartAnchorIdx) + 1]; // A_i
-			const ValueType& P3 = this->Values[AnchorBase(EndAnchorIdx)   + 1]; // A_{i+1} or A_0 (closed)
+			const ValueType& P0 = Base::Values[AnchorBase(StartAnchorIdx) + 1]; // A_i
+			const ValueType& P3 = Base::Values[AnchorBase(EndAnchorIdx)   + 1]; // A_{i+1} or A_0 (closed)
 
 			SegmentLengths.Add(ComputeSegmentLength(P0, P3, Mode));
 		}
@@ -2193,18 +1712,18 @@ private:
 		};
 		
 		// Reset pair knots
-		this->PairKnots.Reset();
+		Base::PairKnots.Reset();
         
-		typename Base::FValidKnotSearchParams SearchParams;
+		FValidKnotSearchParams SearchParams;
 		SearchParams.bSearchLeft = false;	// we are appending repeatedly, we want insertion order to be predictable (always go after conflicting knots).
 
-		typename Base::FKnotSearchCache SearchCache{ this->PairKnots };
+		FKnotSearchCache SearchCache{ Base::PairKnots };
 
-		const auto InsertKnotAndUpdateCache = [&](uint32 InMultiplicity)
+		const auto InsertKnotAndUpdateCache = [this, &SearchCache, &SearchParams](uint32 InMultiplicity)
 		{
-			const float KnotValue = this->GetNearestAvailableKnotValue(SearchParams, SearchCache);
+			const float KnotValue = Base::GetNearestAvailableKnotValue(SearchParams, SearchCache);
 			SearchCache.KnotValuesSet.Add(Param::NormalizeKey(KnotValue));
-			this->InsertKnot(FKnot(KnotValue, InMultiplicity));
+			Base::InsertKnot(FKnot(KnotValue, InMultiplicity));
 		};
 
 		constexpr int32 EndMult = Degree + 1;
@@ -2276,7 +1795,7 @@ private:
         // --------------------------------------------------------------------
 
 		// Mark flat knots cache as dirty
-		this->MarkFlatKnotsCacheDirty();
+		Base::MarkFlatKnotsCacheDirty();
 	}
 
     virtual void FindWindow(float Parameter, FWindow& OutWindow, FWindowStore& Scratch) const override
@@ -2288,9 +1807,10 @@ private:
 			return;
 		}
 
-		int32 SegmentIndex;
 		float LocalT;
-		if (!MapGlobalParameterToLocalSegment(Parameter, SegmentIndex, LocalT))
+		int32 SegmentIndex = FindIndexForParameter(Parameter, LocalT);
+
+		if (SegmentIndex == INDEX_NONE)
 		{
 			return;
 		}
@@ -2314,9 +1834,10 @@ private:
 			return Math::FSplineValidation::GetDefaultValue<ValueType>();
         }
 
-    	int32 SegmentIndex;
-    	float LocalT;
-		if (!MapGlobalParameterToLocalSegment(Parameter, SegmentIndex, LocalT))
+		float LocalT;
+		int32 SegmentIndex = FindIndexForParameter(Parameter, LocalT);
+
+		if (SegmentIndex == INDEX_NONE)
 		{
 			return Math::FSplineValidation::GetDefaultValue<ValueType>();
 		}
@@ -2361,9 +1882,9 @@ private:
 	/**
 	 * Checks if the spline is actually closed, which is semantically different from IsClosedLoop for 0 and 1 point splines.
 	 */
-	bool IsActuallyClosed() const
+	bool HasClosingSegment() const
 	{
-		return this->IsClosedLoop() && GetNumAnchorPoints() > 1;
+		return Base::IsClosedLoop() && GetNumAnchorPoints() > 1;
 	}
 
 	/**
@@ -2373,7 +1894,7 @@ private:
 	 */
 	bool IsClosingSegment(int32 SegmentIndex) const
 	{
-		return IsActuallyClosed() && (SegmentIndex == GetNumAnchorPoints() - 1);
+		return HasClosingSegment() && (SegmentIndex == GetNumAnchorPoints() - 1);
 	}
 
 	/**
@@ -2386,103 +1907,6 @@ private:
     	int32 NumSegments = GetNumberOfSegments();
     	return Index >= 0 && Index < NumSegments;
     }
-
-	/**
-	 * Helper method to add the first segment to the spline
-	 * @param ControlPoint - Bezier control point to insert
-	 * @param OutPointIndex - Index of the newly created bezier anchor point
-	 * @return true if the points were added, false otherwise
-	 */
-	bool InitializeFirstSegment(
-		const ValueType& ControlPoint,
-		int32& OutPointIndex)
-    {
-		const int32 NumAnchors = GetNumAnchorPoints();
-		// 1) No anchors yet: create A0 only, no segment.
-		if (NumAnchors == 0)
-		{
-			// Zero handles: In0 = A0, Out0 = A0
-			AppendAnchor(/*In0*/ ControlPoint, /*A0*/ ControlPoint, /*Out0*/ ControlPoint);
-			OutPointIndex = 0;
-			return true;
-		}
-
-		// 2) One anchor exists, adding this point will create the first segment.
-		if (NumAnchors == 1)
-		{
-			const ValueType& A0 = Base::Values[P0Idx(0)];
-
-			// Simple auto-tangent guess between A0 and new point
-			ValueType Dir = ControlPoint - A0;
-			ValueType P0  = A0;
-			ValueType P3  = ControlPoint;
-			ValueType P1  = P0 + Dir * ValueType(0.33f);
-			ValueType P2  = P0 + Dir * ValueType(0.66f);
-	        
-    		// Add the first segment
-    		SetControlPoints({P0, P1, P2, P3}, EParameterizationPolicy::Centripetal);
-    		OutPointIndex = 1;
-    		return true;
-    	}
-
-		// 3) We already have at least one segment – let caller take the normal path.
-    	return false;
-    }
-
-	/**
-     * Prepends a point to the spline at the beginning.
-     * The new point will be connected to the first segment.
-     * @param Position The position of the new point
-     * @param ParameterizationPolicy How to distribute Knots
-     * @return Index of the new segment
-     */
-    int32 PrependPoint(
-		const ValueType& Position,
-		EParameterizationPolicy ParameterizationPolicy = EParameterizationPolicy::Centripetal)
-	{
-		const ValueType& P3 = this->Values[AnchorBase(0) + 1];
-
-		// Calculate reasonable auto-tangents
-		ValueType Dir = P3 - Position;
-		ValueType P1 = Position + Dir * 0.33f;
-		ValueType P2 = Position + Dir * 0.66f;
-    
-		// Add the segment by reusing existing code
-		int32 NewIndex = AddBezierSegmentInternal(0, Position, P1, P2, P3, ParameterizationPolicy, false);
-		check (NewIndex == 0);
-		return NewIndex;
-	}
-
-	/**
-	 * Appends a point to the end of the spline.
-	 * Creates a new cubic Bezier segment connecting the last point to the new point.
-	 * 
-	 * @param Position The position of the new point
-	 * @param ParameterizationPolicy How to distribute Knots
-	 * @return Index of the new segment
-	 */
-	int32 AppendPoint(
-		const ValueType& Position,
-		EParameterizationPolicy ParameterizationPolicy)
-	{
-		const int32 NumAnchors = GetNumAnchorPoints();
-		const int32 LastAnchorIdx = NumAnchors - 1;
-
-		// Append - create a point at the end
-		const ValueType& P0 = this->Values[AnchorBase(LastAnchorIdx) + 1];
-
-        
-		// Calculate reasonable auto-tangents
-		ValueType Dir = Position - P0;
-		ValueType P1 = P0 + Dir * 0.33f;
-		ValueType P2 = P0 + Dir * 0.66f;
-
-		int32 NumSegments = GetNumberOfSegments();
-		// Add the segment
-		int32 NewIndex = AddBezierSegmentInternal(NumSegments, P0, P1, P2, Position, ParameterizationPolicy/*, StartParam*/);
-
-		return NewIndex;
-	}
 	
 	/**
 	 * Helper to compute segment length based on parameterization policy
@@ -2516,116 +1940,6 @@ private:
 		}
 	}
 
-	int32 AddBezierSegmentInternal(
-		int32 SegmentIndex,
-		const ValueType& P0,
-		const ValueType& P1, 
-		const ValueType& P2, 
-		const ValueType& P3,
-		EParameterizationPolicy ParameterizationPolicy = EParameterizationPolicy::Centripetal,
-		bool bAppendToSegment = true) // true = connect to end, false = connect to start
-	{
-		const int32 NumAnchors  = GetNumAnchorPoints();
-	    const int32 NumSegments = GetNumberOfSegments();
-
-	    // Clamp requested slot into legal insertion range
-	    SegmentIndex = FMath::Clamp(SegmentIndex, 0, NumSegments);
-
-		auto FixKnots = [this](EParameterizationPolicy Policy)
-		{
-			this->ApplyClampedKnotsMultiplicity();
-			ApplyInternalKnotsMultiplicity();
-
-			if (IsActuallyClosed())
-			{
-				RefreshClosingKnot(Policy);
-			}
-		};
-
-	    // Measure for spacing when we do need to place a new boundary knot
-	    const float SegLen = ComputeSegmentLength(P0, P3, ParameterizationPolicy);
-
-	    // ===== CASE A: first-ever segment (no anchors yet) =====
-	    if (NumAnchors == 0)
-	    {
-	        // Seed A0 from P0.
-	        // IMPORTANT: make In0 a zero-handle by setting it equal to A0.
-	        // (We cannot infer a real incoming handle for A0 from {P0,P1,P2,P3} of the *first* segment.)
-	        AppendAnchor(/*In0*/ P0, /*A0*/ P0, /*Out0*/ P1);
-
-	        // Append A1 from {P2,P3}. Out1 is a zero-handle placeholder (equal to A1)
-	        AppendAnchor(/*In1*/ P2, /*A1*/ P3, /*Out1*/ P3);
-
-	        // Build a minimal clamped knot vector for 1 segment (open):
-	        this->ResetKnotVector();
-	        const int32 EndMult = Degree + 1;
-	        this->InsertKnot(FKnot(0.0f, EndMult));
-	        this->InsertKnot(FKnot(SegLen, EndMult));
-	    	FixKnots(ParameterizationPolicy);
-
-	        return 0; // the first (and only) segment index
-	    }
-
-	    // ===== CASE B: append at tail =====
-	    const bool bAppendAtTail = bAppendToSegment && (SegmentIndex >= NumSegments);
-	    if (bAppendAtTail)
-	    {
-	        const int32 Last = NumAnchors - 1;
-
-	        // Set Out_last from P1
-	        SetHandleOutOfAnchor(Last, P1);
-
-	        // Add new tail anchor from {P2,P3}; Out_new is zero-handle until more segments arrive
-	        AppendAnchor(/*In_n*/ P2, /*A_n*/ P3, /*Out_n*/ P3);
-
-	        // Insert a new boundary knot at the end (incremental)
-	        const float Kprev         = this->PairKnots.IsEmpty() ? 0.f : this->PairKnots.Last().Value;
-	        const float NextByLength  = Kprev + SegLen;
-	        const float NextByStep    = UE::Geometry::Spline::Param::NextDistinct(Kprev);
-	        const float Knew          = FMath::Max(NextByLength, NextByStep);
-	        this->InsertKnot(FKnot(Knew, Degree));
-
-	    	FixKnots(ParameterizationPolicy);
-	        return NumSegments; // newly appended segment index
-	    }
-
-	    // ===== CASE C: prepend at head =====
-	    const bool bPrependAtHead = (!bAppendToSegment && SegmentIndex == 0);
-	    if (bPrependAtHead)
-	    {
-	        // Insert new head anchor built from P0,P1
-	        InsertAnchorAt(/*At*/ 0, /*In_-1*/ P0, /*A_-1*/ P0, /*Out_-1*/ P1);
-
-	        // Old head is now anchor 1 — set its In_1 from P2
-	        SetHandleInOfAnchor(/*i=1*/ 1, P2);
-
-	        // Insert a boundary knot at the start (decremental)
-	        const float Kfirst        = this->PairKnots.IsEmpty() ? 0.f : this->PairKnots[0].Value;
-	        const float PrevByLength  = Kfirst - SegLen;
-	        const float PrevByStep    = UE::Geometry::Spline::Param::PrevDistinct(Kfirst);
-	        const float Knew          = FMath::Min(PrevByLength, PrevByStep);
-	        this->InsertKnot(FKnot(Knew, Degree));
-
-	    	FixKnots(ParameterizationPolicy);
-	        return 0;
-	    }
-
-	    // ===== CASE D: insert in the middle (between A_s and A_{s+1}) =====
-	    //   We don’t touch knots here; caller will InsertKnot(TargetParam).
-	    {
-	        const int32 s = FMath::Clamp(SegmentIndex, 0, FMath::Max(0, NumSegments - 1));
-
-	        // Left anchor’s Out_s comes from P1
-	        SetHandleOutOfAnchor(s, P1);
-
-	        // Insert new anchor at s+1 from {P2,P3}; Out_new is zero-handle
-	        InsertAnchorAt(/*At*/ s+1, /*In*/ P2, /*A*/ P3, /*Out*/ P3);
-
-			FixKnots(ParameterizationPolicy);
-	        return s;
-	    }
-	}
-
 private:
 
 	/** This is the last segment fetched by parameter. Used to accelerate predictable sampling patterns. */
@@ -2634,7 +1948,7 @@ private:
 	/** Optional override for FindWindow*/
 	FFindWindowOverride FindWindowOverrideFn = nullptr;
 
-	/** Length of the closing segment’s parameter interval. Only meaningful when IsActuallyClosed() == true. */
+	/** Length of the closing segment’s parameter interval. Only meaningful when HasClosingSegment() == true. */
 	float LoopKnotDelta = 1.f;
 };
 
diff --git a/Engine/Source/Runtime/GeometryCore/Public/Splines/SplineInterfaces.h b/Engine/Source/Runtime/GeometryCore/Public/Splines/SplineInterfaces.h
--- a/Engine/Source/Runtime/GeometryCore/Public/Splines/SplineInterfaces.h
+++ b/Engine/Source/Runtime/GeometryCore/Public/Splines/SplineInterfaces.h
@@ -15,7 +15,7 @@ namespace Spline
 
 // Forward declarations
 class UE_EXPERIMENTAL(5.7, "New spline APIs are experimental.") ISplineInterface;
-template <typename VALUETYPE> class UE_EXPERIMENTAL(5.7, "New spline APIs are experimental.") TSplineInterface;
+template <typename ValueType> class UE_EXPERIMENTAL(5.7, "New spline APIs are experimental.") TSplineInterface;
 
 /**
  * Enumeration for out-of-bounds handling modes
@@ -87,13 +87,13 @@ public:
  * Typed interface for splines with a specific value type
  * Provides evaluation and other type-specific functionality
  */
-template <typename VALUETYPE>
+template <typename InValueType>
 class TSplineInterface : public ISplineInterface
 {
     static FString TypeName;
 
 public:
-    typedef VALUETYPE ValueType;
+    using ValueType = InValueType;
 
     /**
      * Default constructor - ensures value type name is set
@@ -103,7 +103,7 @@ public:
         // Ensure the value type name is initialized if not already
         if (TypeName.IsEmpty())
         {
-            TypeName = TSplineValueTypeTraits<VALUETYPE>::Name;
+            TypeName = TSplineValueTypeTraits<ValueType>::Name;
         }
     }
     
@@ -234,8 +234,8 @@ protected:
 };
 
 // Initialize static member
-template <typename VALUETYPE>
-inline FString TSplineInterface<VALUETYPE>::TypeName;
+template <typename ValueType>
+inline FString TSplineInterface<ValueType>::TypeName;
 
 } // namespace Spline
 } // namespace Geometry
diff --git a/Engine/Source/Runtime/GeometryCore/Public/Splines/TangentBezierSpline.h b/Engine/Source/Runtime/GeometryCore/Public/Splines/TangentBezierSpline.h
--- a/Engine/Source/Runtime/GeometryCore/Public/Splines/TangentBezierSpline.h
+++ b/Engine/Source/Runtime/GeometryCore/Public/Splines/TangentBezierSpline.h
@@ -11,8 +11,8 @@ namespace Geometry
 namespace Spline
 {
 	
-template<typename VALUETYPE> struct UE_EXPERIMENTAL(5.7, "New spline APIs are experimental.") TTangentBezierControlPoint;
-template<typename VALUETYPE> class UE_EXPERIMENTAL(5.7, "New spline APIs are experimental.") TTangentBezierSpline;
+template<typename ValueType> struct UE_EXPERIMENTAL(5.7, "New spline APIs are experimental.") TTangentBezierControlPoint;
+template<typename ValueType> class UE_EXPERIMENTAL(5.7, "New spline APIs are experimental.") TTangentBezierSpline;
 	
 /**
  * Enum defining how tangents are computed for a spline control point.
@@ -135,24 +135,24 @@ struct TTangentBezierControlPoint
  * piecewise Bezier curves internally for evaluation. Supports both manual
  * tangent control and automatic tangent computation.
  */
-template<typename VALUETYPE>
+template<typename InValueType>
 class TTangentBezierSpline :
-	public TSplineWrapper<TPolyBezierSpline<VALUETYPE>>,
-	private TSelfRegisteringSpline<TTangentBezierSpline<VALUETYPE>, VALUETYPE>
+	public TSplineWrapper<TPolyBezierSpline<InValueType>>,
+	private TSelfRegisteringSpline<TTangentBezierSpline<InValueType>, InValueType>
 {
 public:
-	typedef typename TSplineWrapper<TPolyBezierSpline<VALUETYPE>>::ValueType ValueType;
 	
-	using FTangentBezierControlPoint = TTangentBezierControlPoint<ValueType>;
-	using FPolyBezierSpline = TPolyBezierSpline<ValueType>;
-	using FWindow      = typename FPolyBezierSpline::FWindow;
-	using FWindowStore = typename FPolyBezierSpline::FWindowStore;
-	using TangentOps = Math::TTangentOps<ValueType>;
+	using ValueType 					= InValueType;
+	using Base 							= TPolyBezierSpline<ValueType>;
+	using FTangentBezierControlPoint 	= TTangentBezierControlPoint<ValueType>;
+	using FWindow      					= typename Base::FWindow;
+	using FWindowStore 					= typename Base::FWindowStore;
+	using TangentOps 					= Math::TTangentOps<ValueType>;
 	
 	// Generate compile-time type ID for TangentBezier
 	DECLARE_SPLINE_TYPE_ID(
 		TEXT("TangentBezier"),
-		*TSplineValueTypeTraits<VALUETYPE>::Name
+		*TSplineValueTypeTraits<ValueType>::Name
 	);
 	
 	TTangentBezierSpline()
@@ -169,27 +169,7 @@ public:
 	{
 		Tension = 0.f;
 		TangentModes = {ETangentMode::Auto, ETangentMode::Auto};
-		AssignInternalSpline(FPolyBezierSpline3d::CreateLine(StartPoint, EndPoint));
-	}
-
-	/**
-	 * Constructor with full control over tangents
-	 */
-	TTangentBezierSpline(
-		const ValueType& StartPoint, 
-		const ValueType& EndPoint,
-		const ValueType& StartTangent,
-		const ValueType& EndTangent,
-		bool bAutoTangents = false)
-			: TTangentBezierSpline()
-	{
-		Tension = 0.f;
-		TangentModes = {bAutoTangents ? ETangentMode::Auto : ETangentMode::User,
-						bAutoTangents ? ETangentMode::Auto : ETangentMode::User};
-		AssignInternalSpline( FPolyBezierSpline3d(StartPoint,
-											TangentOps::TangentOutToP1(StartPoint, StartTangent),
-											TangentOps::TangentInToP2(EndPoint, EndTangent),
-											EndPoint));
+		AssignInternalSpline(Base::CreateLine(StartPoint, EndPoint));
 	}
    
     /**  Copy constructor */
@@ -316,6 +296,15 @@ public:
 				bStationaryEndpoint == Other.bStationaryEndpoint;
 	}
 
+	static Base::FAnchor AnchorFromControlPoint(const FTangentBezierControlPoint& InPoint)
+	{
+		return {
+			TangentOps::TangentInToP2(InPoint.Position, InPoint.TangentIn),
+			InPoint.Position,
+			TangentOps::TangentOutToP1(InPoint.Position, InPoint.TangentOut)
+		};
+	}
+
 	// Static shape generators
 
     /**
@@ -332,7 +321,7 @@ public:
 		TTangentBezierSpline Result;
         
 		// Clear the default initialization
-		Result.AssignInternalSpline(FPolyBezierSpline3d::CreateLine(StartPoint, EndPoint));
+		Result.AssignInternalSpline(Base::CreateLine(StartPoint, EndPoint));
 		
         Result.Reparameterize(EParameterizationPolicy::Uniform);
         return Result;
@@ -358,7 +347,7 @@ public:
         TTangentBezierSpline Result;
         
         // Clear the default initialization
-        Result.AssignInternalSpline(FPolyBezierSpline3d::CreateCircleArc(
+        Result.AssignInternalSpline(Base::CreateCircleArc(
             Center, Radius, StartAngle, EndAngle, NumSegments));
         
 		// Set up tangent status for all points
@@ -385,7 +374,7 @@ public:
     {
         // Create the circle
         TTangentBezierSpline Result;
-        Result.AssignInternalSpline(FPolyBezierSpline3d::CreateCircle(
+        Result.AssignInternalSpline(Base::CreateCircle(
             ValueType(Center), Radius, NumSegments));
         
         // Set up tangent status for all points
@@ -413,7 +402,7 @@ public:
     {
         // Create the ellipse
         TTangentBezierSpline Result;
-        Result.AssignInternalSpline(FPolyBezierSpline3d::CreateEllipse(
+        Result.AssignInternalSpline(Base::CreateEllipse(
             Center, RadiusX, RadiusY, NumSegments));
         
 		// Set up tangent status for all points
@@ -424,9 +413,10 @@ public:
         
         return Result;
     }
+
 	virtual TUniquePtr<ISplineInterface> Clone() const override
 	{
-		TUniquePtr<TTangentBezierSpline<VALUETYPE>> Clone = MakeUnique<TTangentBezierSpline<VALUETYPE>>();
+		TUniquePtr<TTangentBezierSpline<ValueType>> Clone = MakeUnique<TTangentBezierSpline<ValueType>>();
     
 		// Copy internal spline
 		Clone->AssignInternalSpline(this->InternalSpline);
@@ -445,362 +435,206 @@ public:
     
 	/**
 	 * Replaces all points in the spline with the provided points.
-	 * 
 	 * @param Points - Array of control points to set
+	 * @param Policy - The policy to use for knot generation as the points are added
 	 * @return True if the points were successfully set
 	 */
-	bool SetControlPoints(const TArray<FTangentBezierControlPoint>& Points,
-		EParameterizationPolicy ParameterizationPolicy = EParameterizationPolicy::Uniform)
-	{
-	    const int32 NumPoints = Points.Num();
-		
+	bool SetPoints(const TArray<FTangentBezierControlPoint>& Points, EParameterizationPolicy Policy)
+	{		
 		Clear();
-		
-	    // Special case for less than 2 points
-	    if (NumPoints < 2)
-	    {
-	    	if (NumPoints == 1)
-	    	{
-	    		// Initialize the spline with a single point
-	    		int32 OutNewIndex;
-	    		return InitializeFirstSegment(Points[0], OutNewIndex);
-	    	}
-	    	return true;
-	    }
-	    
-	    // Convert to cubic Bezier control points for the internal spline
-	    TArray<ValueType> BezierPoints;
-	    
-	    // For each segment (between consecutive points), we need 4 control points
-	    // Reserve space based on number of segments
-	    const int32 NumSegments = NumPoints - 1;
-	    BezierPoints.Reserve(NumSegments * 4);
-	    
-	    for (int32 i = 0; i < NumSegments; ++i)
-	    {
-	        const FTangentBezierControlPoint& StartPoint = Points[i];
-	        const FTangentBezierControlPoint& EndPoint = Points[i + 1];
-	        
-	        ValueType P0 = StartPoint.Position;
-	        ValueType P3 = EndPoint.Position;
-	    	
-	    	ValueType OutTangent = StartPoint.TangentOut;
-	    	ValueType InTangent = EndPoint.TangentIn;
-	       	        
-	        // Calculate bezier control points
-	    	ValueType P1 = TangentOps::TangentOutToP1(P0, OutTangent);
-	    	ValueType P2 = TangentOps::TangentInToP2(P3, InTangent);
-	        
-	        // Add all 4 control points for this segment
-	        BezierPoints.Add(P0);
-	        BezierPoints.Add(P1);
-	        BezierPoints.Add(P2);
-	        BezierPoints.Add(P3);
-	    }
-	    
-	    // Set the control points in the internal spline
-	    InternalSpline.SetControlPoints(BezierPoints, ParameterizationPolicy);
-	    
-	    // Store tangent modes for later use
-	    TangentModes.Reset(NumPoints);
-	    for (const FTangentBezierControlPoint& Point : Points)
-	    {
-	        TangentModes.Add(Point.TangentMode);
-	    }
+		TangentModes.Reset(Points.Num());
+
+		for (const FTangentBezierControlPoint& Point : Points)
+		{
+			AppendPoint(Point, Policy);
+		}
 	    	    
 	    return true;
 	}
 
 	/**
-	 * Adds multiple points to the spline.
-	 * 
+	 * Appends multiple points to the spline.
 	 * @param Points - Array of control points to add
-	 * @param bAppend - If true, append to the end of spline; if false, prepend to the beginning
+	 * @param Policy - The policy to use for knot generation as the points are added
 	 * @return True if the points were successfully added
 	 */
-	 bool AddControlPoints(
-        const TArray<FTangentBezierControlPoint>& Points,
-        bool bAppend = true)
+	bool AppendPoints(const TArray<FTangentBezierControlPoint>& Points, EParameterizationPolicy Policy)
     {
         const int32 NumPoints = Points.Num();
         
         // Need at least 2 points to add a valid segment
         if (NumPoints < 2)
         {
-            UE_LOG(LogSpline, Warning, TEXT("AddControlPoints requires at least 2 points to add a valid segment. Got %d points."), NumPoints);
+            UE_LOG(LogSpline, Warning, TEXT("AppendPoints requires at least 2 points to add a valid segment. Got %d points."), NumPoints);
             return false;
         }
         
         // If the spline is currently empty (though it shouldn't be based on our design),
-        // this is equivalent to SetControlPoints
+        // this is equivalent to SetPoints
         if (GetNumPoints() == 0)
         {
-            return SetControlPoints(Points);
-        }
-        
-        if (bAppend)
-        {
-            // For each point, use AddPoint 
-            for (int32 i = 0; i < NumPoints; ++i)
-            {
-                // Skip the first point if it's the same as our last existing point 
-                // (to avoid duplicates at connection point)
-                if (i == 0)
-                {
-                    const int32 ExistingNumPoints = GetNumPoints();
-                    const ValueType LastExistingPoint = GetValue(ExistingNumPoints - 1);
-                    
-                    // Check if connection points are reasonably close
-                	constexpr double ConnectionTolerance = 1e-4;
-                	bool bSameEndpoint;
-
-                	if constexpr (Math::TIsQuaternionType<ValueType>::value)
-                	{
-                		// angular distance
-                		const double dot = Math::Dot(LastExistingPoint, Points[0].Position);
-						const double ang = FMath::Acos(FMath::Clamp(dot, -1.0, 1.0)); // radians
-						bSameEndpoint = (ang <= ConnectionTolerance);
-                	}
-                	else
-                	{
-                		const double DistanceSquared = Math::SizeSquared(LastExistingPoint - Points[0].Position);
-                		bSameEndpoint = (DistanceSquared <= ConnectionTolerance);
-                	}
-                    
-                    if (bSameEndpoint)
-                    {
-                        // Points are close - the first existing point is the same as our first new point
-                        // So update the last point's tangent mode if needed
-                        if (Points[0].TangentMode != ETangentMode::User && Points[0].TangentMode != ETangentMode::Broken)
-                        {
-                            SetPointTangentMode(ExistingNumPoints - 1, Points[0].TangentMode);
-                        }
-                        continue; // Skip adding this point
-                    }
-                }
-                
-                // Add each point normally
-                AppendPoint(Points[i]);
-            }
+            return SetPoints(Points, Policy);
         }
-        else
-        {
-            // For prepending, we need to add points in reverse order
-            for (int32 i = NumPoints - 1; i >= 0; --i)
-            {
-                // Skip the last point if it's the same as our first existing point
-                if (i == NumPoints - 1)
-                {
-                    const ValueType FirstExistingPoint = GetValue(0);
+
+		// For each point, use AddPoint 
+		for (int32 i = 0; i < NumPoints; ++i)
+		{
+			// Skip the first point if it's the same as our last existing point 
+			// (to avoid duplicates at connection point)
+			if (i == 0)
+			{
+				const int32 ExistingNumPoints = GetNumPoints();
+				const ValueType LastExistingPoint = GetValue(ExistingNumPoints - 1);
                     
-                    // Check if connection points are reasonably close
-                    constexpr double ConnectionTolerance = 1e-4;
-                	bool bSameEndpoint;
-
-                	if constexpr (Math::TIsQuaternionType<ValueType>::value)
-                	{
-                		// angular distance
-                		const double dot = Math::Dot(FirstExistingPoint, Points[i].Position);
-                		const double ang = FMath::Acos(FMath::Clamp(dot, -1.0, 1.0)); // radians
-                		bSameEndpoint = (ang <= ConnectionTolerance); // choose a small angular eps, e.g. ~1e-4 rad
-                	}
-	                else
-	                {
-		                const double DistanceSquared = Math::SizeSquared(FirstExistingPoint - Points[i].Position);
-	                	bSameEndpoint = (DistanceSquared <= ConnectionTolerance);
-	                }
+				// Check if connection points are reasonably close
+				constexpr double ConnectionTolerance = 1e-4;
+				bool bSameEndpoint;
+
+				if constexpr (Math::TIsQuaternionType<ValueType>::value)
+				{
+					// angular distance
+					const double dot = Math::Dot(LastExistingPoint, Points[0].Position);
+					const double ang = FMath::Acos(FMath::Clamp(dot, -1.0, 1.0)); // radians
+					bSameEndpoint = (ang <= ConnectionTolerance);
+				}
+				else
+				{
+					const double DistanceSquared = Math::SizeSquared(LastExistingPoint - Points[0].Position);
+					bSameEndpoint = (DistanceSquared <= ConnectionTolerance);
+				}
                     
-                    if (bSameEndpoint)
-                    {
-                        // Points are close - update first point's tangent mode if needed
-                        if (Points[i].TangentMode != ETangentMode::User && Points[i].TangentMode != ETangentMode::Broken)
-                        {
-                            SetPointTangentMode(0, Points[i].TangentMode);
-                        }
-                        continue; // Skip adding this point
-                    }
-                }
+				if (bSameEndpoint)
+				{
+					// Points are close - the first existing point is the same as our first new point
+					// So update the last point's tangent mode if needed
+					if (Points[0].TangentMode != ETangentMode::User && Points[0].TangentMode != ETangentMode::Broken)
+					{
+						SetPointTangentMode(ExistingNumPoints - 1, Points[0].TangentMode);
+					}
+					continue; // Skip adding this point
+				}
+			}
                 
-                // Prepend each point
-                PrependPoint(Points[i]);
-            }
-        }
-        
-        // Reparameterize for consistent parameter distribution
-        Reparameterize();
+			// Add each point normally
+			AppendPoint(Points[i], Policy);
+		}
         
         return true;
     }
-
-	
-
-    /**
-     * Inserts a new point at the specified parameter along the spline.
-     * 
-     * @param Parameter - Parameter value where to insert the point (Spline Space)
-	 * @param ControlPoint - Control point to insert
-     * @return The index of the newly created point
-     */
-    int32 InsertPointAtGlobalParam(float Parameter, 
-                     const FTangentBezierControlPoint& ControlPoint)
-    {
-		int32 SegmentIndex;
-		float LocalT;
-		bool bValidSegment = InternalSpline.MapGlobalParameterToLocalSegment(Parameter, SegmentIndex, LocalT);
-
-		return InsertPointAtSegmentParam(SegmentIndex, LocalT, ControlPoint);
-    }
-
+    
 	/**
-	 * Inserts a point at a specific segment with local parameter.
-	 * This splits the segment at the specified location and places the new point
-	 * at the provided position, which may change the spline shape.
-	 * 
-	 * @param SegmentIndex Index of the segment to split
-	 * @param LocalT Local parameter [0,1] within the segment
-	 * @param Position Position for the new point
-	 * @param ParameterizationPolicy How to distribute parameters
-	 * @return The index of the newly created point
+	 * Prepends a control point to the beginning of the spline.
+	 * @param Point - Control point to prepend
+	 * @param Parameter - The parameter of the control point being prepended, must be less than all existing knots.
+	 * @return Index of the new point, or INDEX_NONE if the operation failed.
 	 */
-	int32 InsertPointAtSegmentParam(
-		int32 SegmentIndex,
-		float LocalT,
-		const FTangentBezierControlPoint& ControlPoint,
-		EParameterizationPolicy ParameterizationPolicy = EParameterizationPolicy::Uniform)
+	int32 PrependPoint(const FTangentBezierControlPoint& Point, float Parameter)
 	{
-		int32 OutNewIndex;
-		// Special case: handle first two points
-		if (InitializeFirstSegment(ControlPoint, OutNewIndex))
+		// Add the anchor
+		int32 NewIndex = InternalSpline.PrependAnchor(AnchorFromControlPoint(Point), Parameter);
+
+		if (NewIndex != INDEX_NONE)
 		{
-			return OutNewIndex;
+			// Set tangent mode
+			TangentModes.InsertDefaulted(NewIndex);
+			SetPointTangentMode(NewIndex, Point.TangentMode);
 		}
+
+		return NewIndex;
+	}
+
+    /**
+	 * Prepends a control point to the beginning of the spline.
+	 * @param Point - Control point to prepend.
+	 * @param Policy - The policy to use for generating the knot for the point being prepended.
+	 * @return Index of the new point, or INDEX_NONE if the operation failed.
+	 */
+	int32 PrependPoint(const FTangentBezierControlPoint& Point, EParameterizationPolicy Policy)
+	{
+		// Add the anchor
+		int32 NewIndex = InternalSpline.PrependAnchor(AnchorFromControlPoint(Point), Policy);
 		
-		// Insert at that parameter with the exact position
-		// Have the internal spline handle the insertion
-		int32 InsertedPointIndex = InternalSpline.InsertPointAtSegmentParam(SegmentIndex, LocalT, ControlPoint.Position, ParameterizationPolicy);
-		
-		// If insertion was successful, update tangents if not auto-computed
-		if (InsertedPointIndex >= 0 &&
-			(ControlPoint.TangentMode == ETangentMode::User ||
-			 ControlPoint.TangentMode == ETangentMode::Broken))
+		if (NewIndex != INDEX_NONE)
 		{
-			SetTangentIn(InsertedPointIndex, ControlPoint.TangentIn);
-			SetTangentOut(InsertedPointIndex, ControlPoint.TangentOut);
+			// Set tangent mode
+			TangentModes.InsertDefaulted(NewIndex);
+			SetPointTangentMode(NewIndex, Point.TangentMode);
 		}
 
-		TangentModes.InsertDefaulted(InsertedPointIndex);
-		SetPointTangentMode(InsertedPointIndex, ControlPoint.TangentMode);
-
-		return InsertedPointIndex;
+		return NewIndex;
 	}
-
-    /**
-     * Inserts a new point at the specified parameter along the spline.
-     * 
-     * @param PointIndex - Index of the new Point
-	 * @param ControlPoint - Control point to insert
-     * @return The index of the newly created point
-     */
-    int32 InsertPointAtPosition(int32 PointIndex, 
-                     const FTangentBezierControlPoint& ControlPoint)
+	
+    /** 
+	 * Appends a control point to the end of the spline.
+	 * @param Point - Control point to append.
+	 * @param Parameter - The parameter of the control point being appended, must be greater than all existing knots (but may be less than the virtual closing knot).
+	 * @return Index of the new point, or INDEX_NONE if the operation failed.
+	 */
+    int32 AppendPoint(const FTangentBezierControlPoint& Point, float Parameter)
     {
-	    int32 NewPointIndex;
-	    if (InitializeFirstSegment(ControlPoint, NewPointIndex)) return NewPointIndex;
+		// Add the anchor
+		int32 NewIndex = InternalSpline.AppendAnchor(AnchorFromControlPoint(Point), Parameter);
 
-	    const int32 NumPoints = GetNumPoints();
-		// need to convert point index to segment index while also clamping,
-		// also keeping in mind that we could append or prepend a segment
-		const int32 SegmentIndex = FMath::Clamp(PointIndex - 1, 0, NumPoints - 1);
-		// Insert at that parameter with the exact position
-        // Have the internal spline handle the insertion
-        int32 InsertedPointIndex = InternalSpline.InsertPointAtPosition(SegmentIndex, ControlPoint.Position);
-		
-		// If insertion was successful, update tangents if not auto-computed
-		if (InsertedPointIndex >= 0 &&
-			(ControlPoint.TangentMode == ETangentMode::User ||
-			 ControlPoint.TangentMode == ETangentMode::Broken))
+		if (NewIndex != INDEX_NONE)
 		{
-			SetTangentIn(InsertedPointIndex, ControlPoint.TangentIn);
-			SetTangentOut(InsertedPointIndex, ControlPoint.TangentOut);
+			// Set tangent mode
+			TangentModes.InsertDefaulted(NewIndex);
+			SetPointTangentMode(NewIndex, Point.TangentMode);
 		}
-		
-		TangentModes.InsertDefaulted(InsertedPointIndex);
-		SetPointTangentMode(InsertedPointIndex, ControlPoint.TangentMode);
 
-        return InsertedPointIndex;
+		return NewIndex;
     }
-    
-    /**
-	 * Helper method to prepend a point to the spline.
-	 * 
-	 * @param ControlPoint Control point structure to prepend
-	 * @return Index of the new point
-	 */
-	void PrependPoint(const FTangentBezierControlPoint& ControlPoint)
-	{
-    	int32 NewIndex;
-
-		// Special case: handle first two points
-    	if (InitializeFirstSegment(ControlPoint, NewIndex))
-    	{
-    		return;
-    	}
-		        
-		// Calculate control points for the cubic Bezier segment
-		ValueType P0 = ControlPoint.Position;
-		ValueType P1 = TangentOps::TangentOutToP1(P0, ControlPoint.TangentOut);
-		ValueType P2 = TangentOps::TangentInToP2(GetValue(0), GetTangentIn(0));
-
-		// Add the segment
-		InternalSpline.PrependBezierSegment(P0, P1, P2, EParameterizationPolicy::Uniform);
-        
-		// Fixup the in tangent of the new anchor point
-		SetTangentIn(0, ControlPoint.TangentIn);
-		
-		// Set tangent mode
-		TangentModes.InsertDefaulted(0);
-		SetPointTangentMode(0, ControlPoint.TangentMode);
-	}
 
     /** 
 	 * Appends a control point to the end of the spline.
-	 * @param ControlPoint Structure containing position and tangent information
+	 * @param Point - Control point to append.
+	 * @param Policy - The policy to use for generating the knot for the point being appended.
+	 * @return Index of the new point, or INDEX_NONE if the operation failed.
 	 */
-    void AppendPoint(const FTangentBezierControlPoint& ControlPoint)
+    int32 AppendPoint(const FTangentBezierControlPoint& Point, EParameterizationPolicy Policy)
     {
-		int32 NewIndex = GetNumPoints();
-
-		// Special case: handle first two points
-    	if (InitializeFirstSegment(ControlPoint, NewIndex))
-    	{
-    		return;
-    	}
-    	
-		// Calculate control points for the cubic Bezier segment
-		ValueType P3 = ControlPoint.Position;
-		ValueType P2 = TangentOps::TangentInToP2(P3, ControlPoint.TangentIn);
-		ValueType P1 = TangentOps::TangentOutToP1(GetValue(NewIndex - 1), GetTangentOut(NewIndex - 1));
-    
-		// Add the segment
-		InternalSpline.AppendBezierSegment(P1, P2, P3, EParameterizationPolicy::Uniform);
-		
-		// Fixup the out tangent of the new anchor point
-		SetTangentOut(NewIndex, ControlPoint.TangentOut);
+		// Add the anchor
+		int32 NewIndex = InternalSpline.AppendAnchor(AnchorFromControlPoint(Point), Policy);
 
-		// Set tangent mode
-		TangentModes.InsertDefaulted(NewIndex);
-		SetPointTangentMode(NewIndex, ControlPoint.TangentMode);
+		if (NewIndex != INDEX_NONE)
+		{
+			// Set tangent mode
+			TangentModes.InsertDefaulted(NewIndex);
+			SetPointTangentMode(NewIndex, Point.TangentMode);
+		}
+
+		return NewIndex;
+    }
+
+	/**
+	 * Inserts a new point at the specified parameter along the spline.
+	 * @param Point - Control point to insert.
+	 * @param Parameter - The parameter of the control point being inserted.
+	 * @return Index of the new point, or INDEX_NONE if the operation failed.
+	 */
+    int32 InsertPoint(const FTangentBezierControlPoint& Point, float Parameter)
+    {
+		// Add the anchor
+		int32 NewIndex = InternalSpline.InsertAnchor(AnchorFromControlPoint(Point), Parameter);
+
+		if (NewIndex != INDEX_NONE)
+		{
+			// Set tangent mode
+			TangentModes.InsertDefaulted(NewIndex);
+			SetPointTangentMode(NewIndex, Point.TangentMode);
+		}
+
+		return NewIndex;
     }
 	
 	/**
-	 * Gets the control point at the specified index
-	 * @param Index Index of the point to retrieve
-	 * @return Control point structure with position and tangent information
+	 * Gets the control point at the specified index.
+	 * @param Index - Index of the point to retrieve.
+	 * @return Control point at the specified index.
 	 */
-	FTangentBezierControlPoint GetControlPoint(int32 Index) const
+	FTangentBezierControlPoint GetPoint(int32 Index) const
 	{
-		const int32 NumPoints = GetNumPoints();
-		if (Index < 0 || Index >= NumPoints)
+		if (Index < 0 || Index >= GetNumPoints())
 		{
 			return FTangentBezierControlPoint();
 		}
@@ -829,38 +663,8 @@ public:
      */
     void RemovePoint(int32 Index)
 	{
-		const int32 NumPoints = GetNumPoints();
-		if (Index < 0 || Index >= NumPoints)
-		{
-			return;
-		}
-		// special case for handling the last point and second last point
-		if (NumPoints == 1)
+		if (Index < 0 || Index >= GetNumPoints())
 		{
-			InternalSpline.Clear();
-			TangentModes.Reset();
-			return;
-		}
-		else if (NumPoints == 2 && !this->IsClosedLoop())
-		{
-			// Figure out which point remains (0 or 1)
-			const int32 RemainingIndex = (Index == 0) ? 1 : 0;
-	
-			// Build a control point for the remaining anchor
-			FTangentBezierControlPoint RemainingCP;
-			RemainingCP.Position    = GetValue(RemainingIndex);
-			RemainingCP.TangentIn   = GetTangentIn(RemainingIndex);
-			RemainingCP.TangentOut  = GetTangentOut(RemainingIndex);
-			RemainingCP.TangentMode = TangentModes.IsValidIndex(RemainingIndex)
-										? TangentModes[RemainingIndex]
-										: ETangentMode::Auto;
-	
-			InternalSpline.Clear();
-			TangentModes.Reset();
-	
-			int32 TmpIdx = INDEX_NONE;
-			InitializeFirstSegment(RemainingCP, TmpIdx);
-	
 			return;
 		}
     
@@ -870,7 +674,7 @@ public:
 			TangentModes.RemoveAt(Index);
 		}
 
-		InternalSpline.RemovePoint(Index);
+		InternalSpline.RemoveAnchor(Index);
     
 		// Ensure tangent modes array size matches number of points after removal
 		if (TangentModes.Num() > GetNumPoints())
@@ -882,38 +686,37 @@ public:
 	/** 
 	 * Modifies an existing point
 	 * @param Index - Index of the point to modify
-	 * @param ControlPoint - New control point data
+	 * @param Point - New control point data
 	 */
-	void ModifyPoint(int32 Index, const FTangentBezierControlPoint& ControlPoint)
+	void ModifyPoint(int32 Index, const FTangentBezierControlPoint& Point)
 	{
-        const int32 NumPoints = GetNumPoints();
-        if (Index < 0 || Index >= NumPoints)
+        if (Index < 0 || Index >= GetNumPoints())
         {
             return;
         }
 
-        SetPointTangentMode(Index, ControlPoint.TangentMode);
+        SetPointTangentMode(Index, Point.TangentMode);
 
         // Update position
-		InternalSpline.SetAnchorPoint(Index, ControlPoint.Position);
+		InternalSpline.SetAnchorPoint(Index, Point.Position);
         
         // If not auto tangents, apply the specified tangents
-        if (ControlPoint.TangentMode != ETangentMode::Auto && 
-            ControlPoint.TangentMode != ETangentMode::AutoClamped &&
-            ControlPoint.TangentMode != ETangentMode::LegacyAuto && 
-            ControlPoint.TangentMode != ETangentMode::LegacyAutoClamped && 
-            ControlPoint.TangentMode != ETangentMode::Linear && 
-            ControlPoint.TangentMode != ETangentMode::LegacyLinear && 
-            ControlPoint.TangentMode != ETangentMode::Constant)
+        if (Point.TangentMode != ETangentMode::Auto && 
+            Point.TangentMode != ETangentMode::AutoClamped &&
+            Point.TangentMode != ETangentMode::LegacyAuto && 
+            Point.TangentMode != ETangentMode::LegacyAutoClamped && 
+            Point.TangentMode != ETangentMode::Linear && 
+            Point.TangentMode != ETangentMode::LegacyLinear && 
+            Point.TangentMode != ETangentMode::Constant)
         {
-            SetTangentIn(Index, ControlPoint.TangentIn);
-            SetTangentOut(Index, ControlPoint.TangentOut);
+            SetTangentIn(Index, Point.TangentIn);
+            SetTangentOut(Index, Point.TangentOut);
         }
     }
 
 	void SetValue(int32 Index, const ValueType& NewValue)
 	{
-		FTangentBezierControlPoint NewPoint = GetControlPoint(Index);
+		FTangentBezierControlPoint NewPoint = GetPoint(Index);
 		NewPoint.Position = NewValue;
 		ModifyPoint(Index, NewPoint);
 	}
@@ -1068,17 +871,13 @@ public:
 	/** Gets the parameter value for the specified anchor point index */
 	float GetParameter(int32 Index) const
 	{
-		const int32 BezierControlPointIdx = (Index == 0) ? 0 : ((Index - 1) * 4 + 3);
-		return InternalSpline.GetParameter(BezierControlPointIdx);
+		return InternalSpline.GetParameter(Index);
 	}
 
 	/** Sets the parameter value for the specified anchor point index */
 	int32 SetParameter(int32 Index, float NewParameter)
     {
-		const int32 BezierControlPointIdx = (Index == 0) ? 0 : ((Index - 1) * 4 + 3);
-		const int32 NewIndex = (InternalSpline.SetParameter(BezierControlPointIdx, NewParameter) + 1) / 4;
-
-		return NewIndex;
+		return InternalSpline.SetParameter(Index, NewParameter);
     }
 
 	virtual FInterval1f GetParameterSpace() const override
@@ -1093,7 +892,7 @@ public:
 
 	int32 FindSegmentIndex(float Parameter, float &OutLocalParam) const
 	{
-		return InternalSpline.FindSegmentIndex(Parameter, OutLocalParam);
+		return InternalSpline.FindIndexForParameter(Parameter, OutLocalParam);
 	}
 
     /**
@@ -1118,7 +917,7 @@ public:
 
 	/**
 	 * Updates tangents for a specific point based on its tangent mode
-	 * @param Index Index of the point to update
+	 * @param Index - Index of the point to update
 	 */
 	void UpdateTangent(int32 Index)
 	{
@@ -1332,7 +1131,7 @@ public:
 	}
     
     /** Getter for internal spline */
-    const FPolyBezierSpline3d& GetInternalSpline() const
+    const Base& GetInternalSpline() const
     {
         return InternalSpline;
     }
@@ -1420,65 +1219,20 @@ private:
 			   Mode == ETangentMode::Broken;
 	}
 	
-	/**
-	 * Helper method to add the first segment to the spline
-	 * @param ControlPoint - Control point to insert
-	 * @param OutPointIndex - Index of the newly created point
-	 * @return true if the points were added, false otherwise
-	 */
-	bool InitializeFirstSegment(const FTangentBezierControlPoint& ControlPoint, int32& OutPointIndex)
-    {
-		const int32 NumPoints = GetNumPoints();
-    
-    	if (NumPoints == 0)
-		{
-			const ValueType A0   = ControlPoint.Position;
-			const ValueType Out0 = TangentOps::TangentOutToP1(A0, ControlPoint.TangentOut);
-			const ValueType In0  = TangentOps::TangentInToP2(A0, ControlPoint.TangentIn);
-	
-			InternalSpline.Clear();
-	
-			// Single anchor (no segments yet), but In/Out cache the tangents
-			InternalSpline.AppendAnchor(In0, A0, Out0);
-	
-			TangentModes.SetNum(1);
-			TangentModes[0] = ControlPoint.TangentMode;
-	
-			OutPointIndex = 0;
-			return true;
-		}
-		else if (NumPoints == 1)
-		{
-			// We already have anchor 0; this is the second point -> anchor 1
-	
-			const ValueType A1   = ControlPoint.Position;
-			const ValueType Out1 = TangentOps::TangentOutToP1(A1, ControlPoint.TangentOut);
-			const ValueType In1  = TangentOps::TangentInToP2(A1, ControlPoint.TangentIn);
-	
-			InternalSpline.AppendAnchor(In1, A1, Out1);
-	
-			TangentModes.Add(ControlPoint.TangentMode);
-	
-			// Rebuild parameterization if needed (Uniform or whatever policy you use)
-			InternalSpline.Reparameterize(EParameterizationPolicy::Uniform);
-	
-			OutPointIndex = 1;
-			return true;
-		}
-	
-		return false;
-    }
-	void AssignInternalSpline(const FPolyBezierSpline& NewSpline)
+	void AssignInternalSpline(const Base& NewSpline)
 	{
 		InternalSpline = NewSpline;
 		InstallFindWindowOverride();
 	}
-	void AssignInternalSpline(FPolyBezierSpline&& NewSpline)
+
+	void AssignInternalSpline(Base&& NewSpline)
 	{
 		InternalSpline = MoveTemp(NewSpline);
 		InstallFindWindowOverride();
 	}
+
 protected:
+
 	/* Tension parameter for auto-computed tangents [0,1] */
 	float Tension = 0.0f;
 
@@ -1487,9 +1241,11 @@ protected:
 
 	bool bStationaryEndpoint = false;
 	
-	TStaticArray<ValueType, TPolyBezierSpline<VALUETYPE>::Base::WindowSize> IntermediateWindow;
+	TStaticArray<ValueType, TPolyBezierSpline<ValueType>::Base::WindowSize> IntermediateWindow;
+
 private:
-	using TSplineWrapper<TPolyBezierSpline<VALUETYPE>>::InternalSpline;
+
+	using TSplineWrapper<TPolyBezierSpline<ValueType>>::InternalSpline;
 };
 
 using FTangentBezierSpline3f = TTangentBezierSpline<FVector3f>; 
