[Backout] - CL49822005
#rnx
Original CL Desc
-----------------------------------------------------------------
PolyBezier: Post anchor refactor cleanup

Since our refactor, we conceptually store anchor points. We evaluate bezier segments and are conceptually storing segments implicitly via their endpoints, but the methods for adding/removing segments don't make much sense anymore.

#Synced-CL 49803947
#Synced-CL 49803947

#Synced-CL 49839727

diff --git a/Engine/Source/Runtime/Engine/Classes/Curves/Splines/TangentSpline.h b/Engine/Source/Runtime/Engine/Classes/Curves/Splines/TangentSpline.h
--- a/Engine/Source/Runtime/Engine/Classes/Curves/Splines/TangentSpline.h
+++ b/Engine/Source/Runtime/Engine/Classes/Curves/Splines/TangentSpline.h
@@ -276,7 +276,7 @@ int32 FTangentSpline::AddAttributeValue(FName Name, const AttrType& Value, float
 		// Special case for empty spline
 		if (NumPoints == 0)
 		{
-			Channel->AppendPoint(ControlPoint, 0.f);
+			Channel->AppendPoint(ControlPoint);
 			SetAttributeChannelRange(Name, FInterval1f(ParentSpaceParameter, ParentSpaceParameter), EMappingRangeSpace::Parent);
 			return 0;
 		}
@@ -285,13 +285,15 @@ int32 FTangentSpline::AddAttributeValue(FName Name, const AttrType& Value, float
 		{
 			if (ParentSpaceParameter > MappedChildSpace.Min)
 			{
-				Channel->AppendPoint(ControlPoint, UE::Geometry::Spline::EParameterizationPolicy::Uniform);
+				Channel->AppendPoint(ControlPoint);
+				Channel->Reparameterize(UE::Geometry::Spline::EParameterizationPolicy::Uniform);
 				SetAttributeChannelRange(Name, FInterval1f(MappedChildSpace.Min, ParentSpaceParameter), EMappingRangeSpace::Parent);
 				return 1;
 			}
 			else
 			{
-				Channel->PrependPoint(ControlPoint, UE::Geometry::Spline::EParameterizationPolicy::Uniform);
+				Channel->PrependPoint(ControlPoint);
+				Channel->Reparameterize(UE::Geometry::Spline::EParameterizationPolicy::Uniform);
 				SetAttributeChannelRange(Name, FInterval1f(ParentSpaceParameter, MappedChildSpace.Min), EMappingRangeSpace::Parent);
 				return 0;
 			}
@@ -302,9 +304,10 @@ int32 FTangentSpline::AddAttributeValue(FName Name, const AttrType& Value, float
 		{
 			// It is important to compute ChildSpaceParameter before the AppendPoint call, otherwise it will not be correct.
 			float ChildSpaceParameter = MapParameterToChildSpace(Name, ParentSpaceParameter);
-			Channel->AppendPoint(ControlPoint, ChildSpaceParameter);
+			Channel->AppendPoint(ControlPoint);
 
 			// By growing the child space and the mapped parent range proportionally, we keep the internal points stable in parent space.
+			Channel->SetParameter(Channel->GetNumPoints() - 1, ChildSpaceParameter);
 			SetAttributeChannelRange(Name, FInterval1f(MappedChildSpace.Min, ParentSpaceParameter), EMappingRangeSpace::Parent);
 
 			return Channel->GetNumPoints() - 1;
@@ -315,16 +318,20 @@ int32 FTangentSpline::AddAttributeValue(FName Name, const AttrType& Value, float
 		{
 			// It is important to compute ChildSpaceParameter before the AppendPoint call, otherwise it will not be correct.
 			float ChildSpaceParameter = MapParameterToChildSpace(Name, ParentSpaceParameter);
-			Channel->PrependPoint(ControlPoint, ChildSpaceParameter);
+			Channel->PrependPoint(ControlPoint);
 
 			// By growing the child space and the mapped parent range proportionally, we keep the internal points stable in parent space.
+			Channel->SetParameter(0, ChildSpaceParameter);
 			SetAttributeChannelRange(Name, FInterval1f(ParentSpaceParameter, MappedChildSpace.Max), EMappingRangeSpace::Parent);
 
 			return 0;
 		}
 
 		float ChildSpaceParameter = MapParameterToChildSpace(Name, ParentSpaceParameter);
-		int NewPointIndex = Channel->InsertPoint(ControlPoint, ChildSpaceParameter);
+		float LocalT;
+		int32 SegmentIndex = Channel->FindSegmentIndex(ChildSpaceParameter, LocalT);
+
+		int NewPointIndex = Channel->InsertPointAtSegmentParam(SegmentIndex, LocalT, ControlPoint);
 		return NewPointIndex;
 	}
 
diff --git a/Engine/Source/Runtime/Engine/Private/Curves/Splines/TangentSpline.cpp b/Engine/Source/Runtime/Engine/Private/Curves/Splines/TangentSpline.cpp
--- a/Engine/Source/Runtime/Engine/Private/Curves/Splines/TangentSpline.cpp
+++ b/Engine/Source/Runtime/Engine/Private/Curves/Splines/TangentSpline.cpp
@@ -164,7 +164,7 @@ FTangentSpline::FTangentSpline(const FLegacyTangentSpline& Other)
 	}
 
 	// initialize spline
-	GetSpline().SetPoints(Points, UE::Geometry::Spline::EParameterizationPolicy::Uniform);
+	GetSpline().SetControlPoints(Points);
 	ResetRotation();
 	ResetScale();
 
@@ -220,7 +220,7 @@ FTangentSpline::FTangentSpline(const FSplineCurves& Other)
 	}
 
 	// initialize spline
-	GetSpline().SetPoints(Points, UE::Geometry::Spline::EParameterizationPolicy::Uniform);
+	GetSpline().SetControlPoints(Points);
 	ResetRotation();
 	ResetScale();
 
@@ -578,8 +578,6 @@ void FTangentSpline::ResetScale()
 
 void FTangentSpline::AddPoint(const FSplinePoint& Point)
 {
-	using namespace UE::Geometry::Spline;
-
 	ValidateRotScale();
 	ON_SCOPE_EXIT
 	{
@@ -597,17 +595,17 @@ void FTangentSpline::AddPoint(const FSplinePoint& Point)
 
 	if (NumPoints == 0 || Index >= NumPoints)
 	{
-		GetSpline().AppendPoint(ControlPoint, EParameterizationPolicy::Uniform);
+		GetSpline().AppendPoint(ControlPoint);
 
-		if (TTangentBezierSpline<FVector>* ScaleChild = GetTypedAttributeChannel<TTangentBezierSpline<FVector>>(ScaleAttrName))
+		if (UE::Geometry::Spline::TTangentBezierSpline<FVector>* ScaleChild = GetTypedAttributeChannel<UE::Geometry::Spline::TTangentBezierSpline<FVector>>(ScaleAttrName))
 		{
-			ScaleChild->AppendPoint(TTangentBezierSpline<FVector>::FTangentBezierControlPoint(Point.Scale), EParameterizationPolicy::Uniform);
+			ScaleChild->AppendPoint(UE::Geometry::Spline::TTangentBezierSpline<FVector>::FTangentBezierControlPoint(Point.Scale));
 			ScaleChild->SetKnotVector(GetSpline().GetKnotVector());
 		}
 
-		if (TTangentBezierSpline<FQuat>* RotChild = GetTypedAttributeChannel<TTangentBezierSpline<FQuat>>(RotationAttrName))
+		if (UE::Geometry::Spline::TTangentBezierSpline<FQuat>* RotChild = GetTypedAttributeChannel<UE::Geometry::Spline::TTangentBezierSpline<FQuat>>(RotationAttrName))
 		{
-			RotChild->AppendPoint(TTangentBezierSpline<FQuat>::FTangentBezierControlPoint(Point.Rotation.Quaternion()), EParameterizationPolicy::Uniform);
+			RotChild->AppendPoint(UE::Geometry::Spline::TTangentBezierSpline<FQuat>::FTangentBezierControlPoint(Point.Rotation.Quaternion()));
 			RotChild->SetKnotVector(GetSpline().GetKnotVector());
 		}
 
@@ -616,17 +614,17 @@ void FTangentSpline::AddPoint(const FSplinePoint& Point)
 	}
 	else if (Index == 0)
 	{
-		GetSpline().PrependPoint(ControlPoint, EParameterizationPolicy::Uniform);
+		GetSpline().PrependPoint(ControlPoint);
 
-		if (TTangentBezierSpline<FVector>* ScaleChild = GetTypedAttributeChannel<TTangentBezierSpline<FVector>>(ScaleAttrName))
+		if (UE::Geometry::Spline::TTangentBezierSpline<FVector>* ScaleChild = GetTypedAttributeChannel<UE::Geometry::Spline::TTangentBezierSpline<FVector>>(ScaleAttrName))
 		{
-			ScaleChild->PrependPoint(TTangentBezierSpline<FVector>::FTangentBezierControlPoint(Point.Scale), EParameterizationPolicy::Uniform);
+			ScaleChild->PrependPoint(UE::Geometry::Spline::TTangentBezierSpline<FVector>::FTangentBezierControlPoint(Point.Scale));
 			ScaleChild->SetKnotVector(GetSpline().GetKnotVector());
 		}
 
-		if (TTangentBezierSpline<FQuat>* RotChild = GetTypedAttributeChannel<TTangentBezierSpline<FQuat>>(RotationAttrName))
+		if (UE::Geometry::Spline::TTangentBezierSpline<FQuat>* RotChild = GetTypedAttributeChannel<UE::Geometry::Spline::TTangentBezierSpline<FQuat>>(RotationAttrName))
 		{
-			RotChild->PrependPoint(TTangentBezierSpline<FQuat>::FTangentBezierControlPoint(Point.Rotation.Quaternion()), EParameterizationPolicy::Uniform);
+			RotChild->PrependPoint(UE::Geometry::Spline::TTangentBezierSpline<FQuat>::FTangentBezierControlPoint(Point.Rotation.Quaternion()));
 			RotChild->SetKnotVector(GetSpline().GetKnotVector());
 		}
 
@@ -635,22 +633,20 @@ void FTangentSpline::AddPoint(const FSplinePoint& Point)
 	}
 
 	// Insert at calculated parameter
-	const int32 SegmentIndex = Index - 1;
-	const float Parameter = GetSpline().GetSegmentParameterRange(SegmentIndex).Interpolate(Fraction);
-
-	int32 InsertedIndex = GetSpline().InsertPoint(ControlPoint, Parameter);
+	int32 InsertedIndex = GetSpline().InsertPointAtSegmentParam(Index, Fraction, ControlPoint);
+	float NewPointInternalParam = GetSpline().GetParameter(InsertedIndex);
 
-	float ScaleParam = MapParameterToChildSpace(ScaleAttrName, Parameter);
-	if (TTangentBezierSpline<FVector>* ScaleChild = GetTypedAttributeChannel<TTangentBezierSpline<FVector>>(ScaleAttrName))
+	float ScaleParam = MapParameterToChildSpace(ScaleAttrName, NewPointInternalParam);
+	if (UE::Geometry::Spline::TTangentBezierSpline<FVector>* ScaleChild = GetTypedAttributeChannel<UE::Geometry::Spline::TTangentBezierSpline<FVector>>(ScaleAttrName))
 	{
-		ScaleChild->InsertPoint(TTangentBezierSpline<FVector>::FTangentBezierControlPoint(Point.Scale), ScaleParam);
+		ScaleChild->InsertPointAtGlobalParam(ScaleParam, UE::Geometry::Spline::TTangentBezierSpline<FVector>::FTangentBezierControlPoint(Point.Scale));
 		ScaleChild->SetKnotVector(GetSpline().GetKnotVector());
 	}
 
-	float RotParam = MapParameterToChildSpace(RotationAttrName, Parameter);
-	if (TTangentBezierSpline<FQuat>* RotChild = GetTypedAttributeChannel<TTangentBezierSpline<FQuat>>(RotationAttrName))
+	float RotParam = MapParameterToChildSpace(RotationAttrName, NewPointInternalParam);
+	if (UE::Geometry::Spline::TTangentBezierSpline<FQuat>* RotChild = GetTypedAttributeChannel<UE::Geometry::Spline::TTangentBezierSpline<FQuat>>(RotationAttrName))
 	{
-		RotChild->InsertPoint(TTangentBezierSpline<FQuat>::FTangentBezierControlPoint(Point.Rotation.Quaternion()), RotParam);
+		RotChild->InsertPointAtGlobalParam(RotParam, UE::Geometry::Spline::TTangentBezierSpline<FQuat>::FTangentBezierControlPoint(Point.Rotation.Quaternion()));
 		RotChild->SetKnotVector(GetSpline().GetKnotVector());
 	}
 
@@ -659,8 +655,6 @@ void FTangentSpline::AddPoint(const FSplinePoint& Point)
 
 void FTangentSpline::InsertPoint(const FSplinePoint& Point, int32 Index)
 {
-	using namespace UE::Geometry::Spline;
-
 	ValidateRotScale();
 	ON_SCOPE_EXIT
 	{
@@ -675,17 +669,17 @@ void FTangentSpline::InsertPoint(const FSplinePoint& Point, int32 Index)
 	// Special case for empty spline
 	if (NumPoints == 0 || Index >= NumPoints)
 	{
-		GetSpline().AppendPoint(ControlPoint, EParameterizationPolicy::Uniform);
+		GetSpline().AppendPoint(ControlPoint);
 
-		if (TTangentBezierSpline<FVector>* ScaleChild = GetTypedAttributeChannel<TTangentBezierSpline<FVector>>(ScaleAttrName))
+		if (UE::Geometry::Spline::TTangentBezierSpline<FVector>* ScaleChild = GetTypedAttributeChannel<UE::Geometry::Spline::TTangentBezierSpline<FVector>>(ScaleAttrName))
 		{
-			ScaleChild->AppendPoint(TTangentBezierSpline<FVector>::FTangentBezierControlPoint(Point.Scale), EParameterizationPolicy::Uniform);
+			ScaleChild->AppendPoint(UE::Geometry::Spline::TTangentBezierSpline<FVector>::FTangentBezierControlPoint(Point.Scale));
 			ScaleChild->SetKnotVector(GetSpline().GetKnotVector());
 		}
 
-		if (TTangentBezierSpline<FQuat>* RotChild = GetTypedAttributeChannel<TTangentBezierSpline<FQuat>>(RotationAttrName))
+		if (UE::Geometry::Spline::TTangentBezierSpline<FQuat>* RotChild = GetTypedAttributeChannel<UE::Geometry::Spline::TTangentBezierSpline<FQuat>>(RotationAttrName))
 		{
-			RotChild->AppendPoint(TTangentBezierSpline<FQuat>::FTangentBezierControlPoint(Point.Rotation.Quaternion()), EParameterizationPolicy::Uniform);
+			RotChild->AppendPoint(UE::Geometry::Spline::TTangentBezierSpline<FQuat>::FTangentBezierControlPoint(Point.Rotation.Quaternion()));
 			RotChild->SetKnotVector(GetSpline().GetKnotVector());
 		}
 
@@ -694,17 +688,17 @@ void FTangentSpline::InsertPoint(const FSplinePoint& Point, int32 Index)
 	}
 	else if (Index == 0)
 	{
-		GetSpline().PrependPoint(ControlPoint, EParameterizationPolicy::Uniform);
+		GetSpline().PrependPoint(ControlPoint);
 
-		if (TTangentBezierSpline<FVector>* ScaleChild = GetTypedAttributeChannel<TTangentBezierSpline<FVector>>(ScaleAttrName))
+		if (UE::Geometry::Spline::TTangentBezierSpline<FVector>* ScaleChild = GetTypedAttributeChannel<UE::Geometry::Spline::TTangentBezierSpline<FVector>>(ScaleAttrName))
 		{
-			ScaleChild->PrependPoint(TTangentBezierSpline<FVector>::FTangentBezierControlPoint(Point.Scale), EParameterizationPolicy::Uniform);
+			ScaleChild->PrependPoint(UE::Geometry::Spline::TTangentBezierSpline<FVector>::FTangentBezierControlPoint(Point.Scale));
 			ScaleChild->SetKnotVector(GetSpline().GetKnotVector());
 		}
 
-		if (TTangentBezierSpline<FQuat>* RotChild = GetTypedAttributeChannel<TTangentBezierSpline<FQuat>>(RotationAttrName))
+		if (UE::Geometry::Spline::TTangentBezierSpline<FQuat>* RotChild = GetTypedAttributeChannel<UE::Geometry::Spline::TTangentBezierSpline<FQuat>>(RotationAttrName))
 		{
-			RotChild->PrependPoint(TTangentBezierSpline<FQuat>::FTangentBezierControlPoint(Point.Rotation.Quaternion()), EParameterizationPolicy::Uniform);
+			RotChild->PrependPoint(UE::Geometry::Spline::TTangentBezierSpline<FQuat>::FTangentBezierControlPoint(Point.Rotation.Quaternion()));
 			RotChild->SetKnotVector(GetSpline().GetKnotVector());
 		}
 
@@ -712,26 +706,22 @@ void FTangentSpline::InsertPoint(const FSplinePoint& Point, int32 Index)
 		return;
 	}
 
-	// We handle 0 and NumPoints cases above
-	Index = FMath::Clamp(Index, 1, GetNumberOfPoints() - 1);
-
-	const int32 SegmentIndex = Index - 1;
-	float SegmentParam;
-	const float Parameter = GetSpline().FindNearestOnSegment(Point.Position, SegmentIndex, SegmentParam);
+	Index = FMath::Clamp(Index, 0, GetNumberOfPoints());
 
-	int32 InsertedIndex = GetSpline().InsertPoint(ControlPoint, Parameter);
+	int32 InsertedIndex = GetSpline().InsertPointAtPosition(Index, ControlPoint);
+	float NewPointInternalParam = GetSpline().GetParameter(InsertedIndex);
 
-	float ScaleParam = MapParameterToChildSpace(ScaleAttrName, Parameter);
-	if (TTangentBezierSpline<FVector>* ScaleChild = GetTypedAttributeChannel<TTangentBezierSpline<FVector>>(ScaleAttrName))
+	float ScaleParam = MapParameterToChildSpace(ScaleAttrName, NewPointInternalParam);
+	if (UE::Geometry::Spline::TTangentBezierSpline<FVector>* ScaleChild = GetTypedAttributeChannel<UE::Geometry::Spline::TTangentBezierSpline<FVector>>(ScaleAttrName))
 	{
-		ScaleChild->InsertPoint(TTangentBezierSpline<FVector>::FTangentBezierControlPoint(Point.Scale), ScaleParam);
+		ScaleChild->InsertPointAtGlobalParam(ScaleParam, UE::Geometry::Spline::TTangentBezierSpline<FVector>::FTangentBezierControlPoint(Point.Scale));
 		ScaleChild->SetKnotVector(GetSpline().GetKnotVector());
 	}
 
-	float RotParam = MapParameterToChildSpace(RotationAttrName, Parameter);
-	if (TTangentBezierSpline<FQuat>* RotChild = GetTypedAttributeChannel<TTangentBezierSpline<FQuat>>(RotationAttrName))
+	float RotParam = MapParameterToChildSpace(RotationAttrName, NewPointInternalParam);
+	if (UE::Geometry::Spline::TTangentBezierSpline<FQuat>* RotChild = GetTypedAttributeChannel<UE::Geometry::Spline::TTangentBezierSpline<FQuat>>(RotationAttrName))
 	{
-		RotChild->InsertPoint(TTangentBezierSpline<FQuat>::FTangentBezierControlPoint(Point.Rotation.Quaternion()), RotParam);
+		RotChild->InsertPointAtGlobalParam(RotParam, UE::Geometry::Spline::TTangentBezierSpline<FQuat>::FTangentBezierControlPoint(Point.Rotation.Quaternion()));
 		RotChild->SetKnotVector(GetSpline().GetKnotVector());
 	}
 
@@ -762,8 +752,6 @@ FSplinePoint FTangentSpline::GetPoint(int32 Index) const
 /** Removes a point from the spline */
 void FTangentSpline::RemovePoint(int32 Index)
 {
-	using namespace UE::Geometry::Spline;
-
 	ValidateRotScale();
 	ON_SCOPE_EXIT
 	{
@@ -780,13 +768,13 @@ void FTangentSpline::RemovePoint(int32 Index)
 
 	GetSpline().RemovePoint(Index);
 
-	if (TTangentBezierSpline<FVector>* ScaleChild = GetTypedAttributeChannel<TTangentBezierSpline<FVector>>(ScaleAttrName))
+	if (UE::Geometry::Spline::TTangentBezierSpline<FVector>* ScaleChild = GetTypedAttributeChannel<UE::Geometry::Spline::TTangentBezierSpline<FVector>>(ScaleAttrName))
 	{
 		ScaleChild->RemovePoint(Index);
 		ScaleChild->SetKnotVector(GetSpline().GetKnotVector());
 	}
 
-	if (TTangentBezierSpline<FQuat>* RotChild = GetTypedAttributeChannel<TTangentBezierSpline<FQuat>>(RotationAttrName))
+	if (UE::Geometry::Spline::TTangentBezierSpline<FQuat>* RotChild = GetTypedAttributeChannel<UE::Geometry::Spline::TTangentBezierSpline<FQuat>>(RotationAttrName))
 	{
 		RotChild->RemovePoint(Index);
 		RotChild->SetKnotVector(GetSpline().GetKnotVector());
@@ -1110,8 +1098,6 @@ void FTangentSpline::UpdateSpline()
 
 void FTangentSpline::UpdateSpline(const FUpdateSplineParams& Params)
 {
-	using namespace UE::Geometry::Spline;
-
 	TRACE_CPUPROFILER_EVENT_SCOPE(FTangentSpline::UpdateSpline);
 
 	// todo: handled unhandled params
@@ -1148,12 +1134,12 @@ void FTangentSpline::UpdateSpline(const FUpdateSplineParams& Params)
 				{
 					GetSpline().SetParameter(GetNumberOfPoints(), LoopKey);
 
-					if (TTangentBezierSpline<FVector>* ScaleChild = GetTypedAttributeChannel<TTangentBezierSpline<FVector>>(ScaleAttrName))
+					if (UE::Geometry::Spline::TTangentBezierSpline<FVector>* ScaleChild = GetTypedAttributeChannel<UE::Geometry::Spline::TTangentBezierSpline<FVector>>(ScaleAttrName))
 					{
 						ScaleChild->SetParameter(GetNumberOfPoints(), LoopKey);
 					}
 
-					if (TTangentBezierSpline<FQuat>* RotChild = GetTypedAttributeChannel<TTangentBezierSpline<FQuat>>(RotationAttrName))
+					if (UE::Geometry::Spline::TTangentBezierSpline<FQuat>* RotChild = GetTypedAttributeChannel<UE::Geometry::Spline::TTangentBezierSpline<FQuat>>(RotationAttrName))
 					{
 						RotChild->SetParameter(GetNumberOfPoints(), LoopKey);
 					}
@@ -1167,8 +1153,6 @@ void FTangentSpline::UpdateSpline(const FUpdateSplineParams& Params)
 
 void FTangentSpline::Reparameterize()
 {
-	using namespace UE::Geometry::Spline;
-
 	TRACE_CPUPROFILER_EVENT_SCOPE(FTangentSpline::Reparameterize);
 
 	// Caching segment space attribute parameters so we can restore them (necessary because we store the mapping in parent space which are about to be invalidated).
@@ -1214,19 +1198,19 @@ void FTangentSpline::Reparameterize()
 		}
 
 		RestoreChannelKnots(ChannelName, CachedChannelParameterizations[ChannelName].Value);
-		if (TTangentBezierSpline<float>* Channel = GetTypedAttributeChannel<TTangentBezierSpline<float>>(ChannelName))
+		if (UE::Geometry::Spline::TTangentBezierSpline<float>* Channel = GetTypedAttributeChannel<UE::Geometry::Spline::TTangentBezierSpline<float>>(ChannelName))
 		{
 			Channel->UpdateTangents();
 		}
 	}
 
-	if (TTangentBezierSpline<FVector>* ScaleSpline = GetTypedAttributeChannel<TTangentBezierSpline<FVector>>(ScaleAttrName))
+	if (UE::Geometry::Spline::TTangentBezierSpline<FVector>* ScaleSpline = GetTypedAttributeChannel<UE::Geometry::Spline::TTangentBezierSpline<FVector>>(ScaleAttrName))
 	{
 		ScaleSpline->SetKnotVector(GetSpline().GetKnotVector());
 		ScaleSpline->UpdateTangents();
 	}
 
-	if (TTangentBezierSpline<FQuat>* RotSpline = GetTypedAttributeChannel<TTangentBezierSpline<FQuat>>(RotationAttrName))
+	if (UE::Geometry::Spline::TTangentBezierSpline<FQuat>* RotSpline = GetTypedAttributeChannel<UE::Geometry::Spline::TTangentBezierSpline<FQuat>>(RotationAttrName))
 	{
 		RotSpline->SetKnotVector(GetSpline().GetKnotVector());
 		RotSpline->UpdateTangents();
diff --git a/Engine/Source/Runtime/GeometryCore/Public/Splines/BSpline.h b/Engine/Source/Runtime/GeometryCore/Public/Splines/BSpline.h
--- a/Engine/Source/Runtime/GeometryCore/Public/Splines/BSpline.h
+++ b/Engine/Source/Runtime/GeometryCore/Public/Splines/BSpline.h
@@ -16,7 +16,7 @@ namespace Spline
 {
 
 struct UE_EXPERIMENTAL(5.7, "New spline APIs are experimental.") FKnot;
-template <typename ValueType, int32 Degree> class UE_EXPERIMENTAL(5.7, "New spline APIs are experimental.") TBSpline;
+template <typename VALUETYPE, int32 DEGREE> class UE_EXPERIMENTAL(5.7, "New spline APIs are experimental.") TBSpline;
 	
 enum class UE_EXPERIMENTAL(5.7, "New spline APIs are experimental.") EParameterizationPolicy
 {
@@ -119,22 +119,22 @@ struct BSplineNameSelector<3>
 	static constexpr const TCHAR* Name = TEXT("BSpline3");
 };
 
-template <typename InValueType, int32 InDegree>
+template <typename VALUETYPE, int32 DEGREE>
 class TBSpline
-	: public TSplineInterface<InValueType>,
-	  private TSelfRegisteringSpline<TBSpline<InValueType, InDegree>, InValueType>
+	: public TSplineInterface<VALUETYPE>,
+	  private TSelfRegisteringSpline<TBSpline<VALUETYPE, DEGREE>, VALUETYPE>
 {
 public:
-	using ValueType = InValueType;
+	static constexpr int32 Degree = DEGREE;
+	using ValueType = typename TSplineInterface<VALUETYPE>::ValueType;
 
-	static constexpr int32 Degree = InDegree;
 	static constexpr int32 WindowSize = Degree + 1;
 	using FWindow = TStaticArray<const ValueType*, WindowSize>;
 	using FWindowStore = TStaticArray<ValueType, WindowSize>;
 	
 	DECLARE_SPLINE_TYPE_ID(
-		BSplineNameSelector<Degree>::Name,
-		*TSplineValueTypeTraits<ValueType>::Name
+		BSplineNameSelector<DEGREE>::Name,
+		*TSplineValueTypeTraits<VALUETYPE>::Name
 	);
 
 	TBSpline() = default;
@@ -238,7 +238,7 @@ public:
 
 	virtual TUniquePtr<ISplineInterface> Clone() const override
 	{
-		TUniquePtr<TBSpline<ValueType, Degree>> Clone = MakeUnique<TBSpline<ValueType, Degree>>();
+		TUniquePtr<TBSpline<VALUETYPE, DEGREE>> Clone = MakeUnique<TBSpline<VALUETYPE, DEGREE>>();
 
 		// Copy all member variables
 		Clone->Values = this->Values;
@@ -641,9 +641,9 @@ protected:
 	/**
 	 * Inserts a knot at the specified parameter value
 	 * @param InKnot - Knot to insert
-	 * @return index of new knot if knot insertion was successful, INDEX_NONE otherwise.
+	 * @return true if knot insertion was successful, false otherwise.
 	 */
-	int32 InsertKnot(FKnot InKnot)
+	bool InsertKnot(FKnot InKnot)
 	{
 		// Binary search the knot vector:
 
@@ -657,7 +657,7 @@ protected:
 			if (PairKnots[Mid].Value == InKnot.Value)
 			{
 				UE_LOG(LogSpline, Warning, TEXT("Knot insertion failed, knot already exists!"));
-				return INDEX_NONE;
+				return false;
 			}
 			else if (PairKnots[Mid].Value < InKnot.Value)
 			{
@@ -669,8 +669,10 @@ protected:
 			}
 		}
 
+		PairKnots.Insert(InKnot, Low);
 		MarkFlatKnotsCacheDirty();
-		return PairKnots.Insert(InKnot, Low);
+
+		return true;
 	}
 
 	struct FValidKnotSearchParams
diff --git a/Engine/Source/Runtime/GeometryCore/Public/Splines/CustomAttributeExample.h b/Engine/Source/Runtime/GeometryCore/Public/Splines/CustomAttributeExample.h
--- a/Engine/Source/Runtime/GeometryCore/Public/Splines/CustomAttributeExample.h
+++ b/Engine/Source/Runtime/GeometryCore/Public/Splines/CustomAttributeExample.h
@@ -199,11 +199,10 @@ void CreatePolyBezierWithFooAttributes()
     }
 
     // Add another segment to the main spline
-
-	Spline.GetSpline().AppendAnchor( FPolyBezierSpline3d::FAnchor(
-		FVector(150, -50, 0),  // In Handle
-		FVector(200, 0, 0),    // Anchor Point
-		FVector(200, 100, 0))  // Out Handle
+    Spline.GetSpline().AppendBezierSegment(
+		FVector(150, 50, 0),   // First control point
+		FVector(150, -50, 0),  // Second control point 
+		FVector(200, 0, 0)     // End point
 	);
 
     // Add another value directly to the BSpline
diff --git a/Engine/Source/Runtime/GeometryCore/Public/Splines/PolyBezierSpline.h b/Engine/Source/Runtime/GeometryCore/Public/Splines/PolyBezierSpline.h
--- a/Engine/Source/Runtime/GeometryCore/Public/Splines/PolyBezierSpline.h
+++ b/Engine/Source/Runtime/GeometryCore/Public/Splines/PolyBezierSpline.h
@@ -17,65 +17,59 @@ namespace Private
 	GEOMETRYCORE_API bool UseFindNearestUseSplineCulling();
 }
 
-template<typename ValueType> class UE_EXPERIMENTAL(5.7, "New spline APIs are experimental.") TPolyBezierSpline;
+template<typename VALUETYPE> class UE_EXPERIMENTAL(5.7, "New spline APIs are experimental.") TPolyBezierSpline;
 	
 /* Augments B-Spline interface */
-template<typename InValueType>
+template<typename VALUETYPE>
 class TPolyBezierSpline :
-	public TBSpline<InValueType, 3>,
-	private TSelfRegisteringSpline<TPolyBezierSpline<InValueType>, InValueType>
+	public TBSpline<VALUETYPE, 3>,
+	private TSelfRegisteringSpline<TPolyBezierSpline<VALUETYPE>, VALUETYPE>
 {
-
 public:
-	
-	using ValueType = InValueType;
-	using Base = TBSpline<ValueType, 3>;
+	using Base = TBSpline<VALUETYPE, 3>;
 	using Base::Degree;
-
+	typedef typename Base::ValueType ValueType;
 	using FWindow = typename Base::FWindow;
 	using FWindowStore = typename Base::FWindowStore;
 	using FFindWindowOverride = void(*)(FWindow& /*InOut*/,
-								int32 /*SegmentIndex*/,
-								const FInterval1f& /*SegmentRange*/,
-								FWindowStore& /*Scratch*/);
+										int32 /*SegmentIndex*/,
+										const FInterval1f& /*SegmentRange*/,
+										FWindowStore& /*Scratch*/);
 	
-private:
-
-	using typename Base::FKnotSearchCache;
-	using typename Base::FValidKnotSearchParams;
-	using Base::ApplyClampedKnotsMultiplicity;
-
-public:
-
 	// Generate compile-time type ID for PolyBezier
 	DECLARE_SPLINE_TYPE_ID(
 		TEXT("PolyBezier"),
-		*TSplineValueTypeTraits<ValueType>::Name
+		*TSplineValueTypeTraits<VALUETYPE>::Name
 	);
 	
     TPolyBezierSpline() = default;
     virtual ~TPolyBezierSpline() override = default;
 	
-	/**
-	 * TPolyBezierSpline consists of an arbitrary number of anchor points.
-	 * Adjacent anchor points implicitly define a cubic bezier segment.
-	 * The cubic bezier segment defined by points A and B is <A_Anchor, A_OutHandle, B_InHandle, B_Anchor>.
-	 */
-	struct FAnchor
-	{
-		ValueType InHandle;
-		ValueType Anchor;
-		ValueType OutHandle;
-	};
-
 	TPolyBezierSpline(
-		const FAnchor& Start,
-		const FAnchor& End,
+		const ValueType& P0, 
+		const ValueType& P1, 
+		const ValueType& P2, 
+		const ValueType& P3,
 		EParameterizationPolicy Parameterization = EParameterizationPolicy::Uniform)
 	: Base()
 	{
-		AppendAnchor(Start, Parameterization);
-		AppendAnchor(End, Parameterization);
+		// Add the initial segment
+		AddBezierSegmentInternal( 0, P0, P1, P2, P3, Parameterization);
+		TPolyBezierSpline<VALUETYPE>::Reparameterize(Parameterization);
+	}
+
+    /**
+	 * Factory method to create an empty spline with a default segment 
+	 */
+	static TPolyBezierSpline CreateDefault()
+	{
+		// Create a minimal valid default segment for serialization
+		ValueType P0(0, 0, 0);
+		ValueType P1(0.33f, 0, 0);
+		ValueType P2(0.66f, 0, 0);
+		ValueType P3(1, 0, 0);
+		
+		return TPolyBezierSpline(P0, P1, P2, P3, EParameterizationPolicy::Uniform);
 	}
 
 	// Static factory methods for common shapes    
@@ -92,15 +86,11 @@ public:
         ValueType P1 = Start + Direction * 0.33f;
         ValueType P2 = Start + Direction * 0.66f;
         
-		TPolyBezierSpline<ValueType> Result(
-			FAnchor{P0, P0, P1},
-			FAnchor{P2, P3, P3},
-			EParameterizationPolicy::Uniform
-		);
-
+        TPolyBezierSpline<ValueType> Result(P0, P1, P2, P3);
+        Result.Reparameterize(EParameterizationPolicy::Uniform);
         return Result;
     }
-
+    
     // Create a circular arc with given center, radius, and angle range
     static TPolyBezierSpline<ValueType> CreateCircleArc(
         const ValueType& Center,
@@ -109,89 +99,84 @@ public:
         float EndAngle,
         int32 NumSegments = 4)
     {
-		// Ensure angles are in the right order
-		if (EndAngle < StartAngle)
-		{
-			float Temp = StartAngle;
-			StartAngle = EndAngle;
-			EndAngle = Temp;
-		}
+           // Ensure angles are in the right order
+           if (EndAngle < StartAngle)
+           {
+               float Temp = StartAngle;
+               StartAngle = EndAngle;
+               EndAngle = Temp;
+           }
            
-		// Ensure reasonable segment count
-		NumSegments = FMath::Max(1, NumSegments);
+           // Ensure reasonable segment count
+           NumSegments = FMath::Max(1, NumSegments);
            
-		// Calculate angle per segment
-		float AngleSpan = EndAngle - StartAngle;
-		float AnglePerSegment = AngleSpan / static_cast<float>(NumSegments);
+           // Calculate angle per segment
+           float AngleSpan = EndAngle - StartAngle;
+           float AnglePerSegment = AngleSpan / static_cast<float>(NumSegments);
            
-		// For large angles per segment, increase segment count for better approximation
-		if (AnglePerSegment > UE_PI/4)
-		{
-			NumSegments = FMath::CeilToInt(AngleSpan / (UE_PI/4));
-			AnglePerSegment = AngleSpan / static_cast<float>(NumSegments);
-		}
+           // For large angles per segment, increase segment count for better approximation
+           if (AnglePerSegment > UE_PI/4)
+           {
+               NumSegments = FMath::CeilToInt(AngleSpan / (UE_PI/4));
+               AnglePerSegment = AngleSpan / static_cast<float>(NumSegments);
+           }
            
-		// Pre-compute all points along the arc for exact positioning
-		TArray<ValueType> Points;
-		Points.SetNum(NumSegments + 1);
+           // Pre-compute all points along the arc for exact positioning
+           TArray<ValueType> Points;
+           Points.SetNum(NumSegments + 1);
            
-		for (int32 i = 0; i <= NumSegments; i++)
-		{
-			float Angle = StartAngle + (static_cast<float>(i) * AnglePerSegment);
-			Points[i] = Center + ValueType(
-				Radius * FMath::Cos(Angle),
-				Radius * FMath::Sin(Angle),
-				0);
-		}
+           for (int32 i = 0; i <= NumSegments; i++)
+           {
+               float Angle = StartAngle + (static_cast<float>(i) * AnglePerSegment);
+               Points[i] = Center + ValueType(
+                   Radius * FMath::Cos(Angle),
+                   Radius * FMath::Sin(Angle),
+                   0);
+           }
            
-		// Calculate first segment with precise tangents
-		ValueType P0 = Points[0];
-		ValueType P3 = Points[1];
+           // Calculate first segment with precise tangents
+           ValueType P0 = Points[0];
+           ValueType P3 = Points[1];
            
-		// Get normalized direction vectors from center to points
-		ValueType Dir0 = Math::GetSafeNormal(P0 - Center);
-		ValueType Dir3 = Math::GetSafeNormal(P3 - Center);
+           // Get normalized direction vectors from center to points
+           ValueType Dir0 = Math::GetSafeNormal(P0 - Center);
+           ValueType Dir3 = Math::GetSafeNormal(P3 - Center);
            
-		// Calculate perpendicular vectors (normalized)
-		ValueType Perp0 = ValueType(-Dir0.Y, Dir0.X, 0);
-		ValueType Perp3 = ValueType(-Dir3.Y, Dir3.X, 0);
+           // Calculate perpendicular vectors (normalized)
+           ValueType Perp0 = ValueType(-Dir0.Y, Dir0.X, 0);
+           ValueType Perp3 = ValueType(-Dir3.Y, Dir3.X, 0);
            
-		// Calculate precise tangent scale
-		float TangentScale = Radius * (4.0f / 3.0f) * 
-									FMath::Tan(AnglePerSegment / 4.0f);
+           // Calculate precise tangent scale
+           float TangentScale = Radius * (4.0f / 3.0f) * 
+                                    FMath::Tan(AnglePerSegment / 4.0f);
            
-		// Create control points with exact perpendicular tangents
-		ValueType P1 = P0 + Perp0 * TangentScale;
-		ValueType P2 = P3 - Perp3 * TangentScale;
-
-		TPolyBezierSpline<ValueType> Result(
-			FAnchor{P0, P0, P1},
-			FAnchor{P2, P3, P3},
-			EParameterizationPolicy::Uniform
-		);
+           // Create control points with exact perpendicular tangents
+           ValueType P1 = P0 + Perp0 * TangentScale;
+           ValueType P2 = P3 - Perp3 * TangentScale;
            
-		// Add remaining segments with carefully calculated control points
-		for (int32 i = 1; i < NumSegments; i++)
-		{
-			P0 = Points[i];
-			P3 = Points[i+1];
+           // Create spline with first segment
+           TPolyBezierSpline<ValueType> Result(P0, P1, P2, P3);
+           
+           // Add remaining segments with carefully calculated control points
+           for (int32 i = 1; i < NumSegments; i++)
+           {
+               P0 = Points[i];
+               P3 = Points[i+1];
                
-			// Calculate exact perpendicular vectors for this segment
-			Dir0 = Math::GetSafeNormal(P0 - Center);
-			Dir3 = Math::GetSafeNormal(P3 - Center);
-			Perp0 = ValueType(-Dir0.Y, Dir0.X, 0);
-			Perp3 = ValueType(-Dir3.Y, Dir3.X, 0);
+               // Calculate exact perpendicular vectors for this segment
+               Dir0 = Math::GetSafeNormal(P0 - Center);
+               Dir3 = Math::GetSafeNormal(P3 - Center);
+               Perp0 = ValueType(-Dir0.Y, Dir0.X, 0);
+               Perp3 = ValueType(-Dir3.Y, Dir3.X, 0);
                
-			// Create control points with exact perpendicular tangents
-			P1 = P0 + Perp0 * TangentScale;
-			P2 = P3 - Perp3 * TangentScale;
+               // Create control points with exact perpendicular tangents
+               P1 = P0 + Perp0 * TangentScale;
+               P2 = P3 - Perp3 * TangentScale;
                
-			const int32 NewAnchorIdx = Result.AppendAnchor(FAnchor { P2, P3, P3 }, EParameterizationPolicy::Uniform);
-			const int32 OldAnchorIdx = NewAnchorIdx - 1;
-			Result.SetHandleOutOfAnchor(OldAnchorIdx, P1);
-		}
+               Result.AppendBezierSegment(P1, P2, P3);
+           }
            
-		return Result;
+           return Result;
     }
     
     // Create a complete circle
@@ -250,11 +235,7 @@ public:
 	    ValueType P2 = P3 - Tangent3;
 	    
 	    // Create the spline with the first segment
-		TPolyBezierSpline<ValueType> Result(
-			FAnchor{P0, P0, P1},
-			FAnchor{P2, P3, P3},
-			EParameterizationPolicy::Uniform
-		);
+	    TPolyBezierSpline<ValueType> Result(P0, P1, P2, P3);
 	    
 	    // Add remaining segments
 	    for (int32 i = 1; i < NumSegments; ++i)
@@ -276,9 +257,7 @@ public:
 	        P2 = P3 - Tangent3;
 	        
 	        // Add the segment
-			const int32 NewAnchorIdx = Result.AppendAnchor(P2, P3, P3);
-			const int32 OldAnchorIdx = NewAnchorIdx - 1;
-			Result.SetHandleOutOfAnchor(OldAnchorIdx, P1);
+	        Result.AppendBezierSegment(P1, P2, P3);
 	    }
 	    
 	    // Set as closed loop
@@ -319,13 +298,13 @@ public:
         TUniquePtr<TPolyBezierSpline> Clone = MakeUnique<TPolyBezierSpline>();
         
         // Copy base class members
-        Clone->Values = Base::Values;
-        Clone->PairKnots = Base::PairKnots;
-        Clone->bIsClosedLoop = Base::bIsClosedLoop;
+        Clone->Values = this->Values;
+        Clone->PairKnots = this->PairKnots;
+        Clone->bIsClosedLoop = this->bIsClosedLoop;
         
         // Copy infinity modes
-        Clone->PreInfinityMode = Base::PreInfinityMode;
-        Clone->PostInfinityMode = Base::PostInfinityMode;
+        Clone->PreInfinityMode = this->PreInfinityMode;
+        Clone->PostInfinityMode = this->PostInfinityMode;
         
         return Clone;
     }
@@ -338,6 +317,7 @@ public:
     		return 0.0f;
     	}
 
+
 		// Get Bezier control points for this segment
 		TStaticArray<ValueType, 4> Coeffs = {};
 
@@ -487,10 +467,9 @@ public:
 		FWindow Window;
 		FWindowStore Scratch;
 		FindWindow(Parameter, Window, Scratch);
-
 		// If FindWindow fails, it will return an array of nullptr. InterpolateWindow assumes validity of elements.
-		if (!Math::FSplineValidation::IsValidWindow<ValueType, Base::WindowSize>(Window) ||
-			!Math::FSplineValidation::IsValidParameter(Parameter))
+		if (!Math::FSplineValidation::IsValidWindow<ValueType, Base::WindowSize>(Window) || // Cubic Bezier needs 4 points
+				!Math::FSplineValidation::IsValidParameter(Parameter))
 		{
 			return Math::FSplineValidation::GetDefaultValue<ValueType>();
 		}
@@ -502,10 +481,9 @@ public:
 	    }
     
     	// Convert global parameter to segment information
-		float LocalT;
-		int32 SegmentIndex = FindIndexForParameter(Parameter, LocalT);
-
-		if (SegmentIndex == INDEX_NONE)
+    	int32 SegmentIndex;
+    	float LocalT;
+		if (!MapGlobalParameterToLocalSegment(Parameter, SegmentIndex, LocalT))
 		{
 			return Math::FSplineValidation::GetDefaultValue<ValueType>();
 		}
@@ -520,479 +498,692 @@ public:
 	/**
 	 * MANIPULATION METHODS
 	 */
-	
-	/** Returns the number of Bezier segments in the spline */
-	virtual int32 GetNumberOfSegments() const override
-	{
-		const int32 NumAnchors = GetNumAnchorPoints();
-		return HasClosingSegment() ? NumAnchors : FMath::Max(0, NumAnchors - 1);
-	}
 
 	/**
-	 * Maps a segment index to its parameter range
-	 * @param SegmentIndex - Index of the segment (0-based)
-	 * @return True if the segment index is valid and mapping succeeded
+	 * Replaces all control points in the spline with the provided points.
+	 * For a cubic Bezier spline, the number of points must be a multiple of 4 (or 4N-1 for open splines).
+	 * 
+	 * @param Points - Array of new control points
+	 * @param ParameterizationPolicy - How to distribute Knots for the new segments
+	 * @return True if the points were successfully set
 	 */
-	virtual FInterval1f GetSegmentParameterRange(int32 SegmentIndex) const override
+	bool SetControlPoints(
+		const TArray<ValueType>& Points,
+		EParameterizationPolicy ParameterizationPolicy)
 	{
-		FInterval1f SegmentRange;
+	   const int32 NumPoints = Points.Num();
 
-		const int32 NumSegments = GetNumberOfSegments();
-		if (SegmentIndex < 0 || SegmentIndex >= NumSegments || Base::PairKnots.Num() == 0 || !Base::PairKnots.IsValidIndex(SegmentIndex))
-		{
-			return SegmentRange;
-		}
+	    // We expect Bezier quartets as input: P0,P1,P2,P3 per segment.
+	    if (NumPoints < 4 || (NumPoints % 4) != 0)
+	    {
+	        UE_LOG(LogSpline, Warning,
+	            TEXT("SetControlPoints expects 4*k Bezier control points (got %d)."), NumPoints);
+	        return false;
+	    }
 
-		SegmentRange.Min = Base::PairKnots[SegmentIndex].Value;
+	    // Remember desired closed state, then reset storage.
+	    const bool bWasClosed = Base::IsClosedLoop();
+	    Clear();
 
-		if (!HasClosingSegment())
-		{
-			// Open: knots are standard boundaries, so we still use i+1.
-			if (SegmentIndex + 1 < Base::PairKnots.Num())
-			{
-				SegmentRange.Max = Base::PairKnots[SegmentIndex + 1].Value;
-			}
-			return SegmentRange;
-		}
+	    const int32 NumSegments = NumPoints / 4;
+		
+	    // First segment bootstraps two anchors
+	    {
+	        const ValueType& P0 = Points[0];
+	        const ValueType& P1 = Points[1];
+	        const ValueType& P2 = Points[2];
+	        const ValueType& P3 = Points[3];
+
+	        // A0
+	        AppendAnchor(/*In0*/ P0, /*A0*/ P0, /*Out0*/ P1);
+	        // A1
+	        AppendAnchor(/*In1*/ P2, /*A1*/ P3, /*Out1*/ P3);
+	    }
 
-		// Closed: PairKnots holds starts of each segment; the final segment
-		// closes the loop at Min + LoopKnotDelta.
-		if (SegmentIndex < NumSegments - 1)
-		{
-			SegmentRange.Max = Base::PairKnots[SegmentIndex + 1].Value;
-		}
-		else
-		{
-			SegmentRange.Max = SegmentRange.Min + LoopKnotDelta;
-		}
+	    // Remaining segments
+	    for (int32 s = 1; s < NumSegments; ++s)
+	    {
+	        const int32 b = 4 * s;
+	        const ValueType& P0 = Points[b + 0];
+	        const ValueType& P1 = Points[b + 1];
+	        const ValueType& P2 = Points[b + 2];
+	        const ValueType& P3 = Points[b + 3];
+
+	        // Optionally warn if P0 doesn't match previous A (only a sanity message).
+	        const ValueType& PrevA = this->Values[AnchorBase(s) + 1];
+	        if (Math::SizeSquared(P0 - PrevA) > UE_KINDA_SMALL_NUMBER)
+	        {
+	            UE_LOG(LogSpline, Verbose,
+	                TEXT("SetControlPoints: segment %d P0 differs from previous A (len^2=%g)."),
+	                s, (double)Math::SizeSquared(P0 - PrevA));
+	        }
+
+	        // Out_s = P1
+	        this->SetValue(AnchorBase(s) + 2, P1);
+
+	        // Append A_{s+1} = P3 with In_{s+1}=P2, Out_{s+1}=A_{s+1} (placeholder until another seg arrives)
+	        AppendAnchor(/*In*/ P2, /*A*/ P3, /*Out*/ P3);
+	    }
 
-		return SegmentRange;
+	    // Apply desired loop state:
+	    // - We DO NOT materialize a physical closing quartet anymore.
+	    // - Closing segment is virtual: {A_last, Out_last, In_0, A_0}.
+	    SetClosedLoop(bWasClosed);
+
+	    // Build / rebuild knot vector once from scratch using the policy
+	    // (You already do this today; keep it — it re-spaces and sets multiplicities.)
+	    Reparameterize(ParameterizationPolicy);
+
+	    return true;
 	}
 
-	// Parameterization methods
 	/**
-	 * Reparameterizes the spline based on the provided points and mode
-	 * @param Mode - New parameterization mode
-	 * @param Points - Array of control points to use for reparameterization
-	 * @return 
+	 * Adds multiple Bezier segments to the spline.
+	 * 
+	 * @param Storage - Storage to add the segments to
+	 * @param Points - Array of control points defining the new segments
+	 * @param bAppend - If true, append to the end of spline; if false, prepend to the beginning
+	 * @param ParameterizationPolicy - How to distribute Knots for the new segments
+	 * @return True if the segments were successfully added
 	 */
-	virtual void Reparameterize(EParameterizationPolicy Mode = EParameterizationPolicy::Centripetal) override
-	{
-		const TArray<ValueType>& Points = Base::Values;
-		const int32 NumValues = Points.Num();
-
-		GenerateDistanceKnotsForBezier(Mode);
-    	
-		UE_LOG(LogSpline, Verbose, TEXT("Set knot vector (Reparameterize) - Points: %d, Knots: %d, Mode: %d"),
-					   NumValues,  Base::PairKnots.Num(), static_cast<int32>(Mode));
-		
-		Base::PrintKnotVector();
-	}
-	
-	virtual FInterval1f GetParameterSpace() const override
+	bool AddBezierSegments(
+		const TArray<ValueType>& Points,
+		bool bAppend,
+		EParameterizationPolicy ParameterizationPolicy)
 	{
-		if (Base::PairKnots.Num() == 0)
+	    const int32 NumPoints = Points.Num();
+	    
+		// Need at least 2 points to add a valid segment
+		if (NumPoints < 4)
 		{
-			return FInterval1f::Empty();
+			UE_LOG(LogSpline, Warning, TEXT("AddBezierSegments requires at least 4 points to add a valid segment. Got %d points."), NumPoints);
+			return false;
 		}
+	    
+	    // If the spline is currently empty (even though it shouldn't be according to design), this is equivalent to SetControlPoints
+	    if (this->Base::NumKeys() == 0)
+	    {
+	        return SetControlPoints(Points, ParameterizationPolicy);
+	    }
+	    
+	    // Check points format based on append/prepend direction
+	    if (bAppend)
+	    {
+	        // For appending, we expect 3N+1 points: P0 (shared with last existing point), then (P1,P2,P3) per segment
+	        if (NumPoints % 3 != 1)
+	        {
+	            UE_LOG(LogSpline, Warning, TEXT("Invalid point count for appending segments. Expected 3N+1, got %d."), NumPoints);
+	            return false;
+	        }
+	    }
+	    else
+	    {
+	        // For prepending, we expect 3N+1 points: (P0,P1,P2) per segment, then P3 (shared with first existing point)
+	        if (NumPoints % 3 != 1)
+	        {
+	            UE_LOG(LogSpline, Warning, TEXT("Invalid point count for prepending segments. Expected 3N+1, got %d."), NumPoints);
+	            return false;
+	        }
+	    }
+		const int32 NumAnchors = GetNumAnchorPoints();
 
-		const float Min = Base::PairKnots[0].Value;
-		const float LastKnot = Base::PairKnots.Last().Value;
+	    // Check connection point
+	    if (bAppend)
+	    {
+	        // First point of new segments should match the last existing anchor
+	        const ValueType& FirstNewPoint = Points[0];
+	    	const ValueType& LastExistingAnchor = this->Values[AnchorBase(NumAnchors - 1) + 1];
+	        
+	        // Check if the connection point is reasonably close
+	        const double DistanceSquared = Math::SizeSquared(LastExistingAnchor, FirstNewPoint);
+	        if (DistanceSquared > UE_KINDA_SMALL_NUMBER)
+	        {
+	            UE_LOG(LogSpline, Warning, TEXT("Connection point mismatch when appending segments. Distance: %f"),
+	                FMath::Sqrt(static_cast<float>(DistanceSquared)));
+	            // Continue anyway but log warning
+	        }
+	    }
+	    else
+	    {
+	        // Last point of new segments should match the first point of existing spline
+	        const ValueType& LastNewPoint = Points.Last();
+	    	const ValueType& FirstExistingAnchor = this->Values[AnchorBase(0) + 1];
+	        
+	        // Check if the connection point is reasonably close
+	        const double DistanceSquared = Math::SizeSquared(LastNewPoint, FirstExistingAnchor);
+	        if (DistanceSquared > UE_KINDA_SMALL_NUMBER)
+	        {
+	            UE_LOG(LogSpline, Warning, TEXT("Connection point mismatch when prepending segments. Distance: %f"),
+	                FMath::Sqrt(static_cast<float>(DistanceSquared)));
+	            // Continue anyway but log warning
+	        }
+	    }
+	    
+	    if (bAppend)
+	    {
+	        // Append segments to the end of the spline
+	        const int32 NumSegmentsToAdd = (NumPoints - 1) / 3;
+	        
+	        // Add each segment
+	        for (int32 SegmentIndex = 0; SegmentIndex < NumSegmentsToAdd; ++SegmentIndex)
+	        {
+	            const int32 BaseIndex = 1 + SegmentIndex * 3; // Skip the first point (P0) for first segment
+	            
+	            // For each segment, we need P1, P2, P3 (P0 is the last point of existing spline)
+	            const ValueType& P1 = Points[BaseIndex];
+	            const ValueType& P2 = Points[BaseIndex + 1];
+	            const ValueType& P3 = Points[BaseIndex + 2];
+	            
+	            // Use the existing AppendBezierSegment method
+	            this->AppendBezierSegment(P1, P2, P3, ParameterizationPolicy);
+	        }
+	    }
+	    else
+	    {
+	        // Prepend segments to the beginning of the spline
+	        const int32 NumSegmentsToPrepend = (NumPoints - 1) / 3;
+	        
+	        // Prepend segments in order starting from the first segment
+	        // For prepending, we need to work backwards since each PrependBezierSegment
+	        // adds to the beginning
+	        for (int32 i = NumSegmentsToPrepend - 1; i >= 0; --i)
+	        {
+	            const int32 BaseIndex = i * 3;
+	            
+	            // For each segment, we need P0, P1, P2 (P3 will connect to the existing spline)
+	            const ValueType& P0 = Points[BaseIndex];
+	            const ValueType& P1 = Points[BaseIndex + 1];
+	            const ValueType& P2 = Points[BaseIndex + 2];
+	            
+	            // Use the existing PrependBezierSegment method
+	            this->PrependBezierSegment(P0, P1, P2, ParameterizationPolicy);
+	        }
+	    }
+	    
+	    // Reparameterize the spline to ensure consistent parameterization
+    	Reparameterize(ParameterizationPolicy);
+	    
+	    return true;
+	}
+    /**
+     * Appends a Bezier segment using only 3 points (P1, P2, P3)
+     * The start point (P0) is automatically inferred from the last point of the previous segment.
+     * 
+     * @param P1 First control point
+     * @param P2 Second control point
+     * @param P3 End point
+     * @param ParameterizationPolicy How to distribute Knots
+     * @return Index of the new segment
+     */
+    int32 AppendBezierSegment(
+        const ValueType& P1, 
+        const ValueType& P2, 
+        const ValueType& P3, 
+        EParameterizationPolicy ParameterizationPolicy = EParameterizationPolicy::Centripetal)
+    {
+		const int32 NumAnchors = GetNumAnchorPoints();
+		const int32 NumSegments = GetNumberOfSegments();
+		
+		// P0 = last anchor position A_{n-1}
+		const int32 LastAnchor = NumAnchors - 1;
+		const ValueType& P0 = this->Values[AnchorBase(LastAnchor) + 1];
 
-		if (!HasClosingSegment())
-		{
-			return FInterval1f(Min, LastKnot);
-		}
+		// Insert new segment after the current last one
+		//   SegmentIndex == NumSegments means "append at tail"
+		return AddBezierSegmentInternal(NumSegments, P0, P1, P2, P3, ParameterizationPolicy, true);
+    }
 
-		// Closed: closing boundary is last real knot + LoopKnotDelta
-		return FInterval1f(Min, LastKnot + LoopKnotDelta);
+    /**
+     * Prepends a Bezier segment using only 3 points (P0, P1, P2)
+     * The end point (P3) is automatically inferred from the first point of the existing spline.
+     * 
+     * @param P0 First control point
+     * @param P1 Second control point
+     * @param P2 Third control point
+     * @param ParameterizationPolicy How to distribute Knots
+     * @return Index of the new segment
+     */
+    int32 PrependBezierSegment(
+		const ValueType& P0, 
+		const ValueType& P1, 
+		const ValueType& P2, 
+		EParameterizationPolicy ParameterizationPolicy = EParameterizationPolicy::Centripetal)
+	{
+		const int32 NumAnchors = GetNumAnchorPoints();
+		const int32 NumSegments = GetNumberOfSegments();
+		
+		// P3 = first anchor position A_0
+		const ValueType& P3 = this->Values[AnchorBase(0) + 1];
+
+		// Insert new segment before segment 0 / at the head
+		return AddBezierSegmentInternal(0, P0, P1, P2, P3,
+										ParameterizationPolicy,
+										false); // Prepend to first segment
 	}
 
 	/**
-	 * Gets the parameter value for a specific anchor point or the virtual closing point.
-	 * The virtual closing point exists when there are at least 2 anchor points and the spline is closed. The index is GetNumAnchorPoints().
-	 * @param Index - Index of the anchor point
-	 * @return Parameter value at the specified point
+	 * Inserts a point at the closest location on a segment to the given position.
+	 * 
+	 * @param SegmentIndex Index of the segment to split
+	 * @param Position Position for the new point
+	 * @param ParameterizationPolicy How to distribute Knots
+	 * @return The index of the newly created point
 	 */
-	virtual float GetParameter(int32 Index) const override
+	int32 InsertPointAtPosition(
+		int32 SegmentIndex,
+		const ValueType& Position,
+		EParameterizationPolicy ParameterizationPolicy = EParameterizationPolicy::Centripetal)
 	{
-		const int32 NumAnchors = GetNumAnchorPoints();
+		// Find the closest point on the segment to the given position
+		float SquaredDistance;
+		float LocalT = FindLocalParameterNearestToPosition(SegmentIndex, Position, SquaredDistance);
+            
+		// Insert at that local parameter
+		return this->InsertPointAtSegmentParam(SegmentIndex, LocalT, Position, ParameterizationPolicy);
+	}
 
-		if (NumAnchors <= 0)
+	/**
+	 * Inserts a point at a specific segment with local parameter.
+	 * This splits the segment at the specified location and places the new point
+	 * at the provided position, which may change the spline shape.
+	 * 
+	 * @param SegmentIndex Index of the segment to split
+	 * @param LocalT Local parameter [0,1] within the segment
+	 * @param Position Position for the new point
+	 * @param ParameterizationPolicy How to distribute Knots
+	 * @return The index of the newly created point
+	 */
+	int32 InsertPointAtSegmentParam(
+	    int32 SegmentIndex,
+	    float LocalT,
+	    const ValueType& Position,
+	    EParameterizationPolicy ParameterizationPolicy = EParameterizationPolicy::Centripetal)
+	{
+    	int32 NumSegments = GetNumberOfSegments();
+	    // Validate inputs
+	    if (NumSegments == 0)
+	    {
+	        // Special case for empty spline
+	        int32 NewPointIndex;
+	        if (InitializeFirstSegment(Position, NewPointIndex))
+	            return NewPointIndex;
+	    }
+	    
+	    SegmentIndex = FMath::Clamp(SegmentIndex, 0, NumSegments - 1);
+	    
+	    // Get the segment's parameter range
+	    FInterval1f SegmentRange = GetSegmentParameterRange(SegmentIndex);
+	    
+	    // Convert local parameter to global parameter
+	    float Parameter = SegmentRange.Min + LocalT * (SegmentRange.Max - SegmentRange.Min);
+	    
+	    return InsertPoint(Parameter, Position, ParameterizationPolicy);
+	}
+		
+	/**
+	 * Inserts a point at the specified parameter along the spline.
+	 * This splits the segment at the parameter value and places the new point
+	 * at the specified position, which may change the spline shape.
+	 * 
+	 * @param Storage - Storage to add the point to
+	 * @param Parameter Parameter value where to split the spline
+	 * @param Position Position for the new point (doesn't need to be on the original curve)
+	 * @param ParameterizationPolicy How to distribute Knots
+	 * @return The index of the newly created point
+	 */
+    int32 InsertPoint(
+        float Parameter, 
+        const ValueType& Position,
+        EParameterizationPolicy ParameterizationPolicy = EParameterizationPolicy::Centripetal)
+    {
+	    if (int32 NewPointIndex;
+	    	InitializeFirstSegment(Position, NewPointIndex))
 		{
-			// todo: fail
-			return 0.f;
+			return NewPointIndex;
 		}
 
-		// Getting the virtual closing knot.
-		if (Index == NumAnchors && HasClosingSegment())
+		typename Base::FKnotSearchCache SearchCache{ this->PairKnots };
+
+		// Make sure we have a valid parameter
+		Parameter = Base::GetNearestAvailableKnotValue(Parameter, SearchCache);
+		
+        FInterval1f ParameterRange = GetParameterSpace();
+		
+		// Handle out-of-range Knots
+		if (Parameter <= ParameterRange.Min)
 		{
-			return Base::PairKnots.Last().Value + LoopKnotDelta;
+			// Prepend
+			return PrependPoint(Position, ParameterizationPolicy);
 		}
-
-		if (Index < 0 || Index >= NumAnchors)
+    
+		if (Parameter >= ParameterRange.Max)
 		{
-			// todo: fail
-			return 0.f;
+			return AppendPoint(Position, ParameterizationPolicy);
 		}
 
-		return Base::PairKnots[Index].Value;
-	}
-
-	/**
-	 * Sets the parameter value for a specific anchor control point or the virtual closing point.
-	 * This function will reorder the knot vector to ensure that it is always monotonically increasing.
-	 * If the knot vector is reordered, anchor points are reordered such that they are always associated with the same knot.
-	 * Setting the parameter of the virtual closing point will not cause any reordering. If reordering would occur, the function fails.
-	 * @param Index - Index of the anchor point
-	 * @param NewParameter - New parameter value to set
-	 * @return The index of the point that was actually updated, or INDEX_NONE if failed
-	 */ 
-	virtual int32 SetParameter(int32 Index, float NewParameter) override
-	{
-		const int32 NumAnchors = GetNumAnchorPoints();
-		const int32 NumSegments = GetNumberOfSegments();
-
-		// Cannot set parameters for empty splines, fail.
-		if (NumAnchors <= 0)
+		int32 SegmentIndex;
+		float LocalT;
+		if (!MapGlobalParameterToLocalSegment(Parameter, SegmentIndex, LocalT))
 		{
 			return INDEX_NONE;
 		}
 
-		// Trivial case, no possible reordering or numerical issues.
-		if (NumAnchors == 1)
+		FWindow Window;
+		if (!BuildRawWindowForSegment(SegmentIndex, Window))
 		{
-			if (Index == 0)
-			{
-				Base::PairKnots[Index].Value = NewParameter;
-				return 0;
-			}
-			else
-			{
-				return INDEX_NONE;
-			}
+			return INDEX_NONE;
 		}
+		ValueType P0 = *Window[0];
+		ValueType P1 = *Window[1];
+		ValueType P2 = *Window[2];
+		ValueType P3 = *Window[3];
+		
+        // Use De Casteljau algorithm to split the curve
+        const float t = LocalT;
+        const float mt = 1.0f - t;
+        
+        // Calculate split control points
+        ValueType Q0 = P0;
+        ValueType Q1 = P0 * mt + P1 * t;
+        ValueType Q2 = (P0 * mt + P1 * t) * mt + (P1 * mt + P2 * t) * t;
+        ValueType Q3 = Position;  // Use provided position for the split point
+        
+        ValueType R0 = Position;  // Use provided position for the split point
+        ValueType R1 = (P1 * mt + P2 * t) * mt + (P2 * mt + P3 * t) * t;
+        ValueType R2 = P2 * mt + P3 * t;
+        ValueType R3 = P3;
+        
+        FInterval1f SegmentRange = GetSegmentParameterRange(SegmentIndex);
+		float TargetSplitParam = SegmentRange.Interpolate(LocalT);
 
-		// Setting the virtual closing knot.
-		if (Index == NumAnchors && HasClosingSegment())
+		if (TargetSplitParam == SegmentRange.Min)
 		{
-			// No reordering when manipulating the closing knot, fail.
-			if (NewParameter <= Base::PairKnots.Last().Value)
-			{
-				return INDEX_NONE;
-			}
-
-			LoopKnotDelta = NewParameter - Base::PairKnots.Last().Value;
-			return Index;
+			TargetSplitParam = Param::NextDistinct(TargetSplitParam);
 		}
 
-		// Because we special case the closing point above, we can trivially bounds check here.
-		if (Index < 0 || Index >= NumAnchors)
+		if (TargetSplitParam == SegmentRange.Max)
+		{
+			TargetSplitParam = Param::PrevDistinct(TargetSplitParam);
+		}
+		
+		// Bail out if we lost usable range
+		if (TargetSplitParam <= SegmentRange.Min || TargetSplitParam >= SegmentRange.Max)
 		{
 			return INDEX_NONE;
 		}
 
-		const int32 SegmentIndex = Index == 0 ? 0 : Index - 1;
-		const int32 LocalPointIndex = Index == 0 ? 0 : 3;
-
-		float OldParameter = Base::PairKnots[Index].Value;
-    
-		// Early exit if parameter isn't changing
-		if (OldParameter == NewParameter)
+		// --- Anchor/handle surgery in 3-tuple layout ---
+		const int32 NumAnchors = GetNumAnchorPoints();
+		if (NumAnchors < 2)
 		{
-			return Index;
+			return INDEX_NONE;
 		}
-		
-		UE_LOG(LogSpline, Verbose, TEXT("\t"));
-		UE_LOG(LogSpline, Verbose, TEXT("Before SetParameter(%d, %f):"), Index, NewParameter);
-		Base::Dump();
-		
-		// Prevent reuse of an existing knot.
-		const float DesiredNewParameter = NewParameter;
-		FKnotSearchCache SearchCache{ Base::PairKnots };
-		const float ActualNewParameter = Base::GetNearestAvailableKnotValue(DesiredNewParameter, SearchCache);
 
-		const float DesiredToOldDistance = FMath::Abs(DesiredNewParameter - OldParameter);
-		const float DesiredToActualNewDistance = FMath::Abs(DesiredNewParameter - ActualNewParameter);
+		// For segment s in an *open* spline, its anchors are A_s and A_{s+1}.
+		// For closed, the closing segment (s == A-1) is between A_{A-1} and A_0.
+		const bool bClosed = IsActuallyClosed();
+		const int32 LeftAnchor = SegmentIndex;
+		const int32 RightAnchor =
+			(!bClosed)
+			? (SegmentIndex + 1)
+			: ((SegmentIndex + 1) % NumAnchors);
 
-		if (DesiredToActualNewDistance > DesiredToOldDistance)
-		{
-			// If our nearest valid knot is further away from the desired value than our current value is, we need to just keep our current value
-			return Index;
-		}
-		
-		int32 ReturnIndex = Index;
-		Base::SetKnot(Index, NewParameter);
+		// 1) Left side: Out_left = Q1
+		SetHandleOutOfAnchor(LeftAnchor, Q1);
 
-		// Fix the knot vector by flipping segments which have inverted parameter ranges until it is valid again.
-		if (NewParameter < OldParameter)
-		{
-			int32 LeftSegmentIndex = Index - 1;
+		// 2) Insert new anchor between left and right with (In_new, A_new, Out_new) = (Q2, Q3, R1)
+		//    For open, this is simply at index RightAnchor.
+		//    For closed, inserting at index RightAnchor still produces a valid cyclic order;
+		//    the notion of "start anchor" is arbitrary on a loop.
+		InsertAnchorAt(RightAnchor, /*In*/ Q2, /*A*/ Q3, /*Out*/ R1);
 
-			// Repeated left flips
-			while (LeftSegmentIndex >= 0 && GetSegmentParameterRange(LeftSegmentIndex).Min > GetSegmentParameterRange(LeftSegmentIndex).Max)
-			{
-				FlipSegment(LeftSegmentIndex);
-				LeftSegmentIndex--;
-				ReturnIndex--;
-			}
-		}
-		else if (NewParameter > OldParameter)
+		// 3) Right side: old right anchor has shifted by +1 after the insert
+		const int32 OldRightAnchorIndex =
+			(!bClosed)
+			? (RightAnchor + 1)
+			: ((RightAnchor + 1) % (NumAnchors + 1)); // one more anchor now in closed loop
+
+		SetHandleInOfAnchor(OldRightAnchorIndex, R2);
+
+		// --- Insert a knot at the split location (incremental) ---
+
+		this->InsertKnot(FKnot(TargetSplitParam, Degree));
+
+		this->ApplyClampedKnotsMultiplicity();
+		ApplyInternalKnotsMultiplicity();
+
+		if (bClosed)
 		{
-			int32 RightSegmentIndex = Index;
+			RefreshClosingKnot(ParameterizationPolicy); // optional helper to ensure closing boundary is consistent
+		}
 
-			while (RightSegmentIndex <= NumSegments - 1 && !IsClosingSegment(RightSegmentIndex) && GetSegmentParameterRange(RightSegmentIndex).Min > GetSegmentParameterRange(RightSegmentIndex).Max)
-			{
-				FlipSegment(RightSegmentIndex);
-				RightSegmentIndex++;
-				ReturnIndex++;
-			}
-		}
-			
-		UE_LOG(LogSpline, Verbose, TEXT("\t"));
-		UE_LOG(LogSpline, Verbose, TEXT("After SetParameter(%d, %f):"), Index, NewParameter);
-		Base::Dump();
-		
-		return ReturnIndex;
-	}
-	
-	void FlipSegment(int32 Segment)
-	{
-		const int32 NumSegments = GetNumberOfSegments();
-		const int32 NumAnchors  = GetNumAnchorPoints();
-		if (Segment < 0 || Segment >= NumSegments || NumAnchors < 2)
-		{
-			return;
-		}
-
-		UE_LOG(LogSpline, Verbose, TEXT("\t"));
-		UE_LOG(LogSpline, Verbose, TEXT("Before FlipSegment(%d):"), Segment);
-		Base::Dump();
-		
-		// 1) Flip this segment's Bezier window: P0<->P3, P1<->P2
-		{
-			int32 Idx0, Idx1, Idx2, Idx3;
-			GetBezierIndices(Segment, Idx0, Idx1, Idx2, Idx3);
-
-			ValueType P0 = Base::GetValue(Idx0);
-			ValueType P1 = Base::GetValue(Idx1);
-			ValueType P2 = Base::GetValue(Idx2);
-			ValueType P3 = Base::GetValue(Idx3);
-
-			Base::SetValue(Idx0, P3);
-			Base::SetValue(Idx1, P2);
-			Base::SetValue(Idx2, P1);
-			Base::SetValue(Idx3, P0);
-		}
+		// The new anchor is at index RightAnchor
+		return RightAnchor;
+    }
 
-		// 2) Fix left neighbor (Segment-1): keep its out-tangent but move its end point to match
-	    if (Segment > 0 || (HasClosingSegment() && Segment == 0))
+	/**
+	* Removes a Bezier segment from the spline
+	*
+	* @param Storage - Storage to remove the segment from
+	* @param SegmentIndex - Segment index to remove
+	* @return true if successfully removed
+	*/
+    bool RemoveSegment(const int32 SegmentIndex)
+    {
+    	const int32 NumSegments = GetNumberOfSegments();
+	    if (SegmentIndex < 0 || SegmentIndex >= NumSegments)
 	    {
-	    	const int32 LeftSeg =
-			 (Segment > 0) ? (Segment - 1) : (NumSegments - 1); // wrap in closed case
-
-	    	int32 ThisIdx0, ThisIdx1, ThisIdx2, ThisIdx3;
-	    	int32 LeftIdx0, LeftIdx1, LeftIdx2, LeftIdx3;
-
-	    	GetBezierIndices(Segment, ThisIdx0, ThisIdx1, ThisIdx2, ThisIdx3);
-	    	GetBezierIndices(LeftSeg, LeftIdx0, LeftIdx1, LeftIdx2, LeftIdx3);
-
-	    	const ValueType ThisP0 = Base::GetValue(ThisIdx0);
-	    	const ValueType LeftP2 = Base::GetValue(LeftIdx2);
-	    	const ValueType LeftP3 = Base::GetValue(LeftIdx3);
-
-	    	const ValueType Delta = LeftP3 - LeftP2; // original out-tangent
-	    	const ValueType NewEnd = ThisP0;
-	    	const ValueType NewIn  = NewEnd - Delta;
-
-	    	Base::SetValue(LeftIdx3, NewEnd);
-	    	Base::SetValue(LeftIdx2, NewIn);
+	        return false;
 	    }
 
-	    // 3) Fix right neighbor (Segment+1): keep its in-tangent but move its start point to match
-	    if (Segment < NumSegments - 1 || IsClosingSegment(Segment))
+    	// Calculate storage indices
+    	const bool bIsLastSegment = (SegmentIndex == NumSegments - 1);
+    
+    	UE_LOG(LogSpline, Verbose, TEXT("Removing segment at index %d"), SegmentIndex);
+    
+        if (bIsLastSegment)
 	    {
-	    	const int32 RightSeg =
-			 (Segment < NumSegments - 1) ? (Segment + 1) : 0; // wrap in closed case
-
-	    	int32 ThisIdx0, ThisIdx1, ThisIdx2, ThisIdx3;
-	    	int32 RightIdx0, RightIdx1, RightIdx2, RightIdx3;
-
-	    	GetBezierIndices(Segment, ThisIdx0, ThisIdx1, ThisIdx2, ThisIdx3);
-	    	GetBezierIndices(RightSeg, RightIdx0, RightIdx1, RightIdx2, RightIdx3);
-
-	    	const ValueType ThisP3  = Base::GetValue(ThisIdx3);
-	    	const ValueType RightP0 = Base::GetValue(RightIdx0);
-	    	const ValueType RightP1 = Base::GetValue(RightIdx1);
-
-	    	const ValueType Delta   = RightP0 - RightP1; // original in-tangent
-	    	const ValueType NewStart = ThisP3;
-	    	const ValueType NewOut   = NewStart - Delta;
-
-	    	Base::SetValue(RightIdx0, NewStart);
-	    	Base::SetValue(RightIdx1, NewOut);
+	    	return RemovePoint(GetNumberOfSegments());
 	    }
+	    else
+	    {
+	    	// in all other cases, removing the segment would be the same as removing the point at that index
+	        return RemovePoint(SegmentIndex);
+	    }
+    }
 
-	    // 4) Swap knot intervals for the flipped segment & its neighbor interval
-	    Base::SwapKnots(Segment, Segment + 1);
-
-	    UE_LOG(LogSpline, Verbose, TEXT("\t"));
-	    UE_LOG(LogSpline, Verbose, TEXT("After FlipSegment(%d):"), Segment);
-	    Base::Dump();
-	}
-	
-	virtual int32 FindIndexForParameter(float Parameter, float& OutLocalParam) const override
+	/**
+	 * Removes a point from the spline.
+	 * @param PointIndex - Index of the point to remove
+	 * @return true if successfully removed
+	 */
+	bool RemovePoint(int32 PointIndex)
 	{
-		const int32 NumSegments = GetNumberOfSegments();
-		const UE::Geometry::FInterval1f ParameterSpace = GetParameterSpace();
+		const int32 NumAnchors  = GetNumAnchorPoints();
+	    const int32 NumSegments = GetNumberOfSegments();
 
-		if (ParameterSpace.IsEmpty())
-		{
-			return INDEX_NONE;
-		}
+	    if (NumAnchors == 0)
+	    {
+	        return false;
+	    }
 
-		// Left end: snap onto first segment start
-		if (Parameter <= ParameterSpace.Min)
+	    // In this model, "points" == anchors
+	    if (PointIndex < 0 || PointIndex >= NumAnchors)
+	    {
+	        return false;
+	    }
+
+		auto RemoveAnchorTriple = [this](int32 AnchorIndex)
 		{
-			OutLocalParam = 0.0f;
-			return 0;
-		}
+			const int32 BaseIdx = AnchorBase(AnchorIndex);
+			// Remove in reverse index order so earlier removes don't shift later indices
+			this->RemoveValue(BaseIdx + 2); // Out_i
+			this->RemoveValue(BaseIdx + 1); // A_i
+			this->RemoveValue(BaseIdx + 0); // In_i
+		};
 
-		// Right end: snap onto last segment end
-		if (Parameter >= ParameterSpace.Max)
+		if (NumAnchors == 1)		// Special case: only 1 anchors → removing one leaves empty spline.
 		{
-			OutLocalParam = 1.0f;
-			return NumSegments - 1;
+			Base::Values.Reset();
+	        this->PairKnots.Reset();
+	        return true;
 		}
+	    else if (NumAnchors == 2)	// Special case: only 2 anchors → removing one leaves 1 anchor, no segments.
+	    {
+	        RemoveAnchorTriple(PointIndex);
+	        // One anchor left, no segments: clear knots
+	        this->PairKnots.Reset();
+	        return true;
+	    }
 
-		const float LastKnot = Base::PairKnots.Last().Value;
+	    // ---- OPEN SPLINE CASE ----
+	    if (!IsActuallyClosed())
+	    {
+	        // For open curves:
+	        //  - anchors = NumSegments + 1
+	        //  - boundary knots are indices [0 .. NumSegments]
+	        //  - interior boundaries (between anchors j and j+1) correspond to knot index j+1?
+
+	        // A) remove first point (anchor 0): drop first segment and knot 0
+	        if (PointIndex == 0)
+	        {
+	            RemoveAnchorTriple(0);
+
+	            // Remove start boundary knot
+	            if (this->PairKnots.Num() > 0)
+	            {
+	                this->RemoveKnot(0);
+	            }
+
+	            this->ApplyClampedKnotsMultiplicity();
+	            ApplyInternalKnotsMultiplicity();
+	            return true;
+	        }
+
+	        // B) remove last point (anchor N-1): drop last segment and last boundary knot
+	        if (PointIndex == NumAnchors - 1)
+	        {
+	            RemoveAnchorTriple(NumAnchors - 1);
+
+	            // Remove end boundary knot (index == NumSegments)
+	            if (this->PairKnots.Num() > 0)
+	            {
+	                const int32 EndKnotIndex = FMath::Clamp(NumSegments, 0, this->PairKnots.Num() - 1);
+	                this->RemoveKnot(EndKnotIndex);
+	            }
+
+	            this->ApplyClampedKnotsMultiplicity();
+	            ApplyInternalKnotsMultiplicity();
+	            return true;
+	        }
+
+	        // C) remove internal point (0 < j < N-1): merge segments (j-1)->j and j->(j+1)
+	        //    into one (j-1)->(j+1). Shape-wise we keep Out_{j-1} and In_{j+1};
+	        //    dropping A_j's triple is enough.
+	        RemoveAnchorTriple(PointIndex);
+
+	        // Interior knot index aligned with the removed anchor index in the old scheme:
+	        // - anchor j corresponds to a boundary between segments (j-1) and j
+	        // - that boundary’s knot index is j
+	        if (this->PairKnots.Num() > PointIndex)
+	        {
+	            this->RemoveKnot(PointIndex);
+	        }
+
+	        this->ApplyClampedKnotsMultiplicity();
+	        ApplyInternalKnotsMultiplicity();
+	        return true;
+	    }
 
-		// We can now assume Parameter is contained by ParameterSpace, no need for clamping.
+	    // ---- CLOSED SPLINE CASE ----
+	    // For closed curves:
+	    //  - anchors == segments
+	    //  - all boundaries are "interior" (no clamped ends)
+	    //  - removing anchor j merges two segments (j-1)->j and j->(j+1)
+	    //    into one (j-1)->(j+1); we keep Out_{j-1} and In_{j+1} and just drop A_j.
+	    {
+	        const int32 RemoveAnchorIdx = PointIndex % NumAnchors;
+	        RemoveAnchorTriple(RemoveAnchorIdx);
+
+	        // In the old closed implementation, a point removal deleted the
+	        // corresponding boundary knot at the same index.
+	        if (this->PairKnots.Num() > RemoveAnchorIdx)
+	        {
+	            this->RemoveKnot(RemoveAnchorIdx);
+	        }
+
+	        // No clamped ends in closed curve: all remaining knots should be interior-style.
+	        ApplyInternalKnotsMultiplicity();
+	        return true;
+	    }
+	}
 
-		// Closed loop: handle the synthetic closing segment [LastKnot, LastKnot + LoopKnotDelta]
-		if (Parameter >= LastKnot && ensure(HasClosingSegment()))
+    /**
+	 * Updates a single control point within a Bezier segment
+	 * @param Storage - Storage to update the point in
+	 * @param SegmentIndex - Which cubic Bezier segment (each has 4 points)
+	 * @param PointIndex - Which point in the segment (0=P0, 1=P1, 2=P2, 3=P3)
+	 * @param NewValue - New position for the control point
+	 * @return true if successfully updated
+	 */
+    bool UpdateSegmentPoint(
+    	const int32 SegmentIndex,
+    	const int32 PointIndex,
+    	const ValueType& NewValue)
+    {
+		const int32 NumSegments = GetNumberOfSegments();
+		if (SegmentIndex < 0 || SegmentIndex >= NumSegments ||
+			PointIndex  < 0 || PointIndex  > 3)
 		{
-			const float t = (Parameter - LastKnot) / FMath::Max(ParameterSpace.Max - LastKnot, UE_SMALL_NUMBER);
-			OutLocalParam = FMath::Clamp(t, 0.0f, 1.0f);
-			return NumSegments - 1;
+			return false;
 		}
-        
-		int32 SegmentIndex = INDEX_NONE;
 
-		// todo: make CVar
-		constexpr bool bEnableSegmentCaching = true;
-		if (bEnableSegmentCaching && LastEvaluatedSegment != INDEX_NONE && LastEvaluatedSegment < Base::PairKnots.Num() - 1)
-		{
-			if (Parameter < Base::PairKnots[LastEvaluatedSegment + 1].Value && Parameter >= Base::PairKnots[LastEvaluatedSegment].Value)
-			{
-				// cache hit!
-				SegmentIndex = LastEvaluatedSegment;
-			}
-		}
-		
-		if (SegmentIndex == INDEX_NONE)
-		{
-			// cache miss!
-			int32 SegmentLow = 0;
-			int32 SegmentHigh = Base::PairKnots.Num() - 2;
-			SegmentIndex = SegmentLow + (SegmentHigh - SegmentLow) / 2;
+		int32 BezierIndices[4];
+		GetBezierIndices(SegmentIndex, BezierIndices[0], BezierIndices[1], BezierIndices[2], BezierIndices[3]);
+		int32 ValueIndex = BezierIndices[PointIndex];
+		this->SetValue(ValueIndex, NewValue);
+		return true;
+    }
 
-			while (SegmentLow <= SegmentHigh)
-			{
-				int32 Mid = SegmentLow + (SegmentHigh - SegmentLow) / 2;
-				const float& StartParam = Base::PairKnots[Mid].Value;
-				const float& EndParam   = Base::PairKnots[Mid + 1].Value;
+    /**
+     * Updates all control points of a Bezier segment at once
+     * @param SegmentIndex - Which cubic Bezier segment to update
+     * @param P0 - Start point of segment
+     * @param P1 - First control point
+     * @param P2 - Second control point
+     * @param P3 - End point of segment
+     * @return true if successfully updated
+     */
+    bool UpdateSegment(
+    	int32 SegmentIndex,
+    	const ValueType& P0,
+    	const ValueType& P1, 
+        const ValueType& P2,
+        const ValueType& P3)
+    {
+        const int32 NumSegments = GetNumberOfSegments();
+        if (SegmentIndex < 0 || SegmentIndex >= NumSegments)
+        {
+            return false;
+        }
 
-				if (Parameter >= StartParam && Parameter < EndParam)
-				{
-					SegmentIndex = Mid;
-					break;
-				}
-				else if (Parameter < StartParam)
-				{
-					SegmentHigh = Mid - 1;
-				}
-				else
-				{
-					SegmentLow = Mid + 1;
-				}
-			}
-		}
+		int32 Idx0, Idx1, Idx2, Idx3;
+		GetBezierIndices(SegmentIndex, Idx0, Idx1, Idx2, Idx3);
 		
-		// Update the cache.
-		LastEvaluatedSegment = SegmentIndex;
-        
-		// Calculate local parameter
-		float SegmentLength = Base::PairKnots[SegmentIndex + 1].Value - Base::PairKnots[SegmentIndex].Value;
-		if (SegmentLength > UE_SMALL_NUMBER)
-		{
-			OutLocalParam = (Parameter - Base::PairKnots[SegmentIndex].Value) / SegmentLength;
-		}
-		else
-		{
-			OutLocalParam = 0.0f;
-		}
-        
-		return SegmentIndex;
-	}
-	
-	virtual void SetClosedLoop(bool bShouldClose) override
-	{
-		// Skip if state isn't changing
-		if (bShouldClose == Base::IsClosedLoop())
-		{
-			return;
-		}
-        
-		// Just flip the topology flag; do not add/remove any control points.
-		Base::SetClosedLoop(bShouldClose);
+		this->SetValue(Idx0, P0);
+		this->SetValue(Idx1, P1);
+		this->SetValue(Idx2, P2);
+		this->SetValue(Idx3, P3);
 
-		// Update knot multiplicities.
-		ApplyClampedKnotsMultiplicity();
-		ApplyInternalKnotsMultiplicity();
-	}
-
-	/**
-	 * Maps a local segment parameter [0,1] to global parameter space
-	 * @param SegmentIndex - Index of the segment
-	 * @param LocalParam - Local parameter within segment [0,1]
-	 * @return Parameter value in global space, or 0 if segment is invalid
-	 */
-	float MapLocalSegmentParameterToGlobal(int32 SegmentIndex, float LocalParam) const
-	{
-		FInterval1f SegmentRange = GetSegmentParameterRange(SegmentIndex);
-		return SegmentRange.Interpolate(LocalParam);
-	}
-
-	void GetBezierIndices(int32 SegmentIndex, int32& OutP0Idx, int32& OutP1Idx, int32& OutP2Idx, int32& OutP3Idx) const
-	{
-		if (!IsClosingSegment(SegmentIndex))
-		{
-			OutP0Idx = P0Idx(SegmentIndex);
-			OutP1Idx = P1Idx(SegmentIndex);
-			OutP2Idx = P2Idx(SegmentIndex);
-			OutP3Idx = P3Idx(SegmentIndex);
-		}
-		else
-		{
-			const int32 NumAnchors = GetNumAnchorPoints();
+		return true;
+    }
 
-			OutP0Idx = Close_P0Idx(NumAnchors);
-			OutP1Idx = Close_P1Idx(NumAnchors);
-			OutP2Idx = Close_P2Idx();
-			OutP3Idx = Close_P3Idx();
-		}
-	}
-	
 	/**
 	 * Anchor points are unique segment endpoints.
 	 * All segments have 2 corresponding anchor points.
 	 * For a closed spline, the final segment's anchor points are the first and last points.
 	 */
-	int32 GetNumAnchorPoints() const
-	{
-		return Base::Values.Num() / 3;
-	}
+	int32 GetNumAnchorPoints() const { return Base::Values.Num() / 3; }
 
 	void SetAnchorPoint(int32 AnchorIndex, const ValueType& NewPosition)
 	{
@@ -1013,328 +1204,627 @@ public:
 		const int32 AnchorValueIndex = AnchorBase(AnchorIndex) + 1;
 		return Base::GetValue(AnchorValueIndex);
 	}
-
-	/** Per-anchor setters/getters */
-	void SetHandleOutOfAnchor(int32 AnchorIdx, const ValueType& Out)
-	{
-		Base::SetValue(AnchorBase(AnchorIdx) + 2, Out);
-	}
-
-	const ValueType& GetHandleOutOfAnchor(int32 AnchorIdx) const
-	{
-		return Base::GetValue(AnchorBase(AnchorIdx) + 2);
-	}
 	
-	void SetHandleInOfAnchor(int32 AnchorIdx, const ValueType& In)
-	{
-		Base::SetValue(AnchorBase(AnchorIdx) + 0, In);
-	}
-
-	const ValueType& GetHandleInOfAnchor(int32 AnchorIdx) const
+	/** Returns the number of Bezier segments in the spline */
+	virtual int32 GetNumberOfSegments() const override
 	{
-		return Base::GetValue(AnchorBase(AnchorIdx) + 0);
+		const int32 NumAnchors = GetNumAnchorPoints();
+		return IsActuallyClosed() ? NumAnchors : FMath::Max(0, NumAnchors - 1);
 	}
 
-	void SetAnchors(const TArray<FAnchor>& InAnchors, EParameterizationPolicy ParameterizationPolicy = EParameterizationPolicy::Centripetal)
+	/**
+	 * Maps a segment index to its parameter range
+	 * @param SegmentIndex - Index of the segment (0-based)
+	 * @return True if the segment index is valid and mapping succeeded
+	 */
+	virtual FInterval1f GetSegmentParameterRange(int32 SegmentIndex) const override
 	{
-		Clear();
-
-		AppendAnchors(InAnchors, ParameterizationPolicy);
-	}
+		FInterval1f SegmentRange;
 
-	void AppendAnchors(const TArray<FAnchor>& InAnchors, EParameterizationPolicy ParameterizationPolicy = EParameterizationPolicy::Centripetal)
-	{
-		const int32 ExpectedNumAnchors = GetNumAnchorPoints() + InAnchors.Num();
+		const int32 NumSegments = GetNumberOfSegments();
+		if (SegmentIndex < 0 || SegmentIndex >= NumSegments || Base::PairKnots.Num() == 0 || !Base::PairKnots.IsValidIndex(SegmentIndex))
+		{
+			return SegmentRange;
+		}
 
-		Base::Values.Reserve(ExpectedNumAnchors * 3);
-		Base::PairKnots.Reserve(ExpectedNumAnchors * 3);
+		SegmentRange.Min = Base::PairKnots[SegmentIndex].Value;
 
-		for (const FAnchor& Anchor : InAnchors)
+		if (!IsActuallyClosed())
 		{
-			AppendAnchor(Anchor, ParameterizationPolicy);
+			// Open: knots are standard boundaries, so we still use i+1.
+			if (SegmentIndex + 1 < Base::PairKnots.Num())
+			{
+				SegmentRange.Max = Base::PairKnots[SegmentIndex + 1].Value;
+			}
+			return SegmentRange;
 		}
-	}
 
-	int32 PrependAnchor(const FAnchor& InAnchor, float Parameter)
-	{
-		// Cannot prepend if Parameter is not less than all of our knots.
-		if (Base::PairKnots.Num() > 0 && Parameter >= Base::PairKnots[0].Value)
+		// Closed: PairKnots holds starts of each segment; the final segment
+		// closes the loop at Min + LoopKnotDelta.
+		if (SegmentIndex < NumSegments - 1)
 		{
-			return INDEX_NONE;
+			SegmentRange.Max = Base::PairKnots[SegmentIndex + 1].Value;
+		}
+		else
+		{
+			SegmentRange.Max = SegmentRange.Min + LoopKnotDelta;
 		}
 
-		Base::InsertValue(0, InAnchor.InHandle);
-		Base::InsertValue(1, InAnchor.Anchor);
-		Base::InsertValue(2, InAnchor.OutHandle);
-
-		const int32 KnotIndex = Base::InsertKnot(FKnot(Parameter));
-
-		ensureAlwaysMsgf(KnotIndex == 0, TEXT("Unexpected knot insertion index when prepending an anchor!"));
+		return SegmentRange;
+	}
 
-		ApplyClampedKnotsMultiplicity();
-		ApplyInternalKnotsMultiplicity();
+	// Parameterization methods
+	/**
+	 * Reparameterizes the spline based on the provided points and mode
+	 * @param Mode - New parameterization mode
+	 * @param Points - Array of control points to use for reparameterization
+	 * @return 
+	 */
+	virtual void Reparameterize(EParameterizationPolicy Mode = EParameterizationPolicy::Centripetal) override
+	{
+		const TArray<ValueType>& Points = this->Values;
+		const int32 NumValues = Points.Num();
 
-		return 0;
+		GenerateDistanceKnotsForBezier(Mode);
+    	
+		UE_LOG(LogSpline, Verbose, TEXT("Set knot vector (Reparameterize) - Points: %d, Knots: %d, Mode: %d"),
+					   NumValues,  this->PairKnots.Num(), static_cast<int32>(Mode));
+		
+		Base::PrintKnotVector();
 	}
 	
-	int32 PrependAnchor(const FAnchor& InAnchor, EParameterizationPolicy ParameterizationPolicy = EParameterizationPolicy::Centripetal)
+	virtual FInterval1f GetParameterSpace() const override
 	{
-		float NewParameter;
-
-		if (GetNumAnchorPoints() <= 0)
+		if (Base::PairKnots.Num() == 0)
 		{
-			NewParameter = 0.f;
+			return FInterval1f::Empty();
 		}
-		else
-		{
-			const ValueType OldAnchor = GetAnchorPoint(0);
-			const float OldParameter = GetParameter(0);
 
-			const float SegmentLength = ComputeSegmentLength(InAnchor.Anchor, OldAnchor, ParameterizationPolicy);
-			const float NewParamByLength = OldParameter - SegmentLength;
-			const float NewParamByStep = Param::PrevDistinct(OldParameter);
+		const float Min = Base::PairKnots[0].Value;
+		const float LastKnot = Base::PairKnots.Last().Value;
 
-			NewParameter = FMath::Min(NewParamByLength, NewParamByStep);
+		if (!IsActuallyClosed())
+		{
+			return FInterval1f(Min, LastKnot);
 		}
 
-		return PrependAnchor(InAnchor, NewParameter);
+		// Closed: closing boundary is last real knot + LoopKnotDelta
+		return FInterval1f(Min, LastKnot + LoopKnotDelta);
 	}
 
-	int32 AppendAnchor(const FAnchor& InAnchor, float Parameter)
+	/**
+	 * Gets the parameter value for a specific Bezier control point
+	 * Index is a *logical* Bezier control index:
+	 * Index = SegmentIndex * 4 + LocalPointIndex (0..3)
+	 * @param Index - Index of the Bezier control point
+	 * @return Parameter value at the specified control point
+	 */
+	virtual float GetParameter(int32 Index) const override
 	{
-		// Cannot append if Parameter is not greater than all of our knots.
-		if (Base::PairKnots.Num() > 0 && Parameter <= Base::PairKnots.Last().Value)
+		const int32 NumSegments = GetNumberOfSegments();
+		if (NumSegments <= 0)
 		{
-			return INDEX_NONE;
+			return 0.f;
 		}
 
-		const int32 InsertIndex = GetNumAnchorPoints();
+		const int32 MaxLogicalIndex = NumSegments * 4;
+		if (Index < 0 || Index > MaxLogicalIndex || Base::PairKnots.IsEmpty())
+		{
+			return 0.f;
+		}
+		
+        int32 SegmentIndex = Index / 4;
+        int32 LocalPointIndex = Index % 4;
+    
+		// Logical "one past last" control index – treat as the domain end.
+		if (SegmentIndex >= NumSegments)
+		{
+			const FInterval1f Domain = GetParameterSpace();
+			return Domain.Max;
+		}
 
-		Base::AddValue(InAnchor.InHandle);
-		Base::AddValue(InAnchor.Anchor);
-		Base::AddValue(InAnchor.OutHandle);
+		// Use the segment parameter range helper, which already knows how to
+		// handle the closing segment for closed loops.
+		const FInterval1f Range = GetSegmentParameterRange(SegmentIndex);
 
-		const int32 KnotIndex = Base::InsertKnot(FKnot(Parameter));
+        const float T = static_cast<float>(LocalPointIndex) / 3.0f;
+        return Range.Interpolate(T);
+	}
 
-		ensureAlwaysMsgf(KnotIndex == InsertIndex, TEXT("Unexpected knot insertion index when appending an anchor!"));
+	/**
+	* Sets the parameter value for a specific Bezier control point
+	* Index is a *logical* Bezier control index:
+	* Index = SegmentIndex * 4 + LocalPointIndex (0..3)
+	* @param Index - Index of the Bezier control point
+	* @param NewParameter - New parameter value to set
+	* @return The index of the bezier point that was actually updated, or INDEX_NONE if failed
+	*/ 
+	virtual int32 SetParameter(int32 Index, float NewParameter) override
+	{
+		const int32 NumSegments = GetNumberOfSegments();
+		if (NumSegments <= 0)
+		{
+			return INDEX_NONE;
+		}
 
-		ApplyClampedKnotsMultiplicity();
-		ApplyInternalKnotsMultiplicity();
+		// Max logical control index is 4 * NumSegments:
+		// For open: P0..P3 of each seg -> 4 * NumSegments controls, with
+		// the last P3 coincident with the last anchor.
+		const int32 MaxLogicalIndex = NumSegments * 4;
+		if (Index < 0 || Index > MaxLogicalIndex)
+		{
+			return INDEX_NONE;
+		}
+		
+		constexpr int32 PointsPerSegment = 4;
+		const int32 SegmentIndex = Index / PointsPerSegment;
+		const int32 LocalPointIndex = Index % PointsPerSegment;
 
-		return InsertIndex;
-	}
-	
-	int32 AppendAnchor(const FAnchor& InAnchor, EParameterizationPolicy ParameterizationPolicy = EParameterizationPolicy::Centripetal)
-	{
-		const int32 NumAnchors = GetNumAnchorPoints();
+		const bool bClosingBoundary = IsClosingSegment(SegmentIndex) && (LocalPointIndex != 0); // P1/P2/P3 of last seg
 
-		float NewParameter;
+		float OldParameter = 0.f;
+		int32 OldParameterIndex = INDEX_NONE;
 
-		if (NumAnchors <= 0)
+		if (!bClosingBoundary)
 		{
-			NewParameter = 0.f;
+			// Legacy behaviour: real knot index
+			OldParameterIndex = (LocalPointIndex == 0) ? SegmentIndex : SegmentIndex + 1;
+			if (!Base::PairKnots.IsValidIndex(OldParameterIndex))
+			{
+				return INDEX_NONE;
+			}
+			OldParameter = Base::PairKnots[OldParameterIndex].Value;
 		}
 		else
 		{
-			const ValueType OldAnchor = GetAnchorPoint(NumAnchors - 1);
-			const float OldParameter = GetParameter(NumAnchors - 1);
-
-			const float SegmentLength = ComputeSegmentLength(InAnchor.Anchor, OldAnchor, ParameterizationPolicy);
-			const float NewParamByLength = OldParameter + SegmentLength;
-			const float NewParamByStep = Param::NextDistinct(OldParameter);
+			// Closing boundary: virtual "last + delta"
+			if (Base::PairKnots.Num() == 0)
+			{
+				return INDEX_NONE;
+			}
 
-			NewParameter = FMath::Max(NewParamByLength, NewParamByStep);
+			const float LastKnot = Base::PairKnots.Last().Value;
+			OldParameter = LastKnot + LoopKnotDelta;
 		}
-
-		return AppendAnchor(InAnchor, NewParameter);
-	}
-
-	/** Per-anchor structural inserts */
-	int32 InsertAnchor(const FAnchor& InAnchor, float Parameter)
-	{
-		if (GetNumAnchorPoints() <= 0)
+    
+		// Early exit if parameter isn't changing
+		if (OldParameter == NewParameter)
 		{
-			return PrependAnchor(InAnchor, Parameter);
+			return Index;
 		}
+		
+		UE_LOG(LogSpline, Verbose, TEXT("\t"));
+		UE_LOG(LogSpline, Verbose, TEXT("Before SetParameter(%d, %f):"), Index, NewParameter);
+		Base::Dump();
+		
+		// Prevent reuse of an existing knot.
+		const float DesiredNewParameter = NewParameter;
+		typename Base::FKnotSearchCache SearchCache{ this->PairKnots };
+		const float ActualNewParameter = this->GetNearestAvailableKnotValue(DesiredNewParameter, SearchCache);
 
-		const FInterval1f ParameterSpace = GetParameterSpace();
+		const float DesiredToOldDistance = FMath::Abs(DesiredNewParameter - OldParameter);
+		const float DesiredToActualNewDistance = FMath::Abs(DesiredNewParameter - ActualNewParameter);
 
-		if (!ensureAlwaysMsgf(!ParameterSpace.IsEmpty(), TEXT("Parameter space is empty but we do not have an empty spline!")))
-		{
-			return INDEX_NONE;
-		}
-		else if (Parameter < ParameterSpace.Min)
+		if (DesiredToActualNewDistance > DesiredToOldDistance)
 		{
-			return PrependAnchor(InAnchor, Parameter);
+			// If our nearest valid knot is further away from the desired value than our current value is, we need to just keep our current value
+			return Index;
 		}
-		else if (Parameter > ParameterSpace.Max)
+		
+		int32 ReturnIndex = INDEX_NONE;
+
+		// --- handle closing boundary separately ---------------------------------
+		if (bClosingBoundary)
 		{
-			return AppendAnchor(InAnchor, Parameter);
+			// We are moving the virtual closing boundary of the last segment.
+			const float LastKnot = Base::PairKnots.Last().Value;
+			const float NewLoopDelta = NewParameter - LastKnot;
+
+			// Must remain a forward-advancing interval.
+			if (NewLoopDelta <= 0.f)
+			{
+				return Index;
+			}
+
+			LoopKnotDelta = NewLoopDelta;
+			Base::MarkFlatKnotsCacheDirty();
+
+			UE_LOG(LogSpline, Verbose, TEXT("\t"));
+			UE_LOG(LogSpline, Verbose, TEXT("After SetParameter(%d, %f) [closing]:"),
+				Index, NewParameter);
+			Base::Dump();
+
+			return Index;
 		}
+		// -----------------------------------------------------------------------------
 
-		float LocalT;
-		const int32 SegmentIndex = FindIndexForParameter(Parameter, LocalT);
+		if (LocalPointIndex == 0) // P0
+		{
+			int32 LeftSegmentIndex = SegmentIndex - 1;
+			int32 RightSegmentIndex = SegmentIndex;
+			
+			Base::SetKnot(OldParameterIndex, NewParameter);
 
-		if (SegmentIndex < 0 || SegmentIndex >= GetNumberOfSegments())
+			int32 Segment = SegmentIndex;
+			if (NewParameter < OldParameter)
+			{
+				while (LeftSegmentIndex >= 0 && GetSegmentParameterRange(LeftSegmentIndex).Min > GetSegmentParameterRange(LeftSegmentIndex).Max)
+				{
+					FlipSegment(LeftSegmentIndex);
+					LeftSegmentIndex--;
+					RightSegmentIndex--;
+					Segment = RightSegmentIndex;
+				}
+			}
+			else if (NewParameter > OldParameter)
+			{
+				while (RightSegmentIndex <= NumSegments - 1 && GetSegmentParameterRange(RightSegmentIndex).Min > GetSegmentParameterRange(RightSegmentIndex).Max)
+				{
+					FlipSegment(RightSegmentIndex);
+					LeftSegmentIndex++;
+					RightSegmentIndex++;
+					Segment = RightSegmentIndex;
+				}
+			}
+			
+			ReturnIndex = Segment * PointsPerSegment + LocalPointIndex;
+		}
+		else if (LocalPointIndex == 3) // P3
 		{
-			// If Parameter is contained by ParameterSpace (confirmed by Prepend/Append special casing above)
-			// but the parameter maps to an invalid segment, something is wrong that we can't handle nicely.
-			return INDEX_NONE;
+			int32 LeftSegmentIndex = SegmentIndex;
+			int32 RightSegmentIndex = SegmentIndex + 1;
+			
+			Base::SetKnot(OldParameterIndex, NewParameter);
+
+			int32 Segment = SegmentIndex;
+			if (NewParameter < OldParameter)
+			{
+				while (LeftSegmentIndex >= 0 && GetSegmentParameterRange(LeftSegmentIndex).Min > GetSegmentParameterRange(LeftSegmentIndex).Max)
+				{
+					FlipSegment(LeftSegmentIndex);
+					LeftSegmentIndex--;
+					RightSegmentIndex--;
+					Segment = LeftSegmentIndex;
+				}
+			}
+			else if (NewParameter > OldParameter)
+			{
+				while (RightSegmentIndex <= NumSegments - 1 && GetSegmentParameterRange(RightSegmentIndex).Min > GetSegmentParameterRange(RightSegmentIndex).Max)
+				{
+					FlipSegment(RightSegmentIndex);
+					LeftSegmentIndex++;
+					RightSegmentIndex++;
+					Segment = LeftSegmentIndex;
+				}
+			}
+			
+			ReturnIndex = Segment * PointsPerSegment + LocalPointIndex;
 		}
 
-		// Segment S is bounded by anchors S and S+1, except for the closing segment which is bounded by anchors S and 0.
-		// We will special case insertion on the closing segment as this is simply an append operation.
-		if (IsClosingSegment(SegmentIndex))
+		UE_LOG(LogSpline, Verbose, TEXT("\t"));
+		UE_LOG(LogSpline, Verbose, TEXT("After SetParameter(%d, %f):"), Index, NewParameter);
+		Base::Dump();
+		
+		return ReturnIndex;
+	}
+	
+	void FlipSegment(int32 Segment)
+	{
+		const int32 NumSegments = GetNumberOfSegments();
+		const int32 NumAnchors  = GetNumAnchorPoints();
+		if (Segment < 0 || Segment >= NumSegments || NumAnchors < 2)
 		{
-			return AppendAnchor(InAnchor, Parameter);
+			return;
 		}
 
-		// We may now assume anchors SegmentIndex and SegmentIndex + 1 are valid without bounds checking.
-		// In order to insert onto segment S, we must insert at the index S + 1 so that the new anchor is the new endpoint of segment S.
+		UE_LOG(LogSpline, Verbose, TEXT("\t"));
+		UE_LOG(LogSpline, Verbose, TEXT("Before FlipSegment(%d):"), Segment);
+		Base::Dump();
 		
-		FInterval1f SegmentRange = GetSegmentParameterRange(SegmentIndex);
-
-		if (Parameter == SegmentRange.Min)
+		// 1) Flip this segment's Bezier window: P0<->P3, P1<->P2
 		{
-			Parameter = Param::NextDistinct(Parameter);
+			int32 Idx0, Idx1, Idx2, Idx3;
+			GetBezierIndices(Segment, Idx0, Idx1, Idx2, Idx3);
+
+			ValueType P0 = this->GetValue(Idx0);
+			ValueType P1 = this->GetValue(Idx1);
+			ValueType P2 = this->GetValue(Idx2);
+			ValueType P3 = this->GetValue(Idx3);
+
+			this->SetValue(Idx0, P3);
+			this->SetValue(Idx1, P2);
+			this->SetValue(Idx2, P1);
+			this->SetValue(Idx3, P0);
 		}
-		else if (Parameter == SegmentRange.Max)
+
+		// 2) Fix left neighbor (Segment-1): keep its out-tangent but move its end point to match
+	    if (Segment > 0 || (IsActuallyClosed() && Segment == 0))
+	    {
+	    	const int32 LeftSeg =
+			 (Segment > 0) ? (Segment - 1) : (NumSegments - 1); // wrap in closed case
+
+	    	int32 ThisIdx0, ThisIdx1, ThisIdx2, ThisIdx3;
+	    	int32 LeftIdx0, LeftIdx1, LeftIdx2, LeftIdx3;
+
+	    	GetBezierIndices(Segment, ThisIdx0, ThisIdx1, ThisIdx2, ThisIdx3);
+	    	GetBezierIndices(LeftSeg, LeftIdx0, LeftIdx1, LeftIdx2, LeftIdx3);
+
+	    	const ValueType ThisP0 = this->GetValue(ThisIdx0);
+	    	const ValueType LeftP2 = this->GetValue(LeftIdx2);
+	    	const ValueType LeftP3 = this->GetValue(LeftIdx3);
+
+	    	const ValueType Delta = LeftP3 - LeftP2; // original out-tangent
+	    	const ValueType NewEnd = ThisP0;
+	    	const ValueType NewIn  = NewEnd - Delta;
+
+	    	this->SetValue(LeftIdx3, NewEnd);
+	    	this->SetValue(LeftIdx2, NewIn);
+	    }
+
+	    // 3) Fix right neighbor (Segment+1): keep its in-tangent but move its start point to match
+	    if (Segment < NumSegments - 1 || IsClosingSegment(Segment))
+	    {
+	    	const int32 RightSeg =
+			 (Segment < NumSegments - 1) ? (Segment + 1) : 0; // wrap in closed case
+
+	    	int32 ThisIdx0, ThisIdx1, ThisIdx2, ThisIdx3;
+	    	int32 RightIdx0, RightIdx1, RightIdx2, RightIdx3;
+
+	    	GetBezierIndices(Segment, ThisIdx0, ThisIdx1, ThisIdx2, ThisIdx3);
+	    	GetBezierIndices(RightSeg, RightIdx0, RightIdx1, RightIdx2, RightIdx3);
+
+	    	const ValueType ThisP3  = this->GetValue(ThisIdx3);
+	    	const ValueType RightP0 = this->GetValue(RightIdx0);
+	    	const ValueType RightP1 = this->GetValue(RightIdx1);
+
+	    	const ValueType Delta   = RightP0 - RightP1; // original in-tangent
+	    	const ValueType NewStart = ThisP3;
+	    	const ValueType NewOut   = NewStart - Delta;
+
+	    	this->SetValue(RightIdx0, NewStart);
+	    	this->SetValue(RightIdx1, NewOut);
+	    }
+
+	    // 4) Swap knot intervals for the flipped segment & its neighbor interval
+	    Base::SwapKnots(Segment, Segment + 1);
+
+	    UE_LOG(LogSpline, Verbose, TEXT("\t"));
+	    UE_LOG(LogSpline, Verbose, TEXT("After FlipSegment(%d):"), Segment);
+	    Base::Dump();
+	}
+	
+	virtual int32 FindIndexForParameter(float Parameter, float& OutLocalParam) const override
+	{
+		const int32 NumSegments = GetNumberOfSegments();
+		if (NumSegments <= 0)
 		{
-			Parameter = Param::PrevDistinct(Parameter);
+			OutLocalParam = 0.0f;
+			return 0;
 		}
-		
-		// Bail out if we lost usable range, this should not happen.
-		if (Parameter <= SegmentRange.Min || Parameter >= SegmentRange.Max)
+
+		int32 SegmentIndex = 0;
+		float LocalT       = 0.0f;
+		if (!MapGlobalParameterToLocalSegment(Parameter, SegmentIndex, LocalT))
 		{
-			ensureAlways(false);
 			return INDEX_NONE;
 		}
 
-		const int32 InsertIndex = SegmentIndex + 1;
-		Base::InsertValue(AnchorBase(InsertIndex) + 0, InAnchor.InHandle);
-		Base::InsertValue(AnchorBase(InsertIndex) + 1, InAnchor.Anchor);
-		Base::InsertValue(AnchorBase(InsertIndex) + 2, InAnchor.OutHandle);
+		SegmentIndex = FMath::Clamp(SegmentIndex, 0, NumSegments - 1);
+		OutLocalParam = LocalT;
+
+		// This function historically returned a "control point index" = 4*segment,
+		// i.e. the (virtual) P0 index of that segment. We preserve that contract.
+		return SegmentIndex * 4;
+	}
+	
+	virtual void SetClosedLoop(bool bShouldClose) override
+	{
+		// Skip if state isn't changing
+		if (bShouldClose == this->IsClosedLoop())
+			return;
+        
+		// Just flip the topology flag; do not add/remove any control points.
+		Base::SetClosedLoop(bShouldClose);
+
+		this->ApplyClampedKnotsMultiplicity();
+		ApplyInternalKnotsMultiplicity();
 
-		// We always use Degree as the multiplicity because this is always an interior knot.
-		const int32 KnotIndex = Base::InsertKnot(FKnot(Parameter, Degree));
-		ensureAlways(InsertIndex == KnotIndex);
+		// Note:
+		// - Knot topology (number and placement of knots) is rebuilt by Reparameterize()
+		//   when the owning spline chooses to (eg, FTangentSpline::Reparameterize).
+		// - We intentionally do NOT call AppendBezierSegment / RemoveSegment here anymore,
+		//   because the closing segment is now virtual: {A_last, Out_last, In_0, A_0}.
 
-		return InsertIndex;
 	}
 
 	/**
-	 * Removes a point from the spline.
-	 * @param PointIndex - Index of the point to remove
-	 * @return true if successfully removed
+	 * Maps a local segment parameter [0,1] to global parameter space
+	 * @param SegmentIndex - Index of the segment
+	 * @param LocalParam - Local parameter within segment [0,1]
+	 * @return Parameter value in global space, or 0 if segment is invalid
+	 */
+	float MapLocalSegmentParameterToGlobal(int32 SegmentIndex, float LocalParam) const
+	{
+		FInterval1f SegmentRange = GetSegmentParameterRange(SegmentIndex);
+		return SegmentRange.Interpolate(LocalParam);
+	}
+
+	/**
+	 * Maps a global parameter to a segment index and local parameter
+	 * @param GlobalParam - Parameter in global space
+	 * @param OutSegmentIndex - Output segment index
+	 * @param OutLocalParam - Output local parameter [0,1]
+	 * @return True if mapping succeeded
 	 */
-	bool RemoveAnchor(int32 PointIndex)
+	bool MapGlobalParameterToLocalSegment(float GlobalParam, int32& OutSegmentIndex, float& OutLocalParam) const
 	{
-		if (PointIndex < 0 || PointIndex >= GetNumAnchorPoints())
+		const int32 NumSegments = GetNumberOfSegments();
+		const UE::Geometry::FInterval1f ParameterSpace = GetParameterSpace();
+
+		if (ParameterSpace.IsEmpty())
 		{
 			return false;
 		}
 
-		const int32 BaseIdx = AnchorBase(PointIndex);
-		Base::RemoveValue(BaseIdx + 2);
-		Base::RemoveValue(BaseIdx + 1);
-		Base::RemoveValue(BaseIdx + 0);
+		// Left end: snap onto first segment start
+		if (GlobalParam <= ParameterSpace.Min)
+		{
+			OutSegmentIndex = 0;
+			OutLocalParam = 0.0f;
+			return true;
+		}
 
-		Base::RemoveKnot(PointIndex);
+		// Right end: snap onto last segment end
+		if (GlobalParam >= ParameterSpace.Max)
+		{
+			OutSegmentIndex = NumSegments - 1;
+			OutLocalParam = 1.0f;
+			return true;
+		}
 
-		ApplyClampedKnotsMultiplicity();
-		ApplyInternalKnotsMultiplicity();
+		const float LastKnot = Base::PairKnots.Last().Value;
 
-		return true;
-	}
+		// We can now assume GlobalParam is contained by ParameterSpace, no need for clamping.
 
-	// Split a segment using De Casteljau
-	bool SplitSegment(int32 SegmentIndex, float SegmentT)
-	{
-		if (SegmentIndex < 0 || SegmentIndex >= GetNumberOfSegments())
+		// Closed loop: handle the synthetic closing segment [LastKnot, LastKnot + LoopKnotDelta]
+		if (GlobalParam >= LastKnot && ensure(IsActuallyClosed()))
 		{
-			return false;
+			OutSegmentIndex = NumSegments - 1;
+			const float t = (GlobalParam - LastKnot) / FMath::Max(ParameterSpace.Max - LastKnot, UE_SMALL_NUMBER);
+			OutLocalParam = FMath::Clamp(t, 0.0f, 1.0f);
+			return true;
 		}
+        
+		int32 SegmentIndex = INDEX_NONE;
 
-		FInterval1f SegmentRange = GetSegmentParameterRange(SegmentIndex);
-		float Parameter = SegmentRange.Interpolate(SegmentT);
-
-		if (Parameter == SegmentRange.Min)
+		// todo: make CVar
+		constexpr bool bEnableSegmentCaching = true;
+		if (bEnableSegmentCaching && LastEvaluatedSegment != INDEX_NONE && LastEvaluatedSegment < Base::PairKnots.Num() - 1)
 		{
-			Parameter = Param::NextDistinct(Parameter);
+			if (GlobalParam < Base::PairKnots[LastEvaluatedSegment + 1].Value && GlobalParam >= Base::PairKnots[LastEvaluatedSegment].Value)
+			{
+				// cache hit!
+				SegmentIndex = LastEvaluatedSegment;
+			}
 		}
-		else if (Parameter == SegmentRange.Max)
+		
+		if (SegmentIndex == INDEX_NONE)
 		{
-			Parameter = Param::PrevDistinct(Parameter);
+			// cache miss!
+			int32 SegmentLow = 0;
+			int32 SegmentHigh = Base::PairKnots.Num() - 2;
+			SegmentIndex = SegmentLow + (SegmentHigh - SegmentLow) / 2;
+
+			while (SegmentLow <= SegmentHigh)
+			{
+				int32 Mid = SegmentLow + (SegmentHigh - SegmentLow) / 2;
+				const float& StartParam = Base::PairKnots[Mid].Value;
+				const float& EndParam   = Base::PairKnots[Mid + 1].Value;
+
+				if (GlobalParam >= StartParam && GlobalParam < EndParam)
+				{
+					SegmentIndex = Mid;
+					break;
+				}
+				else if (GlobalParam < StartParam)
+				{
+					SegmentHigh = Mid - 1;
+				}
+				else
+				{
+					SegmentLow = Mid + 1;
+				}
+			}
 		}
 		
-		// Bail out if we lost usable range
-		if (Parameter <= SegmentRange.Min || Parameter >= SegmentRange.Max)
+		OutSegmentIndex = SegmentIndex;
+		LastEvaluatedSegment = SegmentIndex;
+        
+		// Calculate local parameter
+		float SegmentLength = Base::PairKnots[SegmentIndex + 1].Value - Base::PairKnots[SegmentIndex].Value;
+		if (SegmentLength > UE_SMALL_NUMBER)
 		{
-			return false;
+			OutLocalParam = (GlobalParam - Base::PairKnots[SegmentIndex].Value) / SegmentLength;
 		}
-
-		ValueType Position = Base::Evaluate(Parameter);
-
-		FWindow Window;
-		if (!BuildRawWindowForSegment(SegmentIndex, Window))
+		else
 		{
-			// If we cannot build a window for segment index, something is wrong and we cannot handle it nicely.
-			return false;
+			OutLocalParam = 0.0f;
 		}
+        
+		return true;
+	}
 
-		ValueType P0 = *Window[0];
-		ValueType P1 = *Window[1];
-		ValueType P2 = *Window[2];
-		ValueType P3 = *Window[3];
-		
-		// Use De Casteljau algorithm to split the curve
-		const float t = SegmentT;
-		const float mt = 1.0f - t;
-		
-		// Calculate split control points
-		ValueType Q0 = P0;
-		ValueType Q1 = P0 * mt + P1 * t;
-		ValueType Q2 = (P0 * mt + P1 * t) * mt + (P1 * mt + P2 * t) * t;
-		ValueType Q3 = Position;  // Use provided position for the split point
-		
-		ValueType R0 = Position;  // Use provided position for the split point
-		ValueType R1 = (P1 * mt + P2 * t) * mt + (P2 * mt + P3 * t) * t;
-		ValueType R2 = P2 * mt + P3 * t;
-		ValueType R3 = P3;
+	int32 FindSegmentIndex(float Parameter, float &OutLocalParam) const
+	{
+		int32 ControlPointIdx = FindIndexForParameter(Parameter, OutLocalParam);
+		return ControlPointIdx / 4; // Each segment has 4 control points
+	}
 
-		// Fixup endpoints of segment being split.
-		TPair<int32, int32> SegmentAnchors = GetSegmentAnchors(SegmentIndex);
-		SetHandleOutOfAnchor(SegmentAnchors.Key, Q1);
-		SetHandleInOfAnchor(SegmentAnchors.Value, R2);
+	void GetBezierIndices(int32 SegmentIndex, int32& OutP0Idx, int32& OutP1Idx, int32& OutP2Idx, int32& OutP3Idx) const
+	{
+		if (!IsClosingSegment(SegmentIndex))
+		{
+			OutP0Idx = P0Idx(SegmentIndex);
+			OutP1Idx = P1Idx(SegmentIndex);
+			OutP2Idx = P2Idx(SegmentIndex);
+			OutP3Idx = P3Idx(SegmentIndex);
+		}
+		else
+		{
+			const int32 NumAnchors = GetNumAnchorPoints();
 
-		// Insert anchor, splitting the segment.
-		InsertAnchor(FAnchor { Q2, Q3, R1 }, Parameter);
+			OutP0Idx = Close_P0Idx(NumAnchors);
+			OutP1Idx = Close_P1Idx(NumAnchors);
+			OutP2Idx = Close_P2Idx();
+			OutP3Idx = Close_P3Idx();
+		}
+	}
+	
+	/** Per-anchor setters/getters */
+	void SetHandleOutOfAnchor(int32 AnchorIdx, const ValueType& Out)
+	{
+		this->SetValue(AnchorBase(AnchorIdx) + 2, Out);
+	}
+	const ValueType& GetHandleOutOfAnchor(int32 AnchorIdx) const
+	{
+		return this->GetValue(AnchorBase(AnchorIdx) + 2);
+	}
+	
+	void SetHandleInOfAnchor(int32 AnchorIdx, const ValueType& In)
+	{
+		this->SetValue(AnchorBase(AnchorIdx) + 0, In);
+	}
+	const ValueType& GetHandleInOfAnchor(int32 AnchorIdx) const
+	{
+		return this->GetValue(AnchorBase(AnchorIdx) + 0);
+	}
 
-		return true;
+	/** Per-anchor structural inserts */
+	void InsertAnchorAt(int32 At, const ValueType& In, const ValueType& A, const ValueType& Out)
+	{
+		this->InsertValue(AnchorBase(At) + 0, In);
+		this->InsertValue(AnchorBase(At) + 1, A);
+		this->InsertValue(AnchorBase(At) + 2, Out);
 	}
 
+	/** Appends a new anchor point defined by its In, A, and Out values. */
+	void AppendAnchor(const ValueType& In, const ValueType& A, const ValueType& Out)
+	{
+		this->AddValue(In);
+		this->AddValue(A);
+		this->AddValue(Out);
+	}
+	
 	virtual int32 GetExpectedNumKnots() const override
 	{
-		const int32 NumAnchors = GetNumAnchorPoints();
-		const int32 NumSegments = GetNumberOfSegments();
+		const int32 NumSegments = GetNumberOfSegments();	
 
-		if (NumAnchors == 0)
+		if (NumSegments == 0)
 		{
 			return 0;
 		}
-		else if (NumAnchors == 1)
-		{
-			if (Base::IsClosedLoop())
-			{
-				return 3;
-			}
-			else
-			{
-				return 4;
-			}
-		}
 
-		if (!HasClosingSegment())
+		if (!IsActuallyClosed())
 		{
 			// Open cubic poly-Bézier:
 			// start multiplicity = 4, end multiplicity = 4, each interior boundary = 3
@@ -1355,7 +1845,7 @@ public:
 	{
 		TArray<FKnot> KnotVector = Base::GetKnotVector();
 
-		if (HasClosingSegment() && ensureAlwaysMsgf(KnotVector.Num() > 0, TEXT("Knot vector is empty, cannot compute the closing knot.")))
+		if (IsActuallyClosed() && ensureAlwaysMsgf(KnotVector.Num() > 0, TEXT("Knot vector is empty, cannot compute the closing knot.")))
 		{
 			KnotVector.Add(FKnot(KnotVector.Last().Value + LoopKnotDelta, Degree));
 		}
@@ -1365,32 +1855,32 @@ public:
 
 	virtual bool SetKnotVector(const TArray<FKnot>& NewKnots) override
 	{
-		if (!Base::ValidateKnotVector(NewKnots))
+		if (!this->ValidateKnotVector(NewKnots))
 		{
 			return false;
 		}
 
-		Base::PairKnots = NewKnots;
+		this->PairKnots = NewKnots;
 
-		if (HasClosingSegment())
+		if (IsActuallyClosed())
 		{
-			const int32 LastIdx     = Base::PairKnots.Num() - 1;
-			const float ClosingKnot = Base::PairKnots[LastIdx].Value;
-			const float PrevKnot    = Base::PairKnots[LastIdx - 1].Value;
+			const int32 LastIdx     = this->PairKnots.Num() - 1;
+			const float ClosingKnot = this->PairKnots[LastIdx].Value;
+			const float PrevKnot    = this->PairKnots[LastIdx - 1].Value;
 
 			// Store the closing interval length and drop the closing knot.
 			LoopKnotDelta = ClosingKnot - PrevKnot;
 			ensure(LoopKnotDelta > 0.f);
 
-			Base::PairKnots.RemoveAt(LastIdx, 1, EAllowShrinking::Yes);
+			this->PairKnots.RemoveAt(LastIdx, 1, EAllowShrinking::Yes);
 		}
 		
-		Base::MarkFlatKnotsCacheDirty();
+		this->MarkFlatKnotsCacheDirty();
 
 		UE_LOG(LogSpline, Verbose, TEXT("Set knot vector (Direct) - Unique Knots: %d"),
-			   Base::PairKnots.Num());
+			   this->PairKnots.Num());
 
-		Base::PrintKnotVector();
+		this->PrintKnotVector();
 
 		return true;
 	}
@@ -1399,22 +1889,7 @@ public:
 	
 private:
 
-	// Segment S is bounded by anchors S and S + 1. Closing segment is bounded by S and 0.
-	TPair<int32, int32> GetSegmentAnchors(int32 SegmentIndex)
-	{
-		if (SegmentIndex < 0 || SegmentIndex >= GetNumberOfSegments())
-		{
-			return { INDEX_NONE, INDEX_NONE };
-		}
-
-		if (IsClosingSegment(SegmentIndex))
-		{
-			return { SegmentIndex, 0 };
-		}
-
-		return { SegmentIndex, SegmentIndex + 1 };
-	}
-
+	
 	static constexpr int32 AnchorBase(int32 i){ return 3*i; }
 
 	static constexpr int32 P0Idx(int32 s) { return AnchorBase(s) + 1; } // A_s
@@ -1448,7 +1923,7 @@ private:
 			return false;
 		}
 
-		auto I = [this](int32 idx){ return &Base::Values[idx]; };
+		auto I = [this](int32 idx){ return &this->Values[idx]; };
 
 		if (!IsClosingSegment(SegmentIndex))
 		{
@@ -1468,6 +1943,21 @@ private:
 		return Math::FSplineValidation::IsValidWindow<ValueType, Base::WindowSize>(OutWindow);
 	}
 
+	void RefreshClosingKnot(EParameterizationPolicy Policy)
+	{
+		if (!IsActuallyClosed())
+		{
+			LoopKnotDelta = 0.0f;
+			return;
+		}
+
+		const ValueType& P0 = Base::Values[ AnchorBase(GetNumAnchorPoints() - 1) + 1 ]; // A_last
+		const ValueType& P3 = Base::Values[ AnchorBase(0)                        + 1 ]; // A_0
+		LoopKnotDelta = ComputeSegmentLength(P0, P3, Policy);
+
+		Base::MarkFlatKnotsCacheDirty();
+	}
+
 	void MigrateLegacyBezierLayout()
 	{
 	    const int32 NumValues = Base::Values.Num();
@@ -1515,11 +2005,11 @@ private:
 	    checkf(NumValues % 4 == 0,
 	        TEXT("Legacy PolyBezier layout mismatch (NumValues = %d)"), NumValues);
 
-		const bool bClosed = Base::IsClosedLoop();
+		const bool bClosed = this->IsClosedLoop();
 
 	    const int32 NumSegments = NumValues / 4;
 
-	    const TArray<ValueType>& Legacy = Base::Values;
+	    const TArray<ValueType>& Legacy = this->Values;
 		// For open: anchors = NumSegments + 1
 		// For closed: anchors = NumSegments (closing segment is virtual)
 		const int32 NumAnchors = bClosed ? NumSegments : (NumSegments + 1);
@@ -1616,25 +2106,25 @@ private:
 	        NewValues[NewBase + 2] = Out;
 	    }
 
-	    Base::Values = MoveTemp(NewValues);
+	    this->Values = MoveTemp(NewValues);
 
 		// Fix up knot representation for closed loops.
 		// Legacy data uses an explicit closing knot in PairKnots; the new
 		// representation uses LoopKnotDelta instead.
-		if (bClosed && Base::PairKnots.Num() > 2)
+		if (bClosed && this->PairKnots.Num() > 2)
 		{
-			const int32 LastIdx     = Base::PairKnots.Num() - 1;
-			const float ClosingKnot = Base::PairKnots[LastIdx].Value;
-			const float PrevKnot    = Base::PairKnots[LastIdx - 1].Value;
+			const int32 LastIdx     = this->PairKnots.Num() - 1;
+			const float ClosingKnot = this->PairKnots[LastIdx].Value;
+			const float PrevKnot    = this->PairKnots[LastIdx - 1].Value;
 
 			// Store the closing interval length and drop the closing knot.
 			LoopKnotDelta = ClosingKnot - PrevKnot;
 
-			Base::PairKnots.RemoveAt(LastIdx, 1, EAllowShrinking::Yes);
+			this->PairKnots.RemoveAt(LastIdx, 1, EAllowShrinking::Yes);
 		}
 		
 		// Values changed (and potentially PairKnots layout), so flat-knot cache is dirty.
-	    Base::MarkFlatKnotsCacheDirty();
+	    this->MarkFlatKnotsCacheDirty();
 	}
 
 	/**
@@ -1652,24 +2142,15 @@ private:
 		// Values[3*i + 2] = Out_i
 
 		const int32 NumAnchors  = GetNumAnchorPoints();
+		const int32 NumSegments = GetNumberOfSegments();
 
-		if (NumAnchors <= 0)
-		{
-			Base::PairKnots.Reset();
-			Base::MarkFlatKnotsCacheDirty();
-			return;
-		}
-		else if (NumAnchors == 1)
+		// Need at least one segment
+		if (NumAnchors < 2 || NumSegments <= 0)
 		{
-			// If we have a single anchor, our number of expected knots is itself the multiplicity of a single distinct knot.
-			Base::PairKnots = { FKnot(0.f, GetExpectedNumKnots()) };
-			Base::MarkFlatKnotsCacheDirty();
+			this->PairKnots.Reset();
+			this->MarkFlatKnotsCacheDirty();
 			return;
 		}
-
-		// If we have > 1 anchors (guaranteed by the above branches), we ought to have at least 1 segment.
-		const int32 NumSegments = GetNumberOfSegments();
-		ensureAlways(NumSegments > 0);
 			
 		TArray<float> SegmentLengths;
 		SegmentLengths.Reserve(NumSegments);
@@ -1677,7 +2158,7 @@ private:
 		// For a Bezier spline, the control points don't all lie on the curve
 		// So we need to compute chord lengths between segment endpoints only.
 		// Segment i goes from A_i to A_{i+1} (open) or A_{(i+1)%N} (closed).
-		const bool bClosed = HasClosingSegment();
+		const bool bClosed = IsActuallyClosed();
 
 		for (int32 i = 0; i < NumSegments; ++i)
 		{
@@ -1686,8 +2167,8 @@ private:
 				? (i + 1) % NumAnchors
 				: (i + 1); // for open, NumSegments == NumAnchors - 1, so this is always valid
 	            
-			const ValueType& P0 = Base::Values[AnchorBase(StartAnchorIdx) + 1]; // A_i
-			const ValueType& P3 = Base::Values[AnchorBase(EndAnchorIdx)   + 1]; // A_{i+1} or A_0 (closed)
+			const ValueType& P0 = this->Values[AnchorBase(StartAnchorIdx) + 1]; // A_i
+			const ValueType& P3 = this->Values[AnchorBase(EndAnchorIdx)   + 1]; // A_{i+1} or A_0 (closed)
 
 			SegmentLengths.Add(ComputeSegmentLength(P0, P3, Mode));
 		}
@@ -1712,18 +2193,18 @@ private:
 		};
 		
 		// Reset pair knots
-		Base::PairKnots.Reset();
+		this->PairKnots.Reset();
         
-		FValidKnotSearchParams SearchParams;
+		typename Base::FValidKnotSearchParams SearchParams;
 		SearchParams.bSearchLeft = false;	// we are appending repeatedly, we want insertion order to be predictable (always go after conflicting knots).
 
-		FKnotSearchCache SearchCache{ Base::PairKnots };
+		typename Base::FKnotSearchCache SearchCache{ this->PairKnots };
 
-		const auto InsertKnotAndUpdateCache = [this, &SearchCache, &SearchParams](uint32 InMultiplicity)
+		const auto InsertKnotAndUpdateCache = [&](uint32 InMultiplicity)
 		{
-			const float KnotValue = Base::GetNearestAvailableKnotValue(SearchParams, SearchCache);
+			const float KnotValue = this->GetNearestAvailableKnotValue(SearchParams, SearchCache);
 			SearchCache.KnotValuesSet.Add(Param::NormalizeKey(KnotValue));
-			Base::InsertKnot(FKnot(KnotValue, InMultiplicity));
+			this->InsertKnot(FKnot(KnotValue, InMultiplicity));
 		};
 
 		constexpr int32 EndMult = Degree + 1;
@@ -1795,7 +2276,7 @@ private:
         // --------------------------------------------------------------------
 
 		// Mark flat knots cache as dirty
-		Base::MarkFlatKnotsCacheDirty();
+		this->MarkFlatKnotsCacheDirty();
 	}
 
     virtual void FindWindow(float Parameter, FWindow& OutWindow, FWindowStore& Scratch) const override
@@ -1807,10 +2288,9 @@ private:
 			return;
 		}
 
+		int32 SegmentIndex;
 		float LocalT;
-		int32 SegmentIndex = FindIndexForParameter(Parameter, LocalT);
-
-		if (SegmentIndex == INDEX_NONE)
+		if (!MapGlobalParameterToLocalSegment(Parameter, SegmentIndex, LocalT))
 		{
 			return;
 		}
@@ -1834,10 +2314,9 @@ private:
 			return Math::FSplineValidation::GetDefaultValue<ValueType>();
         }
 
-		float LocalT;
-		int32 SegmentIndex = FindIndexForParameter(Parameter, LocalT);
-
-		if (SegmentIndex == INDEX_NONE)
+    	int32 SegmentIndex;
+    	float LocalT;
+		if (!MapGlobalParameterToLocalSegment(Parameter, SegmentIndex, LocalT))
 		{
 			return Math::FSplineValidation::GetDefaultValue<ValueType>();
 		}
@@ -1882,9 +2361,9 @@ private:
 	/**
 	 * Checks if the spline is actually closed, which is semantically different from IsClosedLoop for 0 and 1 point splines.
 	 */
-	bool HasClosingSegment() const
+	bool IsActuallyClosed() const
 	{
-		return Base::IsClosedLoop() && GetNumAnchorPoints() > 1;
+		return this->IsClosedLoop() && GetNumAnchorPoints() > 1;
 	}
 
 	/**
@@ -1894,7 +2373,7 @@ private:
 	 */
 	bool IsClosingSegment(int32 SegmentIndex) const
 	{
-		return HasClosingSegment() && (SegmentIndex == GetNumAnchorPoints() - 1);
+		return IsActuallyClosed() && (SegmentIndex == GetNumAnchorPoints() - 1);
 	}
 
 	/**
@@ -1907,6 +2386,103 @@ private:
     	int32 NumSegments = GetNumberOfSegments();
     	return Index >= 0 && Index < NumSegments;
     }
+
+	/**
+	 * Helper method to add the first segment to the spline
+	 * @param ControlPoint - Bezier control point to insert
+	 * @param OutPointIndex - Index of the newly created bezier anchor point
+	 * @return true if the points were added, false otherwise
+	 */
+	bool InitializeFirstSegment(
+		const ValueType& ControlPoint,
+		int32& OutPointIndex)
+    {
+		const int32 NumAnchors = GetNumAnchorPoints();
+		// 1) No anchors yet: create A0 only, no segment.
+		if (NumAnchors == 0)
+		{
+			// Zero handles: In0 = A0, Out0 = A0
+			AppendAnchor(/*In0*/ ControlPoint, /*A0*/ ControlPoint, /*Out0*/ ControlPoint);
+			OutPointIndex = 0;
+			return true;
+		}
+
+		// 2) One anchor exists, adding this point will create the first segment.
+		if (NumAnchors == 1)
+		{
+			const ValueType& A0 = Base::Values[P0Idx(0)];
+
+			// Simple auto-tangent guess between A0 and new point
+			ValueType Dir = ControlPoint - A0;
+			ValueType P0  = A0;
+			ValueType P3  = ControlPoint;
+			ValueType P1  = P0 + Dir * ValueType(0.33f);
+			ValueType P2  = P0 + Dir * ValueType(0.66f);
+	        
+    		// Add the first segment
+    		SetControlPoints({P0, P1, P2, P3}, EParameterizationPolicy::Centripetal);
+    		OutPointIndex = 1;
+    		return true;
+    	}
+
+		// 3) We already have at least one segment – let caller take the normal path.
+    	return false;
+    }
+
+	/**
+     * Prepends a point to the spline at the beginning.
+     * The new point will be connected to the first segment.
+     * @param Position The position of the new point
+     * @param ParameterizationPolicy How to distribute Knots
+     * @return Index of the new segment
+     */
+    int32 PrependPoint(
+		const ValueType& Position,
+		EParameterizationPolicy ParameterizationPolicy = EParameterizationPolicy::Centripetal)
+	{
+		const ValueType& P3 = this->Values[AnchorBase(0) + 1];
+
+		// Calculate reasonable auto-tangents
+		ValueType Dir = P3 - Position;
+		ValueType P1 = Position + Dir * 0.33f;
+		ValueType P2 = Position + Dir * 0.66f;
+    
+		// Add the segment by reusing existing code
+		int32 NewIndex = AddBezierSegmentInternal(0, Position, P1, P2, P3, ParameterizationPolicy, false);
+		check (NewIndex == 0);
+		return NewIndex;
+	}
+
+	/**
+	 * Appends a point to the end of the spline.
+	 * Creates a new cubic Bezier segment connecting the last point to the new point.
+	 * 
+	 * @param Position The position of the new point
+	 * @param ParameterizationPolicy How to distribute Knots
+	 * @return Index of the new segment
+	 */
+	int32 AppendPoint(
+		const ValueType& Position,
+		EParameterizationPolicy ParameterizationPolicy)
+	{
+		const int32 NumAnchors = GetNumAnchorPoints();
+		const int32 LastAnchorIdx = NumAnchors - 1;
+
+		// Append - create a point at the end
+		const ValueType& P0 = this->Values[AnchorBase(LastAnchorIdx) + 1];
+
+        
+		// Calculate reasonable auto-tangents
+		ValueType Dir = Position - P0;
+		ValueType P1 = P0 + Dir * 0.33f;
+		ValueType P2 = P0 + Dir * 0.66f;
+
+		int32 NumSegments = GetNumberOfSegments();
+		// Add the segment
+		int32 NewIndex = AddBezierSegmentInternal(NumSegments, P0, P1, P2, Position, ParameterizationPolicy/*, StartParam*/);
+
+		return NewIndex;
+	}
 	
 	/**
 	 * Helper to compute segment length based on parameterization policy
@@ -1940,6 +2516,116 @@ private:
 		}
 	}
 
+	int32 AddBezierSegmentInternal(
+		int32 SegmentIndex,
+		const ValueType& P0,
+		const ValueType& P1, 
+		const ValueType& P2, 
+		const ValueType& P3,
+		EParameterizationPolicy ParameterizationPolicy = EParameterizationPolicy::Centripetal,
+		bool bAppendToSegment = true) // true = connect to end, false = connect to start
+	{
+		const int32 NumAnchors  = GetNumAnchorPoints();
+	    const int32 NumSegments = GetNumberOfSegments();
+
+	    // Clamp requested slot into legal insertion range
+	    SegmentIndex = FMath::Clamp(SegmentIndex, 0, NumSegments);
+
+		auto FixKnots = [this](EParameterizationPolicy Policy)
+		{
+			this->ApplyClampedKnotsMultiplicity();
+			ApplyInternalKnotsMultiplicity();
+
+			if (IsActuallyClosed())
+			{
+				RefreshClosingKnot(Policy);
+			}
+		};
+
+	    // Measure for spacing when we do need to place a new boundary knot
+	    const float SegLen = ComputeSegmentLength(P0, P3, ParameterizationPolicy);
+
+	    // ===== CASE A: first-ever segment (no anchors yet) =====
+	    if (NumAnchors == 0)
+	    {
+	        // Seed A0 from P0.
+	        // IMPORTANT: make In0 a zero-handle by setting it equal to A0.
+	        // (We cannot infer a real incoming handle for A0 from {P0,P1,P2,P3} of the *first* segment.)
+	        AppendAnchor(/*In0*/ P0, /*A0*/ P0, /*Out0*/ P1);
+
+	        // Append A1 from {P2,P3}. Out1 is a zero-handle placeholder (equal to A1)
+	        AppendAnchor(/*In1*/ P2, /*A1*/ P3, /*Out1*/ P3);
+
+	        // Build a minimal clamped knot vector for 1 segment (open):
+	        this->ResetKnotVector();
+	        const int32 EndMult = Degree + 1;
+	        this->InsertKnot(FKnot(0.0f, EndMult));
+	        this->InsertKnot(FKnot(SegLen, EndMult));
+	    	FixKnots(ParameterizationPolicy);
+
+	        return 0; // the first (and only) segment index
+	    }
+
+	    // ===== CASE B: append at tail =====
+	    const bool bAppendAtTail = bAppendToSegment && (SegmentIndex >= NumSegments);
+	    if (bAppendAtTail)
+	    {
+	        const int32 Last = NumAnchors - 1;
+
+	        // Set Out_last from P1
+	        SetHandleOutOfAnchor(Last, P1);
+
+	        // Add new tail anchor from {P2,P3}; Out_new is zero-handle until more segments arrive
+	        AppendAnchor(/*In_n*/ P2, /*A_n*/ P3, /*Out_n*/ P3);
+
+	        // Insert a new boundary knot at the end (incremental)
+	        const float Kprev         = this->PairKnots.IsEmpty() ? 0.f : this->PairKnots.Last().Value;
+	        const float NextByLength  = Kprev + SegLen;
+	        const float NextByStep    = UE::Geometry::Spline::Param::NextDistinct(Kprev);
+	        const float Knew          = FMath::Max(NextByLength, NextByStep);
+	        this->InsertKnot(FKnot(Knew, Degree));
+
+	    	FixKnots(ParameterizationPolicy);
+	        return NumSegments; // newly appended segment index
+	    }
+
+	    // ===== CASE C: prepend at head =====
+	    const bool bPrependAtHead = (!bAppendToSegment && SegmentIndex == 0);
+	    if (bPrependAtHead)
+	    {
+	        // Insert new head anchor built from P0,P1
+	        InsertAnchorAt(/*At*/ 0, /*In_-1*/ P0, /*A_-1*/ P0, /*Out_-1*/ P1);
+
+	        // Old head is now anchor 1 — set its In_1 from P2
+	        SetHandleInOfAnchor(/*i=1*/ 1, P2);
+
+	        // Insert a boundary knot at the start (decremental)
+	        const float Kfirst        = this->PairKnots.IsEmpty() ? 0.f : this->PairKnots[0].Value;
+	        const float PrevByLength  = Kfirst - SegLen;
+	        const float PrevByStep    = UE::Geometry::Spline::Param::PrevDistinct(Kfirst);
+	        const float Knew          = FMath::Min(PrevByLength, PrevByStep);
+	        this->InsertKnot(FKnot(Knew, Degree));
+
+	    	FixKnots(ParameterizationPolicy);
+	        return 0;
+	    }
+
+	    // ===== CASE D: insert in the middle (between A_s and A_{s+1}) =====
+	    //   We don’t touch knots here; caller will InsertKnot(TargetParam).
+	    {
+	        const int32 s = FMath::Clamp(SegmentIndex, 0, FMath::Max(0, NumSegments - 1));
+
+	        // Left anchor’s Out_s comes from P1
+	        SetHandleOutOfAnchor(s, P1);
+
+	        // Insert new anchor at s+1 from {P2,P3}; Out_new is zero-handle
+	        InsertAnchorAt(/*At*/ s+1, /*In*/ P2, /*A*/ P3, /*Out*/ P3);
+
+			FixKnots(ParameterizationPolicy);
+	        return s;
+	    }
+	}
+
 private:
 
 	/** This is the last segment fetched by parameter. Used to accelerate predictable sampling patterns. */
@@ -1948,7 +2634,7 @@ private:
 	/** Optional override for FindWindow*/
 	FFindWindowOverride FindWindowOverrideFn = nullptr;
 
-	/** Length of the closing segment’s parameter interval. Only meaningful when HasClosingSegment() == true. */
+	/** Length of the closing segment’s parameter interval. Only meaningful when IsActuallyClosed() == true. */
 	float LoopKnotDelta = 1.f;
 };
 
diff --git a/Engine/Source/Runtime/GeometryCore/Public/Splines/SplineInterfaces.h b/Engine/Source/Runtime/GeometryCore/Public/Splines/SplineInterfaces.h
--- a/Engine/Source/Runtime/GeometryCore/Public/Splines/SplineInterfaces.h
+++ b/Engine/Source/Runtime/GeometryCore/Public/Splines/SplineInterfaces.h
@@ -15,7 +15,7 @@ namespace Spline
 
 // Forward declarations
 class UE_EXPERIMENTAL(5.7, "New spline APIs are experimental.") ISplineInterface;
-template <typename ValueType> class UE_EXPERIMENTAL(5.7, "New spline APIs are experimental.") TSplineInterface;
+template <typename VALUETYPE> class UE_EXPERIMENTAL(5.7, "New spline APIs are experimental.") TSplineInterface;
 
 /**
  * Enumeration for out-of-bounds handling modes
@@ -87,13 +87,13 @@ public:
  * Typed interface for splines with a specific value type
  * Provides evaluation and other type-specific functionality
  */
-template <typename InValueType>
+template <typename VALUETYPE>
 class TSplineInterface : public ISplineInterface
 {
     static FString TypeName;
 
 public:
-    using ValueType = InValueType;
+    typedef VALUETYPE ValueType;
 
     /**
      * Default constructor - ensures value type name is set
@@ -103,7 +103,7 @@ public:
         // Ensure the value type name is initialized if not already
         if (TypeName.IsEmpty())
         {
-            TypeName = TSplineValueTypeTraits<ValueType>::Name;
+            TypeName = TSplineValueTypeTraits<VALUETYPE>::Name;
         }
     }
     
@@ -234,8 +234,8 @@ protected:
 };
 
 // Initialize static member
-template <typename ValueType>
-inline FString TSplineInterface<ValueType>::TypeName;
+template <typename VALUETYPE>
+inline FString TSplineInterface<VALUETYPE>::TypeName;
 
 } // namespace Spline
 } // namespace Geometry
diff --git a/Engine/Source/Runtime/GeometryCore/Public/Splines/TangentBezierSpline.h b/Engine/Source/Runtime/GeometryCore/Public/Splines/TangentBezierSpline.h
--- a/Engine/Source/Runtime/GeometryCore/Public/Splines/TangentBezierSpline.h
+++ b/Engine/Source/Runtime/GeometryCore/Public/Splines/TangentBezierSpline.h
@@ -11,8 +11,8 @@ namespace Geometry
 namespace Spline
 {
 	
-template<typename ValueType> struct UE_EXPERIMENTAL(5.7, "New spline APIs are experimental.") TTangentBezierControlPoint;
-template<typename ValueType> class UE_EXPERIMENTAL(5.7, "New spline APIs are experimental.") TTangentBezierSpline;
+template<typename VALUETYPE> struct UE_EXPERIMENTAL(5.7, "New spline APIs are experimental.") TTangentBezierControlPoint;
+template<typename VALUETYPE> class UE_EXPERIMENTAL(5.7, "New spline APIs are experimental.") TTangentBezierSpline;
 	
 /**
  * Enum defining how tangents are computed for a spline control point.
@@ -135,24 +135,24 @@ struct TTangentBezierControlPoint
  * piecewise Bezier curves internally for evaluation. Supports both manual
  * tangent control and automatic tangent computation.
  */
-template<typename InValueType>
+template<typename VALUETYPE>
 class TTangentBezierSpline :
-	public TSplineWrapper<TPolyBezierSpline<InValueType>>,
-	private TSelfRegisteringSpline<TTangentBezierSpline<InValueType>, InValueType>
+	public TSplineWrapper<TPolyBezierSpline<VALUETYPE>>,
+	private TSelfRegisteringSpline<TTangentBezierSpline<VALUETYPE>, VALUETYPE>
 {
 public:
+	typedef typename TSplineWrapper<TPolyBezierSpline<VALUETYPE>>::ValueType ValueType;
 	
-	using ValueType 					= InValueType;
-	using Base 							= TPolyBezierSpline<ValueType>;
-	using FTangentBezierControlPoint 	= TTangentBezierControlPoint<ValueType>;
-	using FWindow      					= typename Base::FWindow;
-	using FWindowStore 					= typename Base::FWindowStore;
-	using TangentOps 					= Math::TTangentOps<ValueType>;
+	using FTangentBezierControlPoint = TTangentBezierControlPoint<ValueType>;
+	using FPolyBezierSpline = TPolyBezierSpline<ValueType>;
+	using FWindow      = typename FPolyBezierSpline::FWindow;
+	using FWindowStore = typename FPolyBezierSpline::FWindowStore;
+	using TangentOps = Math::TTangentOps<ValueType>;
 	
 	// Generate compile-time type ID for TangentBezier
 	DECLARE_SPLINE_TYPE_ID(
 		TEXT("TangentBezier"),
-		*TSplineValueTypeTraits<ValueType>::Name
+		*TSplineValueTypeTraits<VALUETYPE>::Name
 	);
 	
 	TTangentBezierSpline()
@@ -169,7 +169,27 @@ public:
 	{
 		Tension = 0.f;
 		TangentModes = {ETangentMode::Auto, ETangentMode::Auto};
-		AssignInternalSpline(Base::CreateLine(StartPoint, EndPoint));
+		AssignInternalSpline(FPolyBezierSpline3d::CreateLine(StartPoint, EndPoint));
+	}
+
+	/**
+	 * Constructor with full control over tangents
+	 */
+	TTangentBezierSpline(
+		const ValueType& StartPoint, 
+		const ValueType& EndPoint,
+		const ValueType& StartTangent,
+		const ValueType& EndTangent,
+		bool bAutoTangents = false)
+			: TTangentBezierSpline()
+	{
+		Tension = 0.f;
+		TangentModes = {bAutoTangents ? ETangentMode::Auto : ETangentMode::User,
+						bAutoTangents ? ETangentMode::Auto : ETangentMode::User};
+		AssignInternalSpline( FPolyBezierSpline3d(StartPoint,
+											TangentOps::TangentOutToP1(StartPoint, StartTangent),
+											TangentOps::TangentInToP2(EndPoint, EndTangent),
+											EndPoint));
 	}
    
     /**  Copy constructor */
@@ -296,15 +316,6 @@ public:
 				bStationaryEndpoint == Other.bStationaryEndpoint;
 	}
 
-	static Base::FAnchor AnchorFromControlPoint(const FTangentBezierControlPoint& InPoint)
-	{
-		return {
-			TangentOps::TangentInToP2(InPoint.Position, InPoint.TangentIn),
-			InPoint.Position,
-			TangentOps::TangentOutToP1(InPoint.Position, InPoint.TangentOut)
-		};
-	}
-
 	// Static shape generators
 
     /**
@@ -321,7 +332,7 @@ public:
 		TTangentBezierSpline Result;
         
 		// Clear the default initialization
-		Result.AssignInternalSpline(Base::CreateLine(StartPoint, EndPoint));
+		Result.AssignInternalSpline(FPolyBezierSpline3d::CreateLine(StartPoint, EndPoint));
 		
         Result.Reparameterize(EParameterizationPolicy::Uniform);
         return Result;
@@ -347,7 +358,7 @@ public:
         TTangentBezierSpline Result;
         
         // Clear the default initialization
-        Result.AssignInternalSpline(Base::CreateCircleArc(
+        Result.AssignInternalSpline(FPolyBezierSpline3d::CreateCircleArc(
             Center, Radius, StartAngle, EndAngle, NumSegments));
         
 		// Set up tangent status for all points
@@ -374,7 +385,7 @@ public:
     {
         // Create the circle
         TTangentBezierSpline Result;
-        Result.AssignInternalSpline(Base::CreateCircle(
+        Result.AssignInternalSpline(FPolyBezierSpline3d::CreateCircle(
             ValueType(Center), Radius, NumSegments));
         
         // Set up tangent status for all points
@@ -402,7 +413,7 @@ public:
     {
         // Create the ellipse
         TTangentBezierSpline Result;
-        Result.AssignInternalSpline(Base::CreateEllipse(
+        Result.AssignInternalSpline(FPolyBezierSpline3d::CreateEllipse(
             Center, RadiusX, RadiusY, NumSegments));
         
 		// Set up tangent status for all points
@@ -413,10 +424,9 @@ public:
         
         return Result;
     }
-
 	virtual TUniquePtr<ISplineInterface> Clone() const override
 	{
-		TUniquePtr<TTangentBezierSpline<ValueType>> Clone = MakeUnique<TTangentBezierSpline<ValueType>>();
+		TUniquePtr<TTangentBezierSpline<VALUETYPE>> Clone = MakeUnique<TTangentBezierSpline<VALUETYPE>>();
     
 		// Copy internal spline
 		Clone->AssignInternalSpline(this->InternalSpline);
@@ -435,206 +445,362 @@ public:
     
 	/**
 	 * Replaces all points in the spline with the provided points.
+	 * 
 	 * @param Points - Array of control points to set
-	 * @param Policy - The policy to use for knot generation as the points are added
 	 * @return True if the points were successfully set
 	 */
-	bool SetPoints(const TArray<FTangentBezierControlPoint>& Points, EParameterizationPolicy Policy)
-	{		
+	bool SetControlPoints(const TArray<FTangentBezierControlPoint>& Points,
+		EParameterizationPolicy ParameterizationPolicy = EParameterizationPolicy::Uniform)
+	{
+	    const int32 NumPoints = Points.Num();
+		
 		Clear();
-		TangentModes.Reset(Points.Num());
-
-		for (const FTangentBezierControlPoint& Point : Points)
-		{
-			AppendPoint(Point, Policy);
-		}
+		
+	    // Special case for less than 2 points
+	    if (NumPoints < 2)
+	    {
+	    	if (NumPoints == 1)
+	    	{
+	    		// Initialize the spline with a single point
+	    		int32 OutNewIndex;
+	    		return InitializeFirstSegment(Points[0], OutNewIndex);
+	    	}
+	    	return true;
+	    }
+	    
+	    // Convert to cubic Bezier control points for the internal spline
+	    TArray<ValueType> BezierPoints;
+	    
+	    // For each segment (between consecutive points), we need 4 control points
+	    // Reserve space based on number of segments
+	    const int32 NumSegments = NumPoints - 1;
+	    BezierPoints.Reserve(NumSegments * 4);
+	    
+	    for (int32 i = 0; i < NumSegments; ++i)
+	    {
+	        const FTangentBezierControlPoint& StartPoint = Points[i];
+	        const FTangentBezierControlPoint& EndPoint = Points[i + 1];
+	        
+	        ValueType P0 = StartPoint.Position;
+	        ValueType P3 = EndPoint.Position;
+	    	
+	    	ValueType OutTangent = StartPoint.TangentOut;
+	    	ValueType InTangent = EndPoint.TangentIn;
+	       	        
+	        // Calculate bezier control points
+	    	ValueType P1 = TangentOps::TangentOutToP1(P0, OutTangent);
+	    	ValueType P2 = TangentOps::TangentInToP2(P3, InTangent);
+	        
+	        // Add all 4 control points for this segment
+	        BezierPoints.Add(P0);
+	        BezierPoints.Add(P1);
+	        BezierPoints.Add(P2);
+	        BezierPoints.Add(P3);
+	    }
+	    
+	    // Set the control points in the internal spline
+	    InternalSpline.SetControlPoints(BezierPoints, ParameterizationPolicy);
+	    
+	    // Store tangent modes for later use
+	    TangentModes.Reset(NumPoints);
+	    for (const FTangentBezierControlPoint& Point : Points)
+	    {
+	        TangentModes.Add(Point.TangentMode);
+	    }
 	    	    
 	    return true;
 	}
 
 	/**
-	 * Appends multiple points to the spline.
+	 * Adds multiple points to the spline.
+	 * 
 	 * @param Points - Array of control points to add
-	 * @param Policy - The policy to use for knot generation as the points are added
+	 * @param bAppend - If true, append to the end of spline; if false, prepend to the beginning
 	 * @return True if the points were successfully added
 	 */
-	bool AppendPoints(const TArray<FTangentBezierControlPoint>& Points, EParameterizationPolicy Policy)
+	 bool AddControlPoints(
+        const TArray<FTangentBezierControlPoint>& Points,
+        bool bAppend = true)
     {
         const int32 NumPoints = Points.Num();
         
         // Need at least 2 points to add a valid segment
         if (NumPoints < 2)
         {
-            UE_LOG(LogSpline, Warning, TEXT("AppendPoints requires at least 2 points to add a valid segment. Got %d points."), NumPoints);
+            UE_LOG(LogSpline, Warning, TEXT("AddControlPoints requires at least 2 points to add a valid segment. Got %d points."), NumPoints);
             return false;
         }
         
         // If the spline is currently empty (though it shouldn't be based on our design),
-        // this is equivalent to SetPoints
+        // this is equivalent to SetControlPoints
         if (GetNumPoints() == 0)
         {
-            return SetPoints(Points, Policy);
+            return SetControlPoints(Points);
         }
-
-		// For each point, use AddPoint 
-		for (int32 i = 0; i < NumPoints; ++i)
-		{
-			// Skip the first point if it's the same as our last existing point 
-			// (to avoid duplicates at connection point)
-			if (i == 0)
-			{
-				const int32 ExistingNumPoints = GetNumPoints();
-				const ValueType LastExistingPoint = GetValue(ExistingNumPoints - 1);
+        
+        if (bAppend)
+        {
+            // For each point, use AddPoint 
+            for (int32 i = 0; i < NumPoints; ++i)
+            {
+                // Skip the first point if it's the same as our last existing point 
+                // (to avoid duplicates at connection point)
+                if (i == 0)
+                {
+                    const int32 ExistingNumPoints = GetNumPoints();
+                    const ValueType LastExistingPoint = GetValue(ExistingNumPoints - 1);
                     
-				// Check if connection points are reasonably close
-				constexpr double ConnectionTolerance = 1e-4;
-				bool bSameEndpoint;
-
-				if constexpr (Math::TIsQuaternionType<ValueType>::value)
-				{
-					// angular distance
-					const double dot = Math::Dot(LastExistingPoint, Points[0].Position);
-					const double ang = FMath::Acos(FMath::Clamp(dot, -1.0, 1.0)); // radians
-					bSameEndpoint = (ang <= ConnectionTolerance);
-				}
-				else
-				{
-					const double DistanceSquared = Math::SizeSquared(LastExistingPoint - Points[0].Position);
-					bSameEndpoint = (DistanceSquared <= ConnectionTolerance);
-				}
+                    // Check if connection points are reasonably close
+                	constexpr double ConnectionTolerance = 1e-4;
+                	bool bSameEndpoint;
+
+                	if constexpr (Math::TIsQuaternionType<ValueType>::value)
+                	{
+                		// angular distance
+                		const double dot = Math::Dot(LastExistingPoint, Points[0].Position);
+						const double ang = FMath::Acos(FMath::Clamp(dot, -1.0, 1.0)); // radians
+						bSameEndpoint = (ang <= ConnectionTolerance);
+                	}
+                	else
+                	{
+                		const double DistanceSquared = Math::SizeSquared(LastExistingPoint - Points[0].Position);
+                		bSameEndpoint = (DistanceSquared <= ConnectionTolerance);
+                	}
                     
-				if (bSameEndpoint)
-				{
-					// Points are close - the first existing point is the same as our first new point
-					// So update the last point's tangent mode if needed
-					if (Points[0].TangentMode != ETangentMode::User && Points[0].TangentMode != ETangentMode::Broken)
-					{
-						SetPointTangentMode(ExistingNumPoints - 1, Points[0].TangentMode);
-					}
-					continue; // Skip adding this point
-				}
-			}
+                    if (bSameEndpoint)
+                    {
+                        // Points are close - the first existing point is the same as our first new point
+                        // So update the last point's tangent mode if needed
+                        if (Points[0].TangentMode != ETangentMode::User && Points[0].TangentMode != ETangentMode::Broken)
+                        {
+                            SetPointTangentMode(ExistingNumPoints - 1, Points[0].TangentMode);
+                        }
+                        continue; // Skip adding this point
+                    }
+                }
                 
-			// Add each point normally
-			AppendPoint(Points[i], Policy);
-		}
+                // Add each point normally
+                AppendPoint(Points[i]);
+            }
+        }
+        else
+        {
+            // For prepending, we need to add points in reverse order
+            for (int32 i = NumPoints - 1; i >= 0; --i)
+            {
+                // Skip the last point if it's the same as our first existing point
+                if (i == NumPoints - 1)
+                {
+                    const ValueType FirstExistingPoint = GetValue(0);
+                    
+                    // Check if connection points are reasonably close
+                    constexpr double ConnectionTolerance = 1e-4;
+                	bool bSameEndpoint;
+
+                	if constexpr (Math::TIsQuaternionType<ValueType>::value)
+                	{
+                		// angular distance
+                		const double dot = Math::Dot(FirstExistingPoint, Points[i].Position);
+                		const double ang = FMath::Acos(FMath::Clamp(dot, -1.0, 1.0)); // radians
+                		bSameEndpoint = (ang <= ConnectionTolerance); // choose a small angular eps, e.g. ~1e-4 rad
+                	}
+	                else
+	                {
+		                const double DistanceSquared = Math::SizeSquared(FirstExistingPoint - Points[i].Position);
+	                	bSameEndpoint = (DistanceSquared <= ConnectionTolerance);
+	                }
+                    
+                    if (bSameEndpoint)
+                    {
+                        // Points are close - update first point's tangent mode if needed
+                        if (Points[i].TangentMode != ETangentMode::User && Points[i].TangentMode != ETangentMode::Broken)
+                        {
+                            SetPointTangentMode(0, Points[i].TangentMode);
+                        }
+                        continue; // Skip adding this point
+                    }
+                }
+                
+                // Prepend each point
+                PrependPoint(Points[i]);
+            }
+        }
+        
+        // Reparameterize for consistent parameter distribution
+        Reparameterize();
         
         return true;
     }
-    
-	/**
-	 * Prepends a control point to the beginning of the spline.
-	 * @param Point - Control point to prepend
-	 * @param Parameter - The parameter of the control point being prepended, must be less than all existing knots.
-	 * @return Index of the new point, or INDEX_NONE if the operation failed.
-	 */
-	int32 PrependPoint(const FTangentBezierControlPoint& Point, float Parameter)
-	{
-		// Add the anchor
-		int32 NewIndex = InternalSpline.PrependAnchor(AnchorFromControlPoint(Point), Parameter);
 
-		if (NewIndex != INDEX_NONE)
-		{
-			// Set tangent mode
-			TangentModes.InsertDefaulted(NewIndex);
-			SetPointTangentMode(NewIndex, Point.TangentMode);
-		}
-
-		return NewIndex;
-	}
+	
 
     /**
-	 * Prepends a control point to the beginning of the spline.
-	 * @param Point - Control point to prepend.
-	 * @param Policy - The policy to use for generating the knot for the point being prepended.
-	 * @return Index of the new point, or INDEX_NONE if the operation failed.
+     * Inserts a new point at the specified parameter along the spline.
+     * 
+     * @param Parameter - Parameter value where to insert the point (Spline Space)
+	 * @param ControlPoint - Control point to insert
+     * @return The index of the newly created point
+     */
+    int32 InsertPointAtGlobalParam(float Parameter, 
+                     const FTangentBezierControlPoint& ControlPoint)
+    {
+		int32 SegmentIndex;
+		float LocalT;
+		bool bValidSegment = InternalSpline.MapGlobalParameterToLocalSegment(Parameter, SegmentIndex, LocalT);
+
+		return InsertPointAtSegmentParam(SegmentIndex, LocalT, ControlPoint);
+    }
+
+	/**
+	 * Inserts a point at a specific segment with local parameter.
+	 * This splits the segment at the specified location and places the new point
+	 * at the provided position, which may change the spline shape.
+	 * 
+	 * @param SegmentIndex Index of the segment to split
+	 * @param LocalT Local parameter [0,1] within the segment
+	 * @param Position Position for the new point
+	 * @param ParameterizationPolicy How to distribute parameters
+	 * @return The index of the newly created point
 	 */
-	int32 PrependPoint(const FTangentBezierControlPoint& Point, EParameterizationPolicy Policy)
+	int32 InsertPointAtSegmentParam(
+		int32 SegmentIndex,
+		float LocalT,
+		const FTangentBezierControlPoint& ControlPoint,
+		EParameterizationPolicy ParameterizationPolicy = EParameterizationPolicy::Uniform)
 	{
-		// Add the anchor
-		int32 NewIndex = InternalSpline.PrependAnchor(AnchorFromControlPoint(Point), Policy);
+		int32 OutNewIndex;
+		// Special case: handle first two points
+		if (InitializeFirstSegment(ControlPoint, OutNewIndex))
+		{
+			return OutNewIndex;
+		}
+		
+		// Insert at that parameter with the exact position
+		// Have the internal spline handle the insertion
+		int32 InsertedPointIndex = InternalSpline.InsertPointAtSegmentParam(SegmentIndex, LocalT, ControlPoint.Position, ParameterizationPolicy);
 		
-		if (NewIndex != INDEX_NONE)
+		// If insertion was successful, update tangents if not auto-computed
+		if (InsertedPointIndex >= 0 &&
+			(ControlPoint.TangentMode == ETangentMode::User ||
+			 ControlPoint.TangentMode == ETangentMode::Broken))
 		{
-			// Set tangent mode
-			TangentModes.InsertDefaulted(NewIndex);
-			SetPointTangentMode(NewIndex, Point.TangentMode);
+			SetTangentIn(InsertedPointIndex, ControlPoint.TangentIn);
+			SetTangentOut(InsertedPointIndex, ControlPoint.TangentOut);
 		}
 
-		return NewIndex;
+		TangentModes.InsertDefaulted(InsertedPointIndex);
+		SetPointTangentMode(InsertedPointIndex, ControlPoint.TangentMode);
+
+		return InsertedPointIndex;
 	}
-	
-    /** 
-	 * Appends a control point to the end of the spline.
-	 * @param Point - Control point to append.
-	 * @param Parameter - The parameter of the control point being appended, must be greater than all existing knots (but may be less than the virtual closing knot).
-	 * @return Index of the new point, or INDEX_NONE if the operation failed.
-	 */
-    int32 AppendPoint(const FTangentBezierControlPoint& Point, float Parameter)
+
+    /**
+     * Inserts a new point at the specified parameter along the spline.
+     * 
+     * @param PointIndex - Index of the new Point
+	 * @param ControlPoint - Control point to insert
+     * @return The index of the newly created point
+     */
+    int32 InsertPointAtPosition(int32 PointIndex, 
+                     const FTangentBezierControlPoint& ControlPoint)
     {
-		// Add the anchor
-		int32 NewIndex = InternalSpline.AppendAnchor(AnchorFromControlPoint(Point), Parameter);
+	    int32 NewPointIndex;
+	    if (InitializeFirstSegment(ControlPoint, NewPointIndex)) return NewPointIndex;
 
-		if (NewIndex != INDEX_NONE)
+	    const int32 NumPoints = GetNumPoints();
+		// need to convert point index to segment index while also clamping,
+		// also keeping in mind that we could append or prepend a segment
+		const int32 SegmentIndex = FMath::Clamp(PointIndex - 1, 0, NumPoints - 1);
+		// Insert at that parameter with the exact position
+        // Have the internal spline handle the insertion
+        int32 InsertedPointIndex = InternalSpline.InsertPointAtPosition(SegmentIndex, ControlPoint.Position);
+		
+		// If insertion was successful, update tangents if not auto-computed
+		if (InsertedPointIndex >= 0 &&
+			(ControlPoint.TangentMode == ETangentMode::User ||
+			 ControlPoint.TangentMode == ETangentMode::Broken))
 		{
-			// Set tangent mode
-			TangentModes.InsertDefaulted(NewIndex);
-			SetPointTangentMode(NewIndex, Point.TangentMode);
+			SetTangentIn(InsertedPointIndex, ControlPoint.TangentIn);
+			SetTangentOut(InsertedPointIndex, ControlPoint.TangentOut);
 		}
+		
+		TangentModes.InsertDefaulted(InsertedPointIndex);
+		SetPointTangentMode(InsertedPointIndex, ControlPoint.TangentMode);
 
-		return NewIndex;
+        return InsertedPointIndex;
     }
+    
+    /**
+	 * Helper method to prepend a point to the spline.
+	 * 
+	 * @param ControlPoint Control point structure to prepend
+	 * @return Index of the new point
+	 */
+	void PrependPoint(const FTangentBezierControlPoint& ControlPoint)
+	{
+    	int32 NewIndex;
+
+		// Special case: handle first two points
+    	if (InitializeFirstSegment(ControlPoint, NewIndex))
+    	{
+    		return;
+    	}
+		        
+		// Calculate control points for the cubic Bezier segment
+		ValueType P0 = ControlPoint.Position;
+		ValueType P1 = TangentOps::TangentOutToP1(P0, ControlPoint.TangentOut);
+		ValueType P2 = TangentOps::TangentInToP2(GetValue(0), GetTangentIn(0));
+
+		// Add the segment
+		InternalSpline.PrependBezierSegment(P0, P1, P2, EParameterizationPolicy::Uniform);
+        
+		// Fixup the in tangent of the new anchor point
+		SetTangentIn(0, ControlPoint.TangentIn);
+		
+		// Set tangent mode
+		TangentModes.InsertDefaulted(0);
+		SetPointTangentMode(0, ControlPoint.TangentMode);
+	}
 
     /** 
 	 * Appends a control point to the end of the spline.
-	 * @param Point - Control point to append.
-	 * @param Policy - The policy to use for generating the knot for the point being appended.
-	 * @return Index of the new point, or INDEX_NONE if the operation failed.
+	 * @param ControlPoint Structure containing position and tangent information
 	 */
-    int32 AppendPoint(const FTangentBezierControlPoint& Point, EParameterizationPolicy Policy)
+    void AppendPoint(const FTangentBezierControlPoint& ControlPoint)
     {
-		// Add the anchor
-		int32 NewIndex = InternalSpline.AppendAnchor(AnchorFromControlPoint(Point), Policy);
-
-		if (NewIndex != INDEX_NONE)
-		{
-			// Set tangent mode
-			TangentModes.InsertDefaulted(NewIndex);
-			SetPointTangentMode(NewIndex, Point.TangentMode);
-		}
-
-		return NewIndex;
-    }
-
-	/**
-	 * Inserts a new point at the specified parameter along the spline.
-	 * @param Point - Control point to insert.
-	 * @param Parameter - The parameter of the control point being inserted.
-	 * @return Index of the new point, or INDEX_NONE if the operation failed.
-	 */
-    int32 InsertPoint(const FTangentBezierControlPoint& Point, float Parameter)
-    {
-		// Add the anchor
-		int32 NewIndex = InternalSpline.InsertAnchor(AnchorFromControlPoint(Point), Parameter);
-
-		if (NewIndex != INDEX_NONE)
-		{
-			// Set tangent mode
-			TangentModes.InsertDefaulted(NewIndex);
-			SetPointTangentMode(NewIndex, Point.TangentMode);
-		}
+		int32 NewIndex = GetNumPoints();
+
+		// Special case: handle first two points
+    	if (InitializeFirstSegment(ControlPoint, NewIndex))
+    	{
+    		return;
+    	}
+    	
+		// Calculate control points for the cubic Bezier segment
+		ValueType P3 = ControlPoint.Position;
+		ValueType P2 = TangentOps::TangentInToP2(P3, ControlPoint.TangentIn);
+		ValueType P1 = TangentOps::TangentOutToP1(GetValue(NewIndex - 1), GetTangentOut(NewIndex - 1));
+    
+		// Add the segment
+		InternalSpline.AppendBezierSegment(P1, P2, P3, EParameterizationPolicy::Uniform);
+		
+		// Fixup the out tangent of the new anchor point
+		SetTangentOut(NewIndex, ControlPoint.TangentOut);
 
-		return NewIndex;
+		// Set tangent mode
+		TangentModes.InsertDefaulted(NewIndex);
+		SetPointTangentMode(NewIndex, ControlPoint.TangentMode);
     }
 	
 	/**
-	 * Gets the control point at the specified index.
-	 * @param Index - Index of the point to retrieve.
-	 * @return Control point at the specified index.
+	 * Gets the control point at the specified index
+	 * @param Index Index of the point to retrieve
+	 * @return Control point structure with position and tangent information
 	 */
-	FTangentBezierControlPoint GetPoint(int32 Index) const
+	FTangentBezierControlPoint GetControlPoint(int32 Index) const
 	{
-		if (Index < 0 || Index >= GetNumPoints())
+		const int32 NumPoints = GetNumPoints();
+		if (Index < 0 || Index >= NumPoints)
 		{
 			return FTangentBezierControlPoint();
 		}
@@ -663,10 +829,40 @@ public:
      */
     void RemovePoint(int32 Index)
 	{
-		if (Index < 0 || Index >= GetNumPoints())
+		const int32 NumPoints = GetNumPoints();
+		if (Index < 0 || Index >= NumPoints)
 		{
 			return;
 		}
+		// special case for handling the last point and second last point
+		if (NumPoints == 1)
+		{
+			InternalSpline.Clear();
+			TangentModes.Reset();
+			return;
+		}
+		else if (NumPoints == 2 && !this->IsClosedLoop())
+		{
+			// Figure out which point remains (0 or 1)
+			const int32 RemainingIndex = (Index == 0) ? 1 : 0;
+	
+			// Build a control point for the remaining anchor
+			FTangentBezierControlPoint RemainingCP;
+			RemainingCP.Position    = GetValue(RemainingIndex);
+			RemainingCP.TangentIn   = GetTangentIn(RemainingIndex);
+			RemainingCP.TangentOut  = GetTangentOut(RemainingIndex);
+			RemainingCP.TangentMode = TangentModes.IsValidIndex(RemainingIndex)
+										? TangentModes[RemainingIndex]
+										: ETangentMode::Auto;
+	
+			InternalSpline.Clear();
+			TangentModes.Reset();
+	
+			int32 TmpIdx = INDEX_NONE;
+			InitializeFirstSegment(RemainingCP, TmpIdx);
+	
+			return;
+		}
     
 		// Remove the corresponding tangent mode
 		if (TangentModes.IsValidIndex(Index))
@@ -674,7 +870,7 @@ public:
 			TangentModes.RemoveAt(Index);
 		}
 
-		InternalSpline.RemoveAnchor(Index);
+		InternalSpline.RemovePoint(Index);
     
 		// Ensure tangent modes array size matches number of points after removal
 		if (TangentModes.Num() > GetNumPoints())
@@ -686,37 +882,38 @@ public:
 	/** 
 	 * Modifies an existing point
 	 * @param Index - Index of the point to modify
-	 * @param Point - New control point data
+	 * @param ControlPoint - New control point data
 	 */
-	void ModifyPoint(int32 Index, const FTangentBezierControlPoint& Point)
+	void ModifyPoint(int32 Index, const FTangentBezierControlPoint& ControlPoint)
 	{
-        if (Index < 0 || Index >= GetNumPoints())
+        const int32 NumPoints = GetNumPoints();
+        if (Index < 0 || Index >= NumPoints)
         {
             return;
         }
 
-        SetPointTangentMode(Index, Point.TangentMode);
+        SetPointTangentMode(Index, ControlPoint.TangentMode);
 
         // Update position
-		InternalSpline.SetAnchorPoint(Index, Point.Position);
+		InternalSpline.SetAnchorPoint(Index, ControlPoint.Position);
         
         // If not auto tangents, apply the specified tangents
-        if (Point.TangentMode != ETangentMode::Auto && 
-            Point.TangentMode != ETangentMode::AutoClamped &&
-            Point.TangentMode != ETangentMode::LegacyAuto && 
-            Point.TangentMode != ETangentMode::LegacyAutoClamped && 
-            Point.TangentMode != ETangentMode::Linear && 
-            Point.TangentMode != ETangentMode::LegacyLinear && 
-            Point.TangentMode != ETangentMode::Constant)
+        if (ControlPoint.TangentMode != ETangentMode::Auto && 
+            ControlPoint.TangentMode != ETangentMode::AutoClamped &&
+            ControlPoint.TangentMode != ETangentMode::LegacyAuto && 
+            ControlPoint.TangentMode != ETangentMode::LegacyAutoClamped && 
+            ControlPoint.TangentMode != ETangentMode::Linear && 
+            ControlPoint.TangentMode != ETangentMode::LegacyLinear && 
+            ControlPoint.TangentMode != ETangentMode::Constant)
         {
-            SetTangentIn(Index, Point.TangentIn);
-            SetTangentOut(Index, Point.TangentOut);
+            SetTangentIn(Index, ControlPoint.TangentIn);
+            SetTangentOut(Index, ControlPoint.TangentOut);
         }
     }
 
 	void SetValue(int32 Index, const ValueType& NewValue)
 	{
-		FTangentBezierControlPoint NewPoint = GetPoint(Index);
+		FTangentBezierControlPoint NewPoint = GetControlPoint(Index);
 		NewPoint.Position = NewValue;
 		ModifyPoint(Index, NewPoint);
 	}
@@ -871,13 +1068,17 @@ public:
 	/** Gets the parameter value for the specified anchor point index */
 	float GetParameter(int32 Index) const
 	{
-		return InternalSpline.GetParameter(Index);
+		const int32 BezierControlPointIdx = (Index == 0) ? 0 : ((Index - 1) * 4 + 3);
+		return InternalSpline.GetParameter(BezierControlPointIdx);
 	}
 
 	/** Sets the parameter value for the specified anchor point index */
 	int32 SetParameter(int32 Index, float NewParameter)
     {
-		return InternalSpline.SetParameter(Index, NewParameter);
+		const int32 BezierControlPointIdx = (Index == 0) ? 0 : ((Index - 1) * 4 + 3);
+		const int32 NewIndex = (InternalSpline.SetParameter(BezierControlPointIdx, NewParameter) + 1) / 4;
+
+		return NewIndex;
     }
 
 	virtual FInterval1f GetParameterSpace() const override
@@ -892,7 +1093,7 @@ public:
 
 	int32 FindSegmentIndex(float Parameter, float &OutLocalParam) const
 	{
-		return InternalSpline.FindIndexForParameter(Parameter, OutLocalParam);
+		return InternalSpline.FindSegmentIndex(Parameter, OutLocalParam);
 	}
 
     /**
@@ -917,7 +1118,7 @@ public:
 
 	/**
 	 * Updates tangents for a specific point based on its tangent mode
-	 * @param Index - Index of the point to update
+	 * @param Index Index of the point to update
 	 */
 	void UpdateTangent(int32 Index)
 	{
@@ -1131,7 +1332,7 @@ public:
 	}
     
     /** Getter for internal spline */
-    const Base& GetInternalSpline() const
+    const FPolyBezierSpline3d& GetInternalSpline() const
     {
         return InternalSpline;
     }
@@ -1219,20 +1420,65 @@ private:
 			   Mode == ETangentMode::Broken;
 	}
 	
-	void AssignInternalSpline(const Base& NewSpline)
+	/**
+	 * Helper method to add the first segment to the spline
+	 * @param ControlPoint - Control point to insert
+	 * @param OutPointIndex - Index of the newly created point
+	 * @return true if the points were added, false otherwise
+	 */
+	bool InitializeFirstSegment(const FTangentBezierControlPoint& ControlPoint, int32& OutPointIndex)
+    {
+		const int32 NumPoints = GetNumPoints();
+    
+    	if (NumPoints == 0)
+		{
+			const ValueType A0   = ControlPoint.Position;
+			const ValueType Out0 = TangentOps::TangentOutToP1(A0, ControlPoint.TangentOut);
+			const ValueType In0  = TangentOps::TangentInToP2(A0, ControlPoint.TangentIn);
+	
+			InternalSpline.Clear();
+	
+			// Single anchor (no segments yet), but In/Out cache the tangents
+			InternalSpline.AppendAnchor(In0, A0, Out0);
+	
+			TangentModes.SetNum(1);
+			TangentModes[0] = ControlPoint.TangentMode;
+	
+			OutPointIndex = 0;
+			return true;
+		}
+		else if (NumPoints == 1)
+		{
+			// We already have anchor 0; this is the second point -> anchor 1
+	
+			const ValueType A1   = ControlPoint.Position;
+			const ValueType Out1 = TangentOps::TangentOutToP1(A1, ControlPoint.TangentOut);
+			const ValueType In1  = TangentOps::TangentInToP2(A1, ControlPoint.TangentIn);
+	
+			InternalSpline.AppendAnchor(In1, A1, Out1);
+	
+			TangentModes.Add(ControlPoint.TangentMode);
+	
+			// Rebuild parameterization if needed (Uniform or whatever policy you use)
+			InternalSpline.Reparameterize(EParameterizationPolicy::Uniform);
+	
+			OutPointIndex = 1;
+			return true;
+		}
+	
+		return false;
+    }
+	void AssignInternalSpline(const FPolyBezierSpline& NewSpline)
 	{
 		InternalSpline = NewSpline;
 		InstallFindWindowOverride();
 	}
-
-	void AssignInternalSpline(Base&& NewSpline)
+	void AssignInternalSpline(FPolyBezierSpline&& NewSpline)
 	{
 		InternalSpline = MoveTemp(NewSpline);
 		InstallFindWindowOverride();
 	}
-
 protected:
-
 	/* Tension parameter for auto-computed tangents [0,1] */
 	float Tension = 0.0f;
 
@@ -1241,11 +1487,9 @@ protected:
 
 	bool bStationaryEndpoint = false;
 	
-	TStaticArray<ValueType, TPolyBezierSpline<ValueType>::Base::WindowSize> IntermediateWindow;
-
+	TStaticArray<ValueType, TPolyBezierSpline<VALUETYPE>::Base::WindowSize> IntermediateWindow;
 private:
-
-	using TSplineWrapper<TPolyBezierSpline<ValueType>>::InternalSpline;
+	using TSplineWrapper<TPolyBezierSpline<VALUETYPE>>::InternalSpline;
 };
 
 using FTangentBezierSpline3f = TTangentBezierSpline<FVector3f>; 
