Add a Sort() method to FIndex3i, to change the indices to a sorted order (ala the one already on FIndex2i)

#Synced-CL 49942557

diff --git a/Engine/Source/Runtime/GeometryCore/Public/IndexTypes.h b/Engine/Source/Runtime/GeometryCore/Public/IndexTypes.h
--- a/Engine/Source/Runtime/GeometryCore/Public/IndexTypes.h
+++ b/Engine/Source/Runtime/GeometryCore/Public/IndexTypes.h
@@ -220,6 +220,23 @@ PRAGMA_ENABLE_DEPRECATION_WARNINGS
 		return (A == Value) || (B == Value) || (C == Value);
 	}
 
+	// Sort the indices from lowest to highest
+	void Sort()
+	{
+		if (A > B)
+		{
+			Swap(A, B);
+		}
+		if (B > C)
+		{
+			Swap(B, C);
+			if (A > B)
+			{
+				Swap(A, B);
+			}
+		}
+	}
+
 	/** @return offset triplet, with the OffsetIndicesBy value added to each index */
 	[[nodiscard]] FIndex3i GetOffsetBy(int32 OffsetIndicesBy) const
 	{
