Add a MeshVertexNormals method that can account for implicit welding -- i.e. a known mapping of some vertices to common weld vertices.

#Synced-CL 49961011

diff --git a/Engine/Source/Runtime/GeometryCore/Public/MeshAdapter/MeshVertexNormals.h b/Engine/Source/Runtime/GeometryCore/Public/MeshAdapter/MeshVertexNormals.h
--- a/Engine/Source/Runtime/GeometryCore/Public/MeshAdapter/MeshVertexNormals.h
+++ b/Engine/Source/Runtime/GeometryCore/Public/MeshAdapter/MeshVertexNormals.h
@@ -222,6 +222,66 @@ namespace MeshNormals
 			}
 		}
 	}
+
+	/**
+	 * Compute vertex normals where some vertices are implicitly welded.
+	 * @param Mesh The mesh to compute normals for
+	 * @param VertexIDWeldMap Mapping from Vertex ID to the representative 'welded' ID. A set of mutually welded vertices must always map to the same welded-ID, which must be one of the vertex IDs from that set.
+	 * @param OutNormals Computed normals
+	 * @param bRequireDeterministicNormals If false, uses a non-deterministic fast path with unordered float operations that can produce slightly different roundings
+	 */
+	template<typename MeshType, typename NormalsArrayType>
+	void ComputeWeldedVertexNormals(const MeshType& Mesh, TFunctionRef<int32(int32)> VertexIDWeldMap, NormalsArrayType& OutNormals, bool bRequireDeterministicNormals = true)
+	{
+		// Minimal mesh wrapper struct to forward triangle IDs to welded IDs
+		struct FWeldMesh
+		{
+			FWeldMesh(const MeshType& InMesh, const TFunctionRef<int32(int32)>& InVertexIDWeldMap) : Mesh(&InMesh), VertexIDWeldMap(&InVertexIDWeldMap){}
+
+			int32 MaxVertexID() const
+			{
+				return Mesh->MaxVertexID();
+			}
+			int32 MaxTriangleID() const
+			{
+				return Mesh->MaxTriangleID();
+			}
+			bool IsTriangle(int32 TID) const
+			{
+				return Mesh->IsTriangle(TID);
+			}
+			bool IsVertex(int32 VID) const
+			{
+				return Mesh->IsVertex(VID);
+			}
+			void GetTriVertices(int32 TID, FVector3d& A, FVector3d& B, FVector3d& C) const
+			{
+				return Mesh->GetTriVertices(TID, A, B, C);
+			}
+			FIndex3i GetTriangle(int32 TID) const
+			{
+				FIndex3i Tri = Mesh->GetTriangle(TID);
+				Tri.A = (*VertexIDWeldMap)(Tri.A);
+				Tri.B = (*VertexIDWeldMap)(Tri.B);
+				Tri.C = (*VertexIDWeldMap)(Tri.C);
+				return Tri;
+			}
+		private:
+			const MeshType* Mesh;
+			const TFunctionRef<int32(int32)>* VertexIDWeldMap;
+		};
+		FWeldMesh WeldMesh(Mesh, VertexIDWeldMap);
+		ComputeVertexNormals(WeldMesh, OutNormals, bRequireDeterministicNormals);
+		// Copy welded IDs back out
+		ParallelFor(Mesh.MaxVertexID(), [&OutNormals, &VertexIDWeldMap](int32 Idx)
+		{
+			int32 WeldIdx = VertexIDWeldMap(Idx);
+			if (WeldIdx > INDEX_NONE && WeldIdx != Idx)
+			{
+				OutNormals[Idx] = OutNormals[WeldIdx];
+			}
+		});
+	}
 }
 
 } // end namespace UE::Geometry
