Add a method for VertexConnectedComponents to merge only the close vertices that pass some other filter, so e.g. one can only do this merge step for boundary vertices (which can be much faster than doing it for all vertices, and also can avoid connecting some intentionally-separated-but-nearby vertices)

+ also, add a thread-safe method for getting the cluster ID of vertex-connected components

#Synced-CL 49961011

diff --git a/Engine/Source/Runtime/GeometryCore/Public/VertexConnectedComponents.h b/Engine/Source/Runtime/GeometryCore/Public/VertexConnectedComponents.h
--- a/Engine/Source/Runtime/GeometryCore/Public/VertexConnectedComponents.h
+++ b/Engine/Source/Runtime/GeometryCore/Public/VertexConnectedComponents.h
@@ -67,14 +67,19 @@ public:
 		}
 	}
 
-	// Note for meshes that may have floating (unreferenced) vertices, using a KeepSizeThreshold > 1 after calling ConnectTriangles will ensure the unreferenced vertices are skipped
+	/**
+	 * Connect a filtered subset of vertices that are closer than a threshold distance in space
+	 * @param CloseVertexThreshold Distance below which vertices should be connected
+	 * @param VertexIDFilter Return true if the vertex should be considered for connection to other vertices
+	 * @param KeepSizeThreshold Skip considering vertices that are in a smaller cluster size than this. Note: A value > 1 after calling ConnectTriangles will ensure the unreferenced vertices are skipped.
+	 */
 	template<typename TriangleMeshType>
-	void ConnectCloseVertices(const TriangleMeshType& Mesh, double CloseVertexThreshold, int32 KeepSizeThreshold = 0)
+	void ConnectCloseFilteredVertices(const TriangleMeshType& Mesh, double CloseVertexThreshold, TFunctionRef<bool(int32)> VertexIDFilter, int32 KeepSizeThreshold = 0)
 	{
 		TPointHashGrid3d<int32> VertexHash(CloseVertexThreshold * 3, -1);
 		for (int32 VID = 0; VID < Mesh.MaxVertexID(); VID++)
 		{
-			if (!Mesh.IsVertex(VID))
+			if (!Mesh.IsVertex(VID) || !VertexIDFilter(VID))
 			{
 				continue;
 			}
@@ -97,6 +102,17 @@ public:
 			VertexHash.InsertPointUnsafe(VID, Pt);
 		}
 	}
+	
+	/**
+	 * Connect vertices that are closer than a threshold distance in space
+	 * @param CloseVertexThreshold Distance below which vertices should be connected
+	 * @param KeepSizeThreshold Skip considering vertices that are in a smaller cluster size than this. Note: A value > 1 after calling ConnectTriangles will ensure the unreferenced vertices are skipped.
+	 */
+	template<typename TriangleMeshType>
+	void ConnectCloseVertices(const TriangleMeshType& Mesh, double CloseVertexThreshold, int32 KeepSizeThreshold = 0)
+	{
+		ConnectCloseFilteredVertices(Mesh, CloseVertexThreshold, [&Mesh](int32 VID)->bool {return true;}, KeepSizeThreshold);
+	}
 
 	// TODO: support more overlap strategies, currently just uses AABB
 	// Note this merges components based on overlap of their bounding boxes as computed *before* any merges; multiple passes may merge additional components
@@ -287,16 +303,25 @@ public:
 		return true;
 	}
 
+	// Note: May not be thread-safe to make multiple GetComponent() calls in parallel, as it will also collapse the underlying disjoint set paths.
 	inline int32 GetComponent(int32 VertexID)
 	{
 		return DisjointSet.Find(VertexID);
 	}
 
-	inline int32 GetComponentSize(int32 VertexID)
+	// Note: Should be safe to make mulitple GetComponentThreadsafe() calls in parallel, as it does not modify the underlying data.
+	// Will not collapse disjoint set paths, so repeated accesses may be slower in an un-threaded context.
+	inline int32 GetComponentThreadsafe(int32 VertexID) const
+	{
+		return DisjointSet.FindWithoutCollapse(VertexID);
+	}
+
+	inline int32 GetComponentSize(int32 VertexID) const
 	{
 		return DisjointSet.GetSize(VertexID);
 	}
 
+	// Note: May not be thread-safe to make multiple GetComponent() calls in parallel, as it will also collapse the underlying disjoint set paths.
 	template<typename TriangleType>
 	inline int32 GetComponent(const TriangleType& Triangle)
 	{
