Avoid MSVC PGO messages about uninitialized variables for callers

#Synced-CL 50342528

diff --git a/Engine/Source/Runtime/GeometryCore/Private/Spatial/GeometrySet3.cpp b/Engine/Source/Runtime/GeometryCore/Private/Spatial/GeometrySet3.cpp
--- a/Engine/Source/Runtime/GeometryCore/Private/Spatial/GeometrySet3.cpp
+++ b/Engine/Source/Runtime/GeometryCore/Private/Spatial/GeometrySet3.cpp
@@ -120,7 +120,8 @@ bool FGeometrySet3::FindNearestPointToRay(const FRay3d& Ray, FNearest& ResultOut
 		}
 	});
 
-	if (NearestID != -1)
+	const bool bFound = NearestID != -1;
+	if (bFound)
 	{
 		ResultOut.ID = NearestID;
 		ResultOut.bIsPoint = true;
@@ -129,7 +130,7 @@ bool FGeometrySet3::FindNearestPointToRay(const FRay3d& Ray, FNearest& ResultOut
 		ResultOut.RayParam = MinRayParamT;
 	}
 
-	return NearestID != -1;
+	return bFound;
 }
 
 bool FGeometrySet3::CollectPointsNearRay(const FRay3d& Ray, TArray<FNearest>& ResultsOut,
