fix unit tests to account for spatial sort changes / deprecations
make the deprecated spatial sort methods still store the (deprecated) mesh array + run the sort

#Synced-CL 50522370

diff --git a/Engine/Source/Runtime/GeometryCore/Public/Spatial/MeshSpatialSort.h b/Engine/Source/Runtime/GeometryCore/Public/Spatial/MeshSpatialSort.h
--- a/Engine/Source/Runtime/GeometryCore/Public/Spatial/MeshSpatialSort.h
+++ b/Engine/Source/Runtime/GeometryCore/Public/Spatial/MeshSpatialSort.h
@@ -38,7 +38,7 @@ public:
 	/**
 	 * Meshes to be spatially sorted
 	 */
-	UE_DEPRECATED(5.8, "This is no longer stored on the class; instead an accessor is passed in per-call")
+	UE_DEPRECATED(5.8, "This will no longer stored be on the class; instead an accessor should be passed in per-call")
 	TArrayView<const TriangleMeshType> InputMeshes;
 
 	/**
@@ -97,7 +97,8 @@ public:
 	TMeshSpatialSort() {}
 
 	UE_DEPRECATED(5.8, "Pass meshes and bounds to the Compute() method directly")
-	TMeshSpatialSort(TArrayView<const TriangleMeshType> InputMeshes, TArrayView<const FAxisAlignedBox3d> MeshBoundsIn = TArrayView<const FAxisAlignedBox3d>())
+	TMeshSpatialSort(TArrayView<const TriangleMeshType> InputMeshesIn, TArrayView<const FAxisAlignedBox3d> MeshBoundsIn = TArrayView<const FAxisAlignedBox3d>())
+		: InputMeshes(InputMeshesIn)
 	{
 		ensure(false);
 	}
@@ -106,6 +107,7 @@ public:
 	void Compute()
 	{
 		ensure(false);
+		Compute(InputMeshes);
 	}
 
 	void Compute(TArrayView<const TriangleMeshType> InInputMeshes, TArrayView<const FAxisAlignedBox3d> MeshBoundsIn = TArrayView<const FAxisAlignedBox3d>())
