Deprecate the version of SplitMeshes that takes a TArray<FDynamicMesh3>, to move code away from unsafe usage of FDynamicMesh3/TArray

#Synced-CL 50511475

diff --git a/Engine/Source/Runtime/GeometryCore/Public/DynamicMeshEditor.h b/Engine/Source/Runtime/GeometryCore/Public/DynamicMeshEditor.h
--- a/Engine/Source/Runtime/GeometryCore/Public/DynamicMeshEditor.h
+++ b/Engine/Source/Runtime/GeometryCore/Public/DynamicMeshEditor.h
@@ -626,7 +626,6 @@ public:
 	/**
 	 * NOTE: Please use the TArray<TUniquePtr<FDynamicMesh3>> version of this method, below, instead of this one, because storing FDynamicMesh3 in TArray directly is not safe.
 	 * This method will not directly cause memory corruption, but any subsequent manipulation of the array will be unsafe.
-	 * TODO: Deprecate this method.
 	 * 
 	 * Create multiple meshes out of the source mesh by splitting triangles out.
 	 * Static because it creates multiple output meshes, so doesn't quite fit in the FDynamicMeshEditor model of operating on a single mesh
@@ -639,6 +638,7 @@ public:
 	 * @param bSortByMeshID If true, SplitMeshes will be sorted in order of increasing Mesh ID
 	 * @return true if needed split, false if there were not multiple mesh ids so no split was needed
 	 */
+	UE_DEPRECATED(5.8, "Use the version of SplitMesh that takes a TArray<TUniquePtr<FDynamicMesh3>> array instead")
 	static GEOMETRYCORE_API bool SplitMesh(const FDynamicMesh3* SourceMesh, TArray<FDynamicMesh3>& SplitMeshes, TFunctionRef<int(int)> TriIDToMeshID, int DeleteMeshID = -1, TArray<int32>* MeshIDPerSplitMesh = nullptr, bool bSortByMeshID = false);
 
 	/**
