PVSimplifyPlantSkeleton

Add new "SimplifyPlantSkeleton" node which allows the user to run a douglas peucker algorithm over the plant skeleton

Key changes:
- PVSimplifyPlantSkeletonSettings.h/.cpp: New node
- PVSimplifyPlantSkeleton.h/.cpp: Iterate the plant branches and run TPolyline::Simplify on the segments between points which has child branches (to avoid removing points with child branches), then call FPVCarve::RemoveEntriesAndRecomputeAttributes to remove the points.

-----

TPolyline::Simplify (Reviewed by[REDACTED_HANDLE])

This CL modifies TPolyline::Simplify to return an array of vertex indices instead of an array of FVector3f. This is so the caller can know which vertices got removed.

Key changes:
- Add new overload of with the output argument "TArray<int32>& OutNewVertexIndices"
- Update original Simplify function to use new overload of Simplify
- Add int32 NewSize to Clear
- Add Reserve to TPolyline to allow user to reserve memory for the expected number of vertices that will be in inserted into the polyline

#Synced-CL 51214504

diff --git a/Engine/Plugins/Experimental/ProceduralVegetationEditor/Source/ProceduralVegetation/Private/Implementations/PVSimplifyPlantSkeleton.cpp b/Engine/Plugins/Experimental/ProceduralVegetationEditor/Source/ProceduralVegetation/Private/Implementations/PVSimplifyPlantSkeleton.cpp
new file mode 100644
--- /dev/null
+++ b/Engine/Plugins/Experimental/ProceduralVegetationEditor/Source/ProceduralVegetation/Private/Implementations/PVSimplifyPlantSkeleton.cpp
@@ -0,0 +1,138 @@
+// Copyright Epic Games, Inc. All Rights Reserved.
+
+#include "PVSimplifyPlantSkeleton.h"
+#include "Helpers/PVPlantTraversalHelper.h"
+#include "Polyline3.h"
+#include "Implementations/PVCarve.h"
+
+void PV::SimplifyPlantSkeleton(FManagedArrayCollection& InOutCollection, float SimplificationAmount)
+{
+	using namespace PV::PlantTraversalHelper;
+	using namespace UE::Geometry;
+
+	const PV::FPointPositionAttributeConstView PointPositionAttribute = PV::FPointPositionAttribute::FindAttribute(InOutCollection);
+	const PV::FBranchChildrenAttributeConstView BranchChildrenAttribute = PV::FBranchChildrenAttribute::FindAttribute(InOutCollection);
+	const PV::FBranchParentNumberAttributeConstView BranchParentNumberAttribute = PV::FBranchParentNumberAttribute::FindAttribute(InOutCollection);
+	const PV::FBranchNumberAttributeConstView BranchNumberAttribute = PV::FBranchNumberAttribute::FindAttribute(InOutCollection);
+	const PV::FBranchPointsAttributeConstView BranchPointsAttribute = PV::FBranchPointsAttribute::FindAttribute(InOutCollection);
+	const PV::FPointLengthFromRootAttributeConstView PointLengthFromRootAttribute = PV::FPointLengthFromRootAttribute::FindAttribute(InOutCollection);
+	const PV::FBranchFoliageIDsAttributeConstView BranchFoliageIDsAttribute = PV::FBranchFoliageIDsAttribute::FindAttribute(InOutCollection);
+	const PV::FFoliageLengthFromRootAttributeConstView FoliageLengthFromRootAttribute = PV::FFoliageLengthFromRootAttribute::FindAttribute(InOutCollection);
+
+	if (!PointPositionAttribute
+		|| !BranchChildrenAttribute
+		|| !BranchParentNumberAttribute
+		|| !BranchNumberAttribute
+		|| !BranchPointsAttribute
+		|| !PointLengthFromRootAttribute
+		|| !BranchFoliageIDsAttribute
+		|| !FoliageLengthFromRootAttribute)
+	{
+		return;
+	}
+
+	const bool bHasFusedPoints = HasFusedPoints(BranchPointsAttribute, BranchParentNumberAttribute, BranchNumberAttribute);
+	const auto GetNumAxillaryBranches = [&](int32 BranchIndex, int32 BranchPointIndex)
+	{
+		if (bHasFusedPoints)
+		{
+			return GetBranchPointChildBranchIndices_FusedPoints(
+				BranchChildrenAttribute,
+				BranchParentNumberAttribute,
+				BranchNumberAttribute,
+				BranchPointsAttribute,
+				BranchIndex,
+				BranchPointIndex
+			).Num();
+		}
+		else
+		{
+			return GetBranchPointChildBranchIndices_SplitPoints(
+				BranchChildrenAttribute,
+				BranchParentNumberAttribute,
+				BranchNumberAttribute,
+				BranchPointsAttribute,
+				PointLengthFromRootAttribute,
+				BranchIndex,
+				BranchPointIndex
+			).Num();
+		}
+	};
+
+	TArray<bool> PointsToRemove;
+	PointsToRemove.Init(false, PointPositionAttribute.Num());
+
+	TArray<bool> BranchesToRemove;
+	BranchesToRemove.Init(false, BranchNumberAttribute.Num());
+
+	TArray<bool> FoliageInstancesToRemove;
+	FoliageInstancesToRemove.Init(false, FoliageLengthFromRootAttribute.Num());
+
+	TArray<int32> NewVertexIndices;
+	const auto SimplifySegment = [&](
+		const TArray<int32>& BranchPoints,
+		UE::Geometry::FPolyline3f& Polyline,
+		int32 SegmentStartIdx,
+		int32 BranchIndex)
+	{
+		const float LineDeviationTolerance = SimplificationAmount;
+		const float ClusterTolerance = 0.001;
+		Polyline.Simplify(ClusterTolerance, LineDeviationTolerance, NewVertexIndices);
+
+		for (int32 i = 1; i < NewVertexIndices.Num(); ++i)
+		{
+			const int32 PrevBranchPointIndex = SegmentStartIdx + NewVertexIndices[i - 1];
+			const int32 ThisBranchPointIndex = SegmentStartIdx + NewVertexIndices[i];
+
+			for (int32 BranchPointIndex = PrevBranchPointIndex + 1; BranchPointIndex < ThisBranchPointIndex; ++BranchPointIndex)
+			{
+				PointsToRemove[BranchPoints[BranchPointIndex]] = true;
+
+				const TArray<int32> BranchSegmentFoliageEntryIds = GetBranchPointFoliageEntryIds(
+					BranchPointsAttribute,
+					PointLengthFromRootAttribute,
+					BranchFoliageIDsAttribute,
+					FoliageLengthFromRootAttribute,
+					BranchIndex,
+					BranchPointIndex
+				);
+
+				for (int32 FoliageEntryIndex : BranchSegmentFoliageEntryIds)
+				{
+					FoliageInstancesToRemove[FoliageEntryIndex] = true;
+				}
+			}
+		}
+	};
+
+	UE::Geometry::FPolyline3f Polyline;
+	for (int32 BranchIndex = 0; BranchIndex < BranchPointsAttribute.Num(); ++BranchIndex)
+	{
+		const TArray<int32>& BranchPoints = BranchPointsAttribute[BranchIndex];
+		Polyline.Clear(BranchPoints.Num());
+
+		int32 SegmentStartIdx = 0;
+		for (int32 BranchPointIndex = 0; BranchPointIndex < BranchPoints.Num(); ++BranchPointIndex)
+		{
+			const int32 PointIndex = BranchPoints[BranchPointIndex];
+			Polyline.AppendVertex(PointPositionAttribute[PointIndex]);
+			if (GetNumAxillaryBranches(BranchIndex, BranchPointIndex) > 0)
+			{
+				SimplifySegment(BranchPoints, Polyline, SegmentStartIdx, BranchIndex);
+				Polyline.Clear();
+				Polyline.AppendVertex(PointPositionAttribute[PointIndex]);
+				SegmentStartIdx = BranchPointIndex;
+			}
+		}
+		SimplifySegment(BranchPoints, Polyline, SegmentStartIdx, BranchIndex);
+	}
+
+	const FManagedArrayCollection SourceCollection = InOutCollection;
+	FPVCarve::RemoveEntriesAndRecomputeAttributes(
+		InOutCollection,
+		SourceCollection,
+		PointsToRemove,
+		BranchesToRemove,
+		FoliageInstancesToRemove
+	);
+}
diff --git a/Engine/Plugins/Experimental/ProceduralVegetationEditor/Source/ProceduralVegetation/Private/Implementations/PVSimplifyPlantSkeleton.h b/Engine/Plugins/Experimental/ProceduralVegetationEditor/Source/ProceduralVegetation/Private/Implementations/PVSimplifyPlantSkeleton.h
new file mode 100644
--- /dev/null
+++ b/Engine/Plugins/Experimental/ProceduralVegetationEditor/Source/ProceduralVegetation/Private/Implementations/PVSimplifyPlantSkeleton.h
@@ -0,0 +1,11 @@
+// Copyright Epic Games, Inc. All Rights Reserved.
+
+#pragma once
+
+#include "CoreMinimal.h"
+#include "GeometryCollection/ManagedArrayCollection.h"
+
+namespace PV
+{
+	extern void SimplifyPlantSkeleton(FManagedArrayCollection& InOutCollection, float SimplificationAmount);
+};
diff --git a/Engine/Plugins/Experimental/ProceduralVegetationEditor/Source/ProceduralVegetation/Private/Nodes/PVSimplifyPlantSkeletonSettings.cpp b/Engine/Plugins/Experimental/ProceduralVegetationEditor/Source/ProceduralVegetation/Private/Nodes/PVSimplifyPlantSkeletonSettings.cpp
new file mode 100644
--- /dev/null
+++ b/Engine/Plugins/Experimental/ProceduralVegetationEditor/Source/ProceduralVegetation/Private/Nodes/PVSimplifyPlantSkeletonSettings.cpp
@@ -0,0 +1,75 @@
+// Copyright Epic Games, Inc. All Rights Reserved.
+
+#include "PVSimplifyPlantSkeletonSettings.h"
+#include "DataTypes/PVGrowthData.h"
+#include "GeometryCollection/ManagedArrayCollection.h"
+#include "Implementations/PVSimplifyPlantSkeleton.h"
+
+#define LOCTEXT_NAMESPACE "PVSimplifyPlantSkeletonSettings"
+
+#if WITH_EDITOR
+FText UPVSimplifyPlantSkeletonSettings::GetDefaultNodeTitle() const
+{ 
+	return LOCTEXT("NodeTitle", "SimplifyPlantSkeleton"); 
+}
+
+FText UPVSimplifyPlantSkeletonSettings::GetNodeTooltipText() const
+{ 
+	return LOCTEXT("NodeTooltip", 
+		"Simplifies the plant skeleton using the \"douglas peucker\" algorithm. Points with child branches will not be removed."
+		"\n\nPress Ctrl + L to lock/unlock node output"
+	); 
+}
+#endif
+
+FPCGDataTypeIdentifier UPVSimplifyPlantSkeletonSettings::GetInputPinTypeIdentifier() const
+{
+	return FPCGDataTypeIdentifier{ FPVDataTypeInfoGrowth::AsId() };
+}
+
+FPCGDataTypeIdentifier UPVSimplifyPlantSkeletonSettings::GetOutputPinTypeIdentifier() const
+{
+	return FPCGDataTypeIdentifier{ FPVDataTypeInfoGrowth::AsId() };
+}
+
+FPCGElementPtr UPVSimplifyPlantSkeletonSettings::CreateElement() const
+{
+	return MakeShared<FPVSimplifyPlantSkeletonElement>();
+}
+
+bool FPVSimplifyPlantSkeletonElement::ExecuteInternal(FPCGContext* InContext) const
+{
+	TRACE_CPUPROFILER_EVENT_SCOPE(FPVSimplifyPlantSkeletonElement::Execute);
+
+	check(InContext);
+
+	const UPVSimplifyPlantSkeletonSettings* Settings = InContext->GetInputSettings<UPVSimplifyPlantSkeletonSettings>();
+	check(Settings);
+
+	for (const FPCGTaggedData& Input : InContext->InputData.GetInputsByPin(PCGPinConstants::DefaultInputLabel))
+	{
+		if (const UPVGrowthData* InputData = Cast<UPVGrowthData>(Input.Data))
+		{
+			FManagedArrayCollection OutCollection = InputData->GetCollection();
+
+			UPVGrowthData* OutManagedArrayCollectionData = FPCGContext::NewObject_AnyThread<UPVGrowthData>(InContext);
+
+			if (Settings->SimplificationAmount > 0)
+			{
+				PV::SimplifyPlantSkeleton(OutCollection, Settings->SimplificationAmount);
+			}
+
+			OutManagedArrayCollectionData->Initialize(MoveTemp(OutCollection));
+			InContext->OutputData.TaggedData.Emplace(OutManagedArrayCollectionData);
+		}
+		else
+		{
+			PCGLog::InputOutput::LogInvalidInputDataError(InContext);
+			return true;
+		}
+	}
+
+	return true;
+}
+
+#undef LOCTEXT_NAMESPACE
\ No newline at end of file
diff --git a/Engine/Plugins/Experimental/ProceduralVegetationEditor/Source/ProceduralVegetation/Private/Nodes/PVSimplifyPlantSkeletonSettings.h b/Engine/Plugins/Experimental/ProceduralVegetationEditor/Source/ProceduralVegetation/Private/Nodes/PVSimplifyPlantSkeletonSettings.h
new file mode 100644
--- /dev/null
+++ b/Engine/Plugins/Experimental/ProceduralVegetationEditor/Source/ProceduralVegetation/Private/Nodes/PVSimplifyPlantSkeletonSettings.h
@@ -0,0 +1,38 @@
+// Copyright Epic Games, Inc. All Rights Reserved.
+#pragma once
+
+#include "CoreMinimal.h"
+#include "Nodes/PVBaseSettings.h"
+#include "PVSimplifyPlantSkeletonSettings.generated.h"
+
+UCLASS(BlueprintType, ClassGroup = (Procedural))
+class UPVSimplifyPlantSkeletonSettings : public UPVBaseSettings
+{
+	GENERATED_BODY()
+
+public:
+#if WITH_EDITOR
+	virtual FName GetDefaultNodeName() const override { return FName(TEXT("ProceduralVegetationSimplifyPlantSkeleton")); }
+	virtual FText GetDefaultNodeTitle() const override;
+	virtual FText GetNodeTooltipText() const override;
+	virtual FLinearColor GetNodeTitleColor() const override { return FLinearColor::Yellow; }
+
+	virtual TArray<EPVRenderType> GetDefaultRenderType() const override { return TArray{ EPVRenderType::PointData }; }
+#endif
+
+protected:
+	virtual FPCGDataTypeIdentifier GetInputPinTypeIdentifier() const override;
+	virtual FPCGDataTypeIdentifier GetOutputPinTypeIdentifier() const override;
+
+	virtual FPCGElementPtr CreateElement() const override;
+
+public:
+	UPROPERTY(EditAnywhere, Category = "SimplifyPlantSkeleton", meta = (ClampMin = 0, UIMin = 0, UIMax = 1.0))
+	float SimplificationAmount = 0.1;
+};
+
+class FPVSimplifyPlantSkeletonElement : public FPVBaseElement
+{
+protected:
+	virtual bool ExecuteInternal(FPCGContext* InContext) const override;
+};
diff --git a/Engine/Source/Runtime/GeometryCore/Public/Polyline.h b/Engine/Source/Runtime/GeometryCore/Public/Polyline.h
--- a/Engine/Source/Runtime/GeometryCore/Public/Polyline.h
+++ b/Engine/Source/Runtime/GeometryCore/Public/Polyline.h
@@ -138,10 +138,24 @@ public:
 	}
 
 
-	/** Discard all vertices of polyline */
-	void Clear()
+	/** 
+	* Discard all vertices of polyline 
+	* 
+	* @param NewSize The expected number of vertices to be appended after calling this function.
+	*/
+	void Clear(int32 NewSize = 0)
 	{
-		Vertices.Reset();
+		Vertices.Reset(NewSize);
+	}
+
+	/** 
+	* Reserves memory for the given number of vertices 
+	* 
+	* @param The number of vertices the polyline should be able to contain after allocation.
+	*/
+	void Reserve(int32 NewSize)
+	{
+		Vertices.Reserve(NewSize);
 	}
 
 	/**
@@ -511,67 +525,95 @@ public:
 	 * @param LineDeviationTolerance Vertices are allowed to deviate this much from the polylines
 	 */
 	void Simplify(T ClusterTolerance, T LineDeviationTolerance)
+	{
+		TArray<int32> NewVertexIndices;
+		Simplify(ClusterTolerance, LineDeviationTolerance, NewVertexIndices);
+	}
+
+	/**
+	 * Simplify the Polyline to reduce the vertex count
+	 * @param ClusterTolerance Vertices closer than this distance will be merged into a single vertex
+	 * @param LineDeviationTolerance Vertices are allowed to deviate this much from the polylines
+	 * @param OutNewVertexIndices Output array containing the indices of Vertices that remain after simplification.
+	 *        The indices reference the Vertices array before simplification and are ordered along the polyline.
+	 */
+	void Simplify(T ClusterTolerance, T LineDeviationTolerance, TArray<int32>& OutNewVertexIndices)
 	{
 		const int32 VertexCount = Vertices.Num();
 		if (VertexCount < 3)
 		{
+			OutNewVertexIndices.SetNumUninitialized(VertexCount);
+			for (int32 i = 0; i < VertexCount; ++i)
+			{
+				OutNewVertexIndices[i] = i;
+			}
 			// we need at least 3 vertices to be able to simplify a line
 			return;
 		}
 
-		TArray<VectorType> NewVertices;
-		NewVertices.Reserve(VertexCount);
+		OutNewVertexIndices.Reset(VertexCount);
 
 		// STAGE 1.  Vertex Reduction within tolerance of prior vertex cluster
 		if (ClusterTolerance > 0)
 		{
 			T ClusterToleranceSquared = ClusterTolerance * ClusterTolerance;
-			NewVertices.Add(Vertices[0]);				// keep the first vertex
+			OutNewVertexIndices.Add(0);				// keep the first vertex
 			for (int32 Index = 1; Index < VertexCount-1; Index++) 
 			{
-				if (Geometry::DistanceSquared(Vertices[Index], NewVertices.Last()) < ClusterToleranceSquared)
+				if (Geometry::DistanceSquared(Vertices[Index], Vertices[OutNewVertexIndices.Last()]) < ClusterToleranceSquared)
 				{
 					continue;
 				}
-				NewVertices.Add(Vertices[Index]);
+				OutNewVertexIndices.Add(Index);
 			}
-			NewVertices.Add(Vertices.Last());		// keep the last vertex
+			OutNewVertexIndices.Add(VertexCount - 1);		// keep the last vertex
 		}
 		else
 		{
-			NewVertices = Vertices;
+			OutNewVertexIndices.SetNumUninitialized(VertexCount);
+			for (int32 i = 0; i < VertexCount; ++i)
+			{
+				OutNewVertexIndices[i] = i;
+			}
 		}
 
 		// STAGE 2.  Douglas-Peucker polyline simplification
 		TArray<bool> Marked;
-		if (LineDeviationTolerance > 0 && NewVertices.Num() >= 3)
+		if (LineDeviationTolerance > 0 && OutNewVertexIndices.Num() >= 3)
 		{
-			Marked.SetNumZeroed(NewVertices.Num());
+			Marked.SetNumZeroed(OutNewVertexIndices.Num());
 
 			// mark the first and last vertices to make sure we keep them
 			Marked[0] = true;
 			Marked.Last() = true;
-			SimplifyDouglasPeucker(LineDeviationTolerance, NewVertices, 0, NewVertices.Num()-1, Marked);
+			SimplifyDouglasPeucker(LineDeviationTolerance, Vertices, OutNewVertexIndices, 0, OutNewVertexIndices.Num()-1, Marked);
 		}
 
-		// STAGE 3. Copy back values in Vertices
-		if (Marked.IsEmpty())
-		{
-			// we have not run Douglas-Pecker so we can just copy the list
-			Vertices = MoveTemp(NewVertices);
-		}
-		else
+		// STAGE 3. Copy back values in OutNewVertexIndices
+		if (!Marked.IsEmpty()) // If we have not run Douglas-Pecker, OutNewVertexIndices contains our final result
 		{
 			// only keep the marked ones
-			Vertices.Reset();
+			int32 OutNumNewVertexIndices = 0;
 			const int32 MarkedCount = Marked.Num();
-			for (int32 Index=0; Index<MarkedCount; ++Index)
+			for (int32 Index = 0; Index < MarkedCount; ++Index)
 			{
 				if (Marked[Index])
 				{
-					Vertices.Add(NewVertices[Index]);
+					OutNewVertexIndices[OutNumNewVertexIndices++] = OutNewVertexIndices[Index];
 				}
 			}
+			OutNewVertexIndices.SetNum(OutNumNewVertexIndices);
+		}
+
+		if (OutNewVertexIndices.Num() != Vertices.Num())
+		{
+			check(OutNewVertexIndices.Num() < Vertices.Num());
+			for (int32 i = 0; i < OutNewVertexIndices.Num(); ++i)
+			{
+				Vertices[i] = Vertices[OutNewVertexIndices[i]]; // This only works since OutNewVertexIndices is sorted from lowest to higest vertex index
+			}
+
+			Vertices.SetNum(OutNewVertexIndices.Num());
 		}
 	}
 
@@ -586,9 +628,9 @@ private:
 	//            v[] = polyline array of vertex points
 	//            j,k = indices for the subchain v[j] to v[k]
 	//    Output: mk[] = array of markers matching vertex array v[]
-	static void SimplifyDouglasPeucker(T Tolerance, const TArray<VectorType>& Vertices, int32 j, int32 k, TArray<bool>& Marked)
+	static void SimplifyDouglasPeucker(T Tolerance, const TArray<VectorType>& Vertices, const TArray<int32>& VertexIndices, int32 j, int32 k, TArray<bool>& Marked)
 	{
-		Marked.SetNum(Vertices.Num());
+		check(Marked.Num() == VertexIndices.Num());
 		if (k <= j + 1) // there is nothing to simplify
 			return;
 
@@ -596,13 +638,13 @@ private:
 		int maxi = j;          // index of vertex farthest from S
 		T maxd2 = 0;         // distance squared of farthest vertex
 		T tol2 = Tolerance * Tolerance;  // tolerance squared
-		SegmentType S = SegmentType(Vertices[j], Vertices[k]);    // segment from v[j] to v[k]
+		SegmentType S = SegmentType(Vertices[VertexIndices[j]], Vertices[VertexIndices[k]]);    // segment from v[j] to v[k]
 
 		// test each vertex v[i] for max distance from S
 		// Note: this works in any dimension (2D, 3D, ...)
 		for (int i = j + 1; i < k; i++)
 		{
-			T dv2 = S.DistanceSquared(Vertices[i]);
+			T dv2 = S.DistanceSquared(Vertices[VertexIndices[i]]);
 			if (dv2 <= maxd2)
 				continue;
 			// v[i] is a max vertex
@@ -614,8 +656,8 @@ private:
 			// split the polyline at the farthest vertex from S
 			Marked[maxi] = true;      // mark v[maxi] for the simplified polyline
 			// recursively simplify the two subpolylines at v[maxi]
-			SimplifyDouglasPeucker(Tolerance, Vertices, j, maxi, Marked);  // polyline v[j] to v[maxi]
-			SimplifyDouglasPeucker(Tolerance, Vertices, maxi, k, Marked);  // polyline v[maxi] to v[k]
+			SimplifyDouglasPeucker(Tolerance, Vertices, VertexIndices, j, maxi, Marked);  // polyline v[j] to v[maxi]
+			SimplifyDouglasPeucker(Tolerance, Vertices, VertexIndices, maxi, k, Marked);  // polyline v[maxi] to v[k]
 		}
 		// else the approximation is OK, so ignore intermediate Vertices
 		return;
