Move DynamicMeshBulkEdit to the dynamic mesh modules

#Synced-CL 51402414

#Synced-CL 51402414

diff --git a/Engine/Source/Runtime/GeometryCore/Private/DynamicMesh/DynamicMeshBulkEdit.cpp b/Engine/Source/Runtime/GeometryCore/Private/DynamicMesh/DynamicMeshBulkEdit.cpp
new file mode 100644
--- /dev/null
+++ b/Engine/Source/Runtime/GeometryCore/Private/DynamicMesh/DynamicMeshBulkEdit.cpp
@@ -0,0 +1,154 @@
+// Copyright Epic Games, Inc. All Rights Reserved.
+
+#include "DynamicMesh/DynamicMeshBulkEdit.h"
+
+#include <fstream>
+
+namespace UE 
+{
+namespace Geometry
+{
+
+	FDynamicMeshBulkEdit::FDynamicMeshBulkEdit(FDynamicMesh3& EditMesh, bool bInBuildEdgeLists)
+		: Mesh(EditMesh)
+		, bBuildEdgeLists(bInBuildEdgeLists)
+	{
+	}
+
+	FDynamicMeshBulkEdit::~FDynamicMeshBulkEdit()
+	{
+		Finalize();
+	}
+
+	void FDynamicMeshBulkEdit::Initialize(int VertexCount, int TriangleCount, int EdgeCount)
+	{
+		Mesh.Clear();
+		check(Mesh.VertexEdgeLists.Size() == 0);
+
+		Mesh.Vertices.Resize(VertexCount);
+		Mesh.VertexRefCounts.InitDense(VertexCount);
+
+		Mesh.Triangles.Resize(TriangleCount);
+		Mesh.TriangleRefCounts.InitDense(TriangleCount);
+		Mesh.TriangleEdges.Resize(TriangleCount);
+
+		Mesh.Edges.Resize(EdgeCount);
+		Mesh.EdgeRefCounts.InitDense(EdgeCount);
+	}
+
+	void FDynamicMeshBulkEdit::Finalize()
+	{
+		if (Mesh.Edges.IsEmpty())
+		{
+			BuildTopology();
+		}
+
+		if (bBuildEdgeLists && Mesh.VertexEdgeLists.Size() == 0)
+		{
+			BuildEdgeLists();
+		}
+	}
+
+	void FDynamicMeshBulkEdit::BuildTopology()
+	{
+		check(Mesh.Edges.IsEmpty());
+		check(Mesh.EdgeRefCounts.IsEmpty());
+
+		struct FSortEdge
+		{
+			int32 HalfEdgeIndex;
+			int32 VtxA, VtxB;
+
+			FORCEINLINE bool operator<(const FSortEdge& other) const
+			{
+				return VtxA < other.VtxA || (VtxA == other.VtxA && VtxB < other.VtxB);
+			}
+		};
+
+		TArray<FSortEdge> SortedHalfEdges;
+
+		const int32 NumHalfEdges = Mesh.TriangleCount() * 3;
+		SortedHalfEdges.Reserve(NumHalfEdges);
+	
+		for (int32 i = 0; i < NumHalfEdges; ++i)
+		{
+			const int32 V0 = Mesh.Triangles[i/3][ i   %3];
+			const int32 V1 = Mesh.Triangles[i/3][(i+1)%3];
+			SortedHalfEdges.Emplace( i, V0 < V1 ? V0 : V1, V0 < V1 ? V1 : V0 );
+		}
+		SortedHalfEdges.Sort();
+	
+		// linear walk to count number of edges
+		int32 NumEdges = 0;
+		for (int32 HID=0; HID < NumHalfEdges; ++HID)
+		{
+			NumEdges++;
+			if (HID<NumHalfEdges-1 && 
+				SortedHalfEdges[HID].VtxA == SortedHalfEdges[HID+1].VtxA &&
+				SortedHalfEdges[HID].VtxB == SortedHalfEdges[HID+1].VtxB)
+			{
+				++HID;
+			}
+		}
+	
+		Mesh.Edges.Resize(NumEdges);
+		Mesh.EdgeRefCounts.InitDense(NumEdges);
+
+		// linear walk to assign edges and halfedges
+		int32 EdgeIdx = 0;
+		for (int32 HID=0; HID<NumHalfEdges; ++HID)
+		{
+			if (HID<NumHalfEdges-1 && 
+				SortedHalfEdges[HID].VtxA == SortedHalfEdges[HID+1].VtxA && 
+				SortedHalfEdges[HID].VtxB == SortedHalfEdges[HID+1].VtxB)
+			{
+				const int32 HalfEdgeIndex0 = SortedHalfEdges[HID  ].HalfEdgeIndex;
+				const int32 HalfEdgeIndex1 = SortedHalfEdges[HID+1].HalfEdgeIndex;
+
+				FDynamicMesh3::FEdge Edge;
+				Edge.Vert[0] = Mesh.Triangles[HalfEdgeIndex0 / 3][HalfEdgeIndex0 % 3];
+				Edge.Vert[1] = Mesh.Triangles[HalfEdgeIndex1 / 3][HalfEdgeIndex1 % 3];
+				check(Edge.Vert[1] == Mesh.Triangles[HalfEdgeIndex0 / 3][ (HalfEdgeIndex0+1) % 3]);
+
+				Edge.Tri[0] = HalfEdgeIndex0 / 3;
+				Edge.Tri[1] = HalfEdgeIndex1 / 3;
+
+				Mesh.Edges[EdgeIdx] = Edge;
+				Mesh.TriangleEdges[Edge.Tri[0]][ HalfEdgeIndex0 % 3 ] = EdgeIdx;
+				Mesh.TriangleEdges[Edge.Tri[1]][ HalfEdgeIndex1 % 3 ] = EdgeIdx;
+
+				// we found a pair, skip one more time
+				++HID;
+			}
+			else
+			{
+				const int32 HalfEdgeIndex = SortedHalfEdges[HID].HalfEdgeIndex;
+			
+				FDynamicMesh3::FEdge Edge;
+				Edge.Tri[0] = HalfEdgeIndex / 3;
+				Edge.Tri[1] = IndexConstants::InvalidID;
+
+				Edge.Vert[0] = Mesh.Triangles[HalfEdgeIndex / 3][HalfEdgeIndex % 3];
+				Edge.Vert[1] = Mesh.Triangles[HalfEdgeIndex / 3][(HalfEdgeIndex+1) % 3];
+
+				Mesh.Edges[EdgeIdx] = Edge;
+				Mesh.TriangleEdges[Edge.Tri[0]][HalfEdgeIndex % 3] = EdgeIdx;
+			}
+		
+			EdgeIdx++;
+		}
+	}
+
+	void FDynamicMeshBulkEdit::BuildEdgeLists()
+	{
+		// populate edge lists
+		Mesh.VertexEdgeLists.ResizeAndAllocateBlocks(Mesh.VertexCount());
+		for (typename TArray<FVector3d>::SizeType EdgeIdx=0; EdgeIdx<Mesh.Edges.Num(); ++EdgeIdx)
+		{
+			Mesh.VertexEdgeLists.Insert(Mesh.Edges[EdgeIdx].Vert[0], EdgeIdx);
+			Mesh.VertexEdgeLists.Insert(Mesh.Edges[EdgeIdx].Vert[1], EdgeIdx);
+		}
+	}
+
+} // namespace Geometry
+} // namespace UE
diff --git a/Engine/Source/Runtime/GeometryCore/Public/DynamicMesh/DynamicMeshBulkEdit.h b/Engine/Source/Runtime/GeometryCore/Public/DynamicMesh/DynamicMeshBulkEdit.h
new file mode 100644
--- /dev/null
+++ b/Engine/Source/Runtime/GeometryCore/Public/DynamicMesh/DynamicMeshBulkEdit.h
@@ -0,0 +1,139 @@
+// Copyright Epic Games, Inc. All Rights Reserved.
+
+#pragma once
+
+#include "DynamicMesh/DynamicMesh3.h"
+#include "Templates/UnrealTemplate.h"
+#include "Containers/Array.h"
+#include "IndexTypes.h"
+
+namespace UE 
+{
+namespace Geometry
+{
+
+class FDynamicMeshBulkEdit : FNoncopyable
+{
+public:
+	GEOMETRYCORE_API FDynamicMeshBulkEdit(FDynamicMesh3& EditMesh, bool bInBuildEdgeLists = true);
+	GEOMETRYCORE_API ~FDynamicMeshBulkEdit();
+
+	// Set the vertex, triangle and edge count.
+	// 
+	// If EdgeCount is set > 0, SetTriEdgeIndices and SetEdge should be called for all triangles and edges respectively.
+	// 
+	// If EdgeCount is set to 0, these methods do not need to be called, the edge topology will be built in Finalize.
+	// 
+	// Providing edges explicitly, if available, is of course the faster option.
+	// 
+	// SetVertex, SetTriVertexIndices should be called exactly once per vertex/triangle.
+	//
+	GEOMETRYCORE_API void Initialize(int VertexCount, int TriangleCount, int EdgeCount);
+
+	// per-triangle vertex indices, which correspond to the vertices provided by SetVertex.
+	// not thread-safe
+	inline void SetTriVertexIndices(const int TriangleID, const FIndex3i& VertexIndices)
+	{
+		Mesh.Triangles[TriangleID] = VertexIndices;
+		Mesh.VertexRefCounts.Increment(VertexIndices.A);
+		Mesh.VertexRefCounts.Increment(VertexIndices.B);
+		Mesh.VertexRefCounts.Increment(VertexIndices.C);
+	}
+
+	// per-triangle edge indices, which correspond to the edges provided by SetEdge, thread-safe
+	inline void SetTriEdgeIndices(const int TriangleID, const FIndex3i& EdgeIndices)
+	{
+		Mesh.TriangleEdges[TriangleID] = EdgeIndices;
+	}
+
+	// thread-safe
+	inline void SetVertex(const int VertexID, const FVector3d& Position)
+	{
+		Mesh.Vertices[VertexID] = Position;
+	}
+
+	// thread-safe
+	inline void SetEdge(const int EdgeID, const FDynamicMesh3::FEdge& Edge)
+	{
+		Mesh.Edges[EdgeID] = Edge;
+	}
+
+	// Explicit finalization (also called in dtor), after which the DynamicMesh is usable.
+	// At the minimum the vertices and vertex indices need to be provided (the topology will be built),
+	// or, more efficiencly, the edges are provided directly.
+	GEOMETRYCORE_API void Finalize();
+	
+private:
+
+	void BuildTopology();
+
+	void BuildEdgeLists();
+
+	FDynamicMesh3& Mesh;
+	bool bBuildEdgeLists { true };
+};
+
+
+template <typename VertexAccessFuncType,                //< FVector3d (*)(int VertexIndex)
+          typename TriangleVertexIndicesAccessFuncType> //< FIndex3i (*)(int TriangleIndex)
+inline void BulkConvertToDynamicMesh3(
+	FDynamicMesh3& TargetMesh,
+	int VertexCount,
+	int TriangleCount,
+	const VertexAccessFuncType& VertexAccessFunc,
+	const TriangleVertexIndicesAccessFuncType& TriangleVertexIndicesAccessFunc)
+{
+	FDynamicMeshBulkEdit DynamicMeshBulkEdit(TargetMesh);
+	DynamicMeshBulkEdit.Initialize(VertexCount, TriangleCount, 0);
+
+	// parallel for would also work
+	for (int VID=0; VID < VertexCount; ++VID)
+	{
+		DynamicMeshBulkEdit.SetVertex(VertexAccessFuncType(VID));
+	}
+
+	for (int TID=0; TID < TriangleCount; ++TID)
+	{
+		DynamicMeshBulkEdit.SetTriVertexIndices(TriangleVertexIndicesAccessFunc(TID));
+	}
+}						   
+
+template <typename VertexAccessFuncType,                //< FVector3d (*)(int VertexIndex)
+          typename TriangleVertexIndicesAccessFuncType, //< FIndex3i (*)(int TriangleIndex)
+		  typename TriangleEdgeIndicesAccessFuncType,   //< FIndex3i (*)(int TriangleIndex)
+		  typename EdgeAccessFuncType>                  //< FDynamicMesh3::FEdge (*)(int EdgeIdex)
+inline void BulkConvertToDynamicMesh3(
+	FDynamicMesh3& TargetMesh,
+	int VertexCount,
+	int TriangleCount,
+	int EdgeCount,
+	const VertexAccessFuncType& VertexAccessFunc,
+	const TriangleVertexIndicesAccessFuncType& TriangleVertexIndicesAccessFunc,
+	const TriangleEdgeIndicesAccessFuncType& TriangleEdgeIndicesAccessFunc,
+	const EdgeAccessFuncType& EdgeAccessFunc)
+{
+	FDynamicMeshBulkEdit DynamicMeshBulkEdit(TargetMesh);
+	DynamicMeshBulkEdit.Initialize(VertexCount, TriangleCount, EdgeCount);
+
+	// parallel for would also work
+	for (int VID=0; VID < VertexCount; ++VID)
+	{
+		DynamicMeshBulkEdit.SetVertex(VertexAccessFunc(VID));
+	}
+
+	for (int TID=0; TID < TriangleCount; ++TID)
+	{
+		DynamicMeshBulkEdit.SetTriVertexIndices(TriangleVertexIndicesAccessFunc(TID));
+		DynamicMeshBulkEdit.SetTriEdgeIndices(TriangleEdgeIndicesAccessFunc(TID));
+	}
+
+	for (int EID=0; EID < EdgeCount; ++EID)
+	{
+		DynamicMeshBulkEdit.SetEdge(EID, EdgeAccessFunc(EID));
+	}
+}						   
+
+
+
+} // namespace Geometry
+} // namespace UE
