Added allocated size to polygons

#rnx
#Synced-CL 51379665

diff --git a/Engine/Source/Runtime/GeometryCore/Public/Curve/GeneralPolygon2.h b/Engine/Source/Runtime/GeometryCore/Public/Curve/GeneralPolygon2.h
--- a/Engine/Source/Runtime/GeometryCore/Public/Curve/GeneralPolygon2.h
+++ b/Engine/Source/Runtime/GeometryCore/Public/Curve/GeneralPolygon2.h
@@ -359,6 +359,18 @@ public:
 		}
 	}
 
+	[[nodiscard]] SIZE_T GetAllocatedSize(void) const
+	{
+		SIZE_T AllocatedSize = Outer.GetVertices().GetAllocatedSize();
+
+		AllocatedSize += Holes.GetAllocatedSize();
+		for (const TPolygon2<T>& Hole : Holes)
+		{
+			AllocatedSize += Hole.GetVertices().GetAllocatedSize();
+		}
+
+		return AllocatedSize;
+	}
 };
 
 typedef TGeneralPolygon2<double> FGeneralPolygon2d;
