Initial MeshPartition mesh skirt transformer

#Synced-CL 51481065

diff --git a/Engine/Plugins/Experimental/MeshPartition/Source/MeshPartitionEditor/Private/MeshPartitionSkirtTransformer.cpp b/Engine/Plugins/Experimental/MeshPartition/Source/MeshPartitionEditor/Private/MeshPartitionSkirtTransformer.cpp
new file mode 100644
--- /dev/null
+++ b/Engine/Plugins/Experimental/MeshPartition/Source/MeshPartitionEditor/Private/MeshPartitionSkirtTransformer.cpp
@@ -0,0 +1,100 @@
+// Copyright Epic Games, Inc. All Rights Reserved.
+
+#include "MeshPartitionSkirtTransformer.h"
+
+#include "MeshPartitionDependencyInterface.h"
+
+#include "IndexTypes.h"
+#include "MeshAdapter/MeshBoundaries.h"
+
+namespace UE::MeshPartition
+{
+
+bool FMeshSkirtTransformer::Execute(MeshPartition::FTransformerContext& InTransformerContext) const
+{
+	// Helper to add a mesh skirt to a given MeshData
+	auto AddMeshSkirt = [this](FMeshData &MeshData) -> void
+	{
+		using namespace UE::Geometry;
+
+		MeshBoundaries::FBoundaryFilterOptions FilterOptions;
+		FilterOptions.MinPerimeter = BoundaryMinPerimeter;
+
+		TArray<int32> BoundaryVertLoops, BoundaryTriLoops;
+		TArray<int32> WeldMap = MeshBoundaries::ComputeVertexWeldMapping(MeshData, VertexSnapTolerance);
+		MeshBoundaries::ComputeBoundaryLoops(MeshData, BoundaryVertLoops, &BoundaryTriLoops, [&WeldMap](int32 VID) { return WeldMap[VID]; }, FilterOptions);
+
+		TArray<FName> WeightLayerNames = MeshData.GetWeightLayerNames();
+		auto AddVertexFn = [this, &WeightLayerNames, &MeshData](FIndex2i PrevNextVert, FIndex2i PrevNextTri, int32 SourceVID) -> int32
+		{
+			FVector3d SourcePos = MeshData.GetVertex(SourceVID);
+			FVector3d PrevPos = MeshData.GetVertex(PrevNextVert.A);
+			FVector3d NextPos = MeshData.GetVertex(PrevNextVert.B);
+			FVector3d SourceN = (FVector3d)MeshData.GetVertexNormal(SourceVID);
+			FVector3d PrevEdgeOffDir = SourceN ^ (SourcePos - PrevPos);
+			FVector3d NextEdgeOffDir = SourceN ^ (NextPos - SourcePos);
+			PrevEdgeOffDir.Normalize();
+			NextEdgeOffDir.Normalize();
+			FVector3d OffDir = PrevEdgeOffDir + NextEdgeOffDir;
+			OffDir.Normalize();
+			FVector3d UsePushDirection = PushDirection;
+			if (PushMethod == EMeshSkirtDirectionMethod::VertexNormal)
+			{
+				UsePushDirection = -SourceN;
+			}
+			FVector3d NewPos = SourcePos + OffDir * Width + UsePushDirection * PushDown;
+			int32 NewVID = MeshData.AppendVertex(NewPos);
+			// Copy attribute values from source to the new vertex
+			for (const FName& LayerName : WeightLayerNames)
+			{
+				MeshData.SetWeightLayerValue(LayerName, NewVID, MeshData.GetWeightLayerValue(LayerName, SourceVID));
+			}
+			MeshData.SetVertexUV(NewVID, MeshData.GetVertexUV(SourceVID));
+			// We directly copy the source vertex normal also (to keep skirt geometry shading consistent, regardless of the skirt geometry)
+			MeshData.SetVertexNormal(NewVID, MeshData.GetVertexNormal(SourceVID));
+
+			return NewVID;
+		};
+		auto AddTriFn = [&MeshData](const FIndex3i& Tri, int32 SourceTID)
+		{
+			int32 NewTID = MeshData.AppendTriangle(Tri);
+			if (ensure(NewTID >= 0))
+			{
+				MeshData.SetBaseID(NewTID, MeshData.GetBaseID(SourceTID));
+			}
+		};
+		MeshBoundaries::ExtrudeBoundaryLoopEdges(MeshData, BoundaryVertLoops, BoundaryTriLoops, AddVertexFn, AddTriFn);
+	};
+
+	TArray<MeshPartition::FTransformerUnit> Results;
+	Results.SetNum(InTransformerContext.TransformerUnits.Num());
+	ParallelFor(Results.Num(), [&Results, this, &AddMeshSkirt, &InTransformerContext](int32 TransformerIdx)
+	{
+		const MeshPartition::FTransformerUnit& InTransformerUnit = InTransformerContext.TransformerUnits[TransformerIdx];
+		MeshPartition::FMeshData ToSkirt;
+		// TODO: Is there a way to edit the meshes in-place rather than doing a full copy? Maybe a reason to do this on the static mesh transformer rather than here?
+		ToSkirt.Copy(*InTransformerUnit.MeshData);
+		AddMeshSkirt(ToSkirt);
+
+		MeshPartition::FTransformerUnit& NewEntry = Results[TransformerIdx];
+		NewEntry.Section = InTransformerUnit.Section;
+		NewEntry.MeshData = MakeShared<const MeshPartition::FMeshData>(MoveTemp(ToSkirt));
+		NewEntry.bShouldRecomputeNormals = InTransformerUnit.bShouldRecomputeNormals;
+		NewEntry.bShouldRecomputeTangents = InTransformerUnit.bShouldRecomputeTangents;
+	});
+	InTransformerContext.TransformerUnits = MoveTemp(Results);
+
+	return true;
+}
+
+void FMeshSkirtTransformer::GatherDependencies(MeshPartition::IDependencyInterface& InDependencies) const
+{
+	InDependencies += Width;
+	InDependencies += PushDown;
+	InDependencies += PushMethod;
+	InDependencies += PushDirection;
+	InDependencies += VertexSnapTolerance;
+	InDependencies += BoundaryMinPerimeter;
+}
+
+} // namespace UE::MeshPartition
diff --git a/Engine/Plugins/Experimental/MeshPartition/Source/MeshPartitionEditor/Public/MeshPartitionSkirtTransformer.h b/Engine/Plugins/Experimental/MeshPartition/Source/MeshPartitionEditor/Public/MeshPartitionSkirtTransformer.h
new file mode 100644
--- /dev/null
+++ b/Engine/Plugins/Experimental/MeshPartition/Source/MeshPartitionEditor/Public/MeshPartitionSkirtTransformer.h
@@ -0,0 +1,66 @@
+// Copyright Epic Games, Inc. All Rights Reserved.
+
+#pragma once
+
+#include "CoreMinimal.h"
+
+#include "MeshPartitionTransformer.h"
+
+#include "MeshPartitionSkirtTransformer.generated.h"
+
+#define UE_API MESHPARTITIONEDITOR_API
+
+namespace UE::MeshPartition
+{
+
+UENUM()
+enum class EMeshSkirtDirectionMethod
+{
+	FixedDirection,
+	VertexNormal
+};
+
+/**
+ * Mesh Partition Transformer producing mesh skirts for each meshdata in the pipeline transformer units.
+ */
+USTRUCT(MinimalAPI)
+struct FMeshSkirtTransformer : public MeshPartition::FTransformer
+{
+	GENERATED_BODY()
+
+public:
+	UE_API virtual bool Execute(MeshPartition::FTransformerContext& InTransformerContext) const override;
+
+	UE_API virtual void GatherDependencies(MeshPartition::IDependencyInterface& InDependencies) const override;
+
+private:
+
+	// Distance to offset the skirt vertices from the boundary, in the local plane of the boundary
+	UPROPERTY(EditAnywhere, Category="MeshSkirt", meta = (ClampMin = 0))
+	float Width = 10.f;
+
+	// Amount to push down the offset vertices
+	UPROPERTY(EditAnywhere, Category="MeshSkirt", meta = (ClampMin = 0))
+	float PushDown = 10.f;
+
+	// Method to choose the direction to push offset vertices
+	UPROPERTY(EditAnywhere, Category="MeshSkirt")
+	EMeshSkirtDirectionMethod PushMethod = EMeshSkirtDirectionMethod::VertexNormal;
+
+	// Direction to push the offset skirt vertices, if using a fixed direction Push Method
+	UPROPERTY(EditAnywhere, Category = "MeshSkirt", meta = (EditCondition = "DirectionMethod==EMeshSkirtDirectionMethod::FixedDirection"))
+	FVector PushDirection = FVector(0, 0, -1);
+
+	// Consider vertices closer than this distance to be the same for purposes of boundary walking
+	UPROPERTY(EditAnywhere, Category = "MeshSkirt", AdvancedDisplay, meta = (ClampMin = 0))
+	float VertexSnapTolerance = UE_SMALL_NUMBER;
+
+	// Do not add skirts to boundaries with perimeter smaller than this (to skip mesh boundaries that are not section boundaries)
+	UPROPERTY(EditAnywhere, Category = "MeshSkirt", AdvancedDisplay, meta = (ClampMin = 0))
+	float BoundaryMinPerimeter = 1000.f;
+
+};
+
+} // namespace UE::MeshPartition
+
+#undef UE_API
diff --git a/Engine/Source/Runtime/GeometryCore/Public/MeshAdapter/MeshBoundaries.h b/Engine/Source/Runtime/GeometryCore/Public/MeshAdapter/MeshBoundaries.h
new file mode 100644
--- /dev/null
+++ b/Engine/Source/Runtime/GeometryCore/Public/MeshAdapter/MeshBoundaries.h
@@ -0,0 +1,256 @@
+// Copyright Epic Games, Inc. All Rights Reserved.
+
+#pragma once
+
+#include "Async/ParallelFor.h"
+#include "Containers/Array.h"
+#include "Containers/Map.h"
+#include "IndexTypes.h"
+#include "VectorUtil.h"
+#include "VertexConnectedComponents.h"
+
+namespace UE
+{
+namespace Geometry
+{
+
+// Generic methods for computing and operating on mesh boundaries, for any mesh type that implements the MeshAdapter interface
+namespace MeshBoundaries
+{
+	/**
+	 * Compute a mapping to track vertices that are within a WeldThreshold distance of other vertices, intended for use with ComputeBoundaries
+	 * The mapping will map each vertex to a representative vertex of its welded (within-distance) group (and to itself, if unwelded)
+	 */
+	template<typename MeshType>
+	TArray<int32> ComputeVertexWeldMapping(const MeshType& Mesh, double WeldThreshold, TFunctionRef<bool(int32)> VertexIDFilter = [](int32) { return true; })
+	{
+		FVertexConnectedComponents WeldVerts(Mesh.MaxVertexID());
+		WeldVerts.ConnectCloseFilteredVertices(Mesh, WeldThreshold, VertexIDFilter);
+		TArray<int32> WeldIdx;
+		WeldIdx.SetNumUninitialized(Mesh.MaxVertexID());
+		ParallelFor(Mesh.MaxVertexID(), [&WeldIdx, &WeldVerts](int32 VID) { WeldIdx[VID] = WeldVerts.GetComponentThreadsafe(VID); });
+		return WeldIdx;
+	}
+
+	struct FBoundaryFilterOptions
+	{
+		// Filter boundary loops with perimeter less than this
+		double MinPerimeter = 0;
+		// Filter boundary loops with fewer edges than this
+		int32 MinEdges = 0;
+	};
+
+	/** 
+	 * Compute open boundary loops of a mesh.
+	 * 
+	 * @param Mesh The mesh to compute boundaries of
+	 * @param OutPackedBoundaryVertLoops The boundary loops packed as [NumVertsInLoop0, LoopVID0, LoopVID1, ... LoopVIDN, NumVertsInLoop1, etc]
+	 * @param OptionalOutPackedBoundaryTriLoops If non-null, the boundary loops packed as [NumTrisInLoop0, LoopTri0, LoopTri1, ... LoopTriN, NumTrisInLoop1, etc]. Tris are ordered s.t. LoopTri0 has boundary edge (LoopVID0, LoopVID1).
+	 * @param VertexWeldFn A function taking a vertex ID to a 'welded' ID (which must always be one of the IDs of the vertices that was welded, i.e. if vertex 3 and 5 are welded, the weld ID is either 3 or 5)
+	 * @param FilterOptions Options filtering loops from inclusion in the output, e.g. if they are too small
+	 * @return Number of boundary loops found
+	 */
+	template<typename MeshType>
+	int32 ComputeBoundaryLoops(const MeshType& Mesh, TArray<int32>& OutPackedBoundaryVertLoops, TArray<int32>* OptionalOutPackedBoundaryTriLoops,
+		TFunctionRef<int32(int32)> VertexWeldFn = [](int32 VID) {return VID;}, FBoundaryFilterOptions FilterOptions = FBoundaryFilterOptions())
+	{
+		OutPackedBoundaryVertLoops.Reset();
+		if (OptionalOutPackedBoundaryTriLoops)
+		{
+			OptionalOutPackedBoundaryTriLoops->Reset();
+		}
+
+		// 1. Find the edges w/ no paired reverse edge (i.e. the potential hole boundary edges)
+
+		TMultiMap<int32, FIndex2i> OpenEdges; // VID -> (VID, TID)
+		// remove an edge from OpenEdges, assuming the vertices are already passed through VertexWeldFn
+		auto RemoveEdge = [&OpenEdges](int32 A, int32 B)
+		{
+			int32 NumRemovedPairs = 0;
+			for (TMultiMap<int32, FIndex2i>::TKeyIterator It = OpenEdges.CreateKeyIterator(A); It; ++It)
+			{
+				if (It.Value().A == B)
+				{
+					It.RemoveCurrent();
+					return true;
+				}
+			}
+			return false;
+		};
+		auto AddEdge = [&OpenEdges, &VertexWeldFn, &RemoveEdge](int32 A, int32 B, int32 TID)
+		{
+			A = VertexWeldFn(A);
+			B = VertexWeldFn(B);
+			if (A == B) // weld collapsed edge, skip it
+			{
+				return;
+			}
+			if (!RemoveEdge(B, A))
+			{
+				OpenEdges.Add(A, FIndex2i(B, TID));
+			}
+		};
+		// Consider all tri edges as possible boundary edges
+		// Note: This is necessary for meshes where we only have a vertex and triangle buffer --
+		//  for dynamic mesh can filter to boundary edge and add a potentially much smaller subset of candidates to OpenEdges
+		for (int32 TID = 0; TID < Mesh.MaxTriangleID(); ++TID)
+		{
+			if (!Mesh.IsTriangle(TID))
+			{
+				continue;
+			}
+			FIndex3i Tri = Mesh.GetTriangle(TID);
+			for (int32 SubIdx = 0, Prev = 2; SubIdx < 3; Prev = SubIdx++)
+			{
+				AddEdge(Tri[Prev], Tri[SubIdx], TID);
+			}
+		}
+
+		// 2. Walk the open boundary edges to find the loops that define holes
+
+		const int32 MinEdgesPerLoop = FMath::Max(3, FilterOptions.MinEdges);
+		int32 NumHoles = 0;
+		while (true)
+		{
+			const auto& EdgeItr = OpenEdges.CreateConstIterator();
+			if (!EdgeItr)
+			{
+				break;
+			}
+
+			// Set up to add a new loop starting at this index
+			int32 BdryStartIdx = OutPackedBoundaryVertLoops.Add(INDEX_NONE);
+			// Helper to clear the added data, e.g. if we don't find a valid loop
+			auto ClearLoop = [BdryStartIdx, &OutPackedBoundaryVertLoops, &OptionalOutPackedBoundaryTriLoops]()
+			{
+				OutPackedBoundaryVertLoops.SetNum(BdryStartIdx);
+				if (OptionalOutPackedBoundaryTriLoops)
+				{
+					OptionalOutPackedBoundaryTriLoops->SetNum(BdryStartIdx);
+				}
+			};
+
+			int32 Start = EdgeItr.Key();
+			FIndex2i WalkPair = EdgeItr.Value();
+			int32 Walk = WalkPair.A;
+
+			OutPackedBoundaryVertLoops.Add(Start);
+			if (OptionalOutPackedBoundaryTriLoops)
+			{
+				OptionalOutPackedBoundaryTriLoops->Add(INDEX_NONE);
+				OptionalOutPackedBoundaryTriLoops->Add(WalkPair.B);
+				checkSlow(OptionalOutPackedBoundaryTriLoops->Num() == OutPackedBoundaryVertLoops.Num());
+			}
+			RemoveEdge(Start, Walk);
+			while (Walk != Start) // this will either loop back to start or dead end (note every non-dead-end iter removes an edge from OpenEdges)
+			{
+				OutPackedBoundaryVertLoops.Add(Walk);
+				FIndex2i* Next = OpenEdges.Find(Walk);
+				if (!Next)
+				{
+					// dead-ended chain of edges, discard the vertices
+					ClearLoop();
+					break;
+				}
+				int32 NextVal = Next->A;
+				if (OptionalOutPackedBoundaryTriLoops)
+				{
+					OptionalOutPackedBoundaryTriLoops->Add(Next->B);
+				}
+				
+				RemoveEdge(Walk, NextVal);
+				Walk = NextVal;
+			}
+			int32 Count = OutPackedBoundaryVertLoops.Num() - BdryStartIdx - 1;
+			if (Count < MinEdgesPerLoop) // failed to add a boundary, try again next loop
+			{
+				ClearLoop();
+				continue;
+			}
+			if (FilterOptions.MinPerimeter > 0)
+			{
+				double Perim = 0;
+				FVector3d PrevPos = Mesh.GetVertex(OutPackedBoundaryVertLoops.Last());
+				for (int32 Idx = BdryStartIdx + 1; Perim < FilterOptions.MinPerimeter && Idx < OutPackedBoundaryVertLoops.Num(); ++Idx)
+				{
+					FVector3d CurPos = Mesh.GetVertex(OutPackedBoundaryVertLoops[Idx]);
+					Perim += FVector3d::Dist(PrevPos, CurPos);
+					PrevPos = CurPos;
+				}
+				if (Perim < FilterOptions.MinPerimeter)
+				{
+					ClearLoop();
+					continue;
+				}
+			}
+			
+			OutPackedBoundaryVertLoops[BdryStartIdx] = Count; // record length of added chain
+			if (OptionalOutPackedBoundaryTriLoops)
+			{
+				(*OptionalOutPackedBoundaryTriLoops)[BdryStartIdx] = Count;
+			}
+			NumHoles++;
+		}
+		if (OptionalOutPackedBoundaryTriLoops)
+		{
+			check(OutPackedBoundaryVertLoops.Num() == OptionalOutPackedBoundaryTriLoops->Num());
+		}
+
+		return NumHoles;
+	}
+
+	/**
+	 * Add strip of offset geometry along mesh boundaries loops (as computed by ComputeBoundaries above)
+	 *
+	 * @param Mesh The mesh to extend with offset geometry
+	 * @param PackedBoundaryVertLoops The boundary loops packed as [NumVertsInLoop0, LoopVID0, LoopVID1, ... LoopVIDN, NumVertsInLoop1, etc]
+	 * @param PackedBoundaryTriLoops The boundary loops packed as [NumTrisInLoop0, LoopTri0, LoopTri1, ... LoopTriN, NumTrisInLoop1, etc]. Tris are ordered s.t. LoopTri0 has boundary edge (LoopVID0, LoopVID1).
+	 * @param AddVertexFromVertexFn A function to add a new offset vertex to the mesh, extended locally from the boundary
+	 * @param AddTriFn A function to add a new triangle to the mesh, locally extending from a source triangle
+	 */
+	template<typename MeshType>
+	void ExtrudeBoundaryLoopEdges(MeshType& Mesh, TConstArrayView<int32> PackedBoundaryVertLoops, TConstArrayView<int32> PackedBoundaryTriLoops,
+		TFunctionRef<int32(FIndex2i PrevNextVert, FIndex2i PrevNextTri, int32 SourceVID)> AddVertexFromVertexFn,
+		TFunctionRef<void(const FIndex3i& Tri, int32 SourceTID)> AddTriFn)
+	{
+		if (!ensure(PackedBoundaryVertLoops.Num() == PackedBoundaryTriLoops.Num()))
+		{
+			return;
+		}
+		TArray<int32> ExtendedVIDs; // storage for IDs of offset vertices
+		for (int32 LoopCountIdx = 0; LoopCountIdx < PackedBoundaryVertLoops.Num(); LoopCountIdx += PackedBoundaryVertLoops[LoopCountIdx] + 1)
+		{
+			const int32 LoopCount = PackedBoundaryVertLoops[LoopCountIdx];
+			if (!ensure(PackedBoundaryTriLoops[LoopCountIdx] == LoopCount))
+			{
+				return;
+			}
+			const int32 LoopStartIdx = LoopCountIdx + 1;
+			const int32 LoopLastIdx = LoopCountIdx + LoopCount;
+			ExtendedVIDs.SetNumUninitialized(LoopCount, EAllowShrinking::No);
+			for (int32 PrevIdx = LoopLastIdx - 1, Idx = LoopLastIdx, NextIdx = LoopStartIdx; NextIdx <= LoopLastIdx; PrevIdx = Idx, Idx = NextIdx++)
+			{
+				int32 SourceVID{ PackedBoundaryVertLoops[Idx] };
+				FIndex2i PrevNextVert{ PackedBoundaryVertLoops[PrevIdx], PackedBoundaryVertLoops[NextIdx] };
+				// by convention of the ComputeBoundaries method, tris are along the edge to the next vertex, so tris surrounding the loop vertex at Idx are at PrevIdx and Idx
+				FIndex2i PrevNextTri{ PackedBoundaryTriLoops[PrevIdx], PackedBoundaryTriLoops[Idx] };
+				ExtendedVIDs[Idx - LoopStartIdx] = AddVertexFromVertexFn(PrevNextVert, PrevNextTri, SourceVID);
+			}
+			for (int32 PrevIdx = LoopLastIdx, Idx = LoopStartIdx; Idx <= LoopLastIdx; PrevIdx = Idx++)
+			{
+				int32 IdxVID = PackedBoundaryVertLoops[Idx];
+				int32 PrevIdxVID = PackedBoundaryVertLoops[PrevIdx];
+				int32 OffIdxVID = ExtendedVIDs[Idx - LoopStartIdx];
+				int32 OffPrevIdxVID = ExtendedVIDs[PrevIdx - LoopStartIdx];
+				FIndex3i NearTri{ IdxVID, PrevIdxVID, OffIdxVID };
+				FIndex3i OffTri{ OffIdxVID, PrevIdxVID, OffPrevIdxVID };
+				AddTriFn(NearTri, PackedBoundaryTriLoops[Idx]);
+				AddTriFn(OffTri, PackedBoundaryTriLoops[Idx]);
+			}
+		}
+	}
+
+}
+
+} // end namespace UE::Geometry
+} // end namespace UE
