remove extraneous mod from geometrycore edgeloop

#Synced-CL 51597777

diff --git a/Engine/Source/Runtime/GeometryCore/Private/EdgeLoop.cpp b/Engine/Source/Runtime/GeometryCore/Private/EdgeLoop.cpp
--- a/Engine/Source/Runtime/GeometryCore/Private/EdgeLoop.cpp
+++ b/Engine/Source/Runtime/GeometryCore/Private/EdgeLoop.cpp
@@ -30,7 +30,7 @@ void FEdgeLoop::InitializeFromEdges(const TArray<int>& EdgesIn)
 	FIndex2i prev_ev = start_ev;
 	for (int i = 1; i < NumEdges; ++i)
 	{
-		FIndex2i next_ev = Mesh->GetEdgeV(Edges[i % NumEdges]);
+		FIndex2i next_ev = Mesh->GetEdgeV(Edges[i]);
 		Vertices[i] = IndexUtil::FindSharedEdgeVertex(prev_ev, next_ev);
 		prev_ev = next_ev;
 	}
