DynamicMesh & MeshDescription: possibility to add a FName attribute per triangle

#Synced-CL 51635322

diff --git a/Engine/Source/Runtime/Core/Public/UObject/FortniteMainBranchObjectVersions.inl b/Engine/Source/Runtime/Core/Public/UObject/FortniteMainBranchObjectVersions.inl
--- a/Engine/Source/Runtime/Core/Public/UObject/FortniteMainBranchObjectVersions.inl
+++ b/Engine/Source/Runtime/Core/Public/UObject/FortniteMainBranchObjectVersions.inl
@@ -781,5 +781,8 @@ FORTNITE_MAIN_VERSION(EnumUnderlyingType, 255)
 // Composite layer name deprecation
 FORTNITE_MAIN_VERSION(Composite_LayerNameDeprecation, 256)
 
+// FDynamicMeshAttributeSet has triangle labels
+FORTNITE_MAIN_VERSION(DynamicMeshAttributesTriangleLabels, 257)
+
 // -----<new versions can be added above this line>-------------------------------------------------
 #undef FORTNITE_MAIN_VERSION
diff --git a/Engine/Source/Runtime/GeometryCore/Private/DynamicMesh/DynamicMeshAttributeSet.cpp b/Engine/Source/Runtime/GeometryCore/Private/DynamicMesh/DynamicMeshAttributeSet.cpp
--- a/Engine/Source/Runtime/GeometryCore/Private/DynamicMesh/DynamicMeshAttributeSet.cpp
+++ b/Engine/Source/Runtime/GeometryCore/Private/DynamicMesh/DynamicMeshAttributeSet.cpp
@@ -77,6 +77,13 @@ void FDynamicMeshAttributeSet::Copy(const FDynamicMeshAttributeSet& Copy)
 
 	ResetRegisteredAttributes();
 
+	TriangleLabelAttributes.Reset();
+	for (const TPair<FName, TUniquePtr<FDynamicMeshTriangleLabelAttribute>>& AttribPair: Copy.TriangleLabelAttributes)
+	{
+		AttachTriangleLabelAttribute(AttribPair.Key,
+			static_cast<FDynamicMeshTriangleLabelAttribute*>(AttribPair.Value->MakeCopy(ParentMesh)));
+	}
+	
 	SkinWeightAttributes.Reset();
 	for (const TPair<FName, TUniquePtr<FDynamicMeshVertexSkinWeightsAttribute>>& AttribPair : Copy.SkinWeightAttributes)
 	{
@@ -181,6 +188,13 @@ void FDynamicMeshAttributeSet::CompactCopy(const FCompactMaps& CompactMaps, cons
 
 	ResetRegisteredAttributes();
 
+	TriangleLabelAttributes.Reset();
+	for (const TPair<FName, TUniquePtr<FDynamicMeshTriangleLabelAttribute>>& AttribPair: Copy.TriangleLabelAttributes)
+	{
+		AttachTriangleLabelAttribute(AttribPair.Key,
+			static_cast<FDynamicMeshTriangleLabelAttribute *>(AttribPair.Value->MakeCompactCopy(CompactMaps, ParentMesh)));
+	}
+	
 	SkinWeightAttributes.Reset();
 	for (const TPair<FName, TUniquePtr<FDynamicMeshVertexSkinWeightsAttribute>>& AttribPair : Copy.SkinWeightAttributes)
 	{
@@ -244,7 +258,12 @@ void FDynamicMeshAttributeSet::Append(const FDynamicMeshAttributeSet& ToAppend,
 	{
 		AppendHelper(PolygroupLayers[Idx], Idx < ToAppend.NumPolygroupLayers() ? ToAppend.GetPolygroupLayer(Idx) : nullptr);
 	}
-
+	
+	for (TPair<FName, TUniquePtr<FDynamicMeshTriangleLabelAttribute>>& AttribPair: TriangleLabelAttributes)
+	{
+		AppendHelper(*AttribPair.Value, ToAppend.FindTriangleLabelAttribute(AttribPair.Key));
+	}
+	
 	for (int Idx = 0; Idx < NumWeightLayers(); ++Idx)
 	{
 		AppendHelper(WeightLayers[Idx], Idx < ToAppend.NumWeightLayers() ? ToAppend.GetWeightLayer(Idx) : nullptr);
@@ -335,6 +354,11 @@ void FDynamicMeshAttributeSet::AppendDefaulted(const FDynamicMesh3::FAppendInfo&
 	{
 		AppendHelper(WeightLayers[Idx]);
 	}
+	
+	for (const TPair<FName, TUniquePtr<FDynamicMeshTriangleLabelAttribute>>& AttribPair : TriangleLabelAttributes)
+	{
+		AppendHelper(*AttribPair.Value);
+	}
 
 	for (const TPair<FName, TUniquePtr<FDynamicMeshVertexSkinWeightsAttribute>>& AttribPair : SkinWeightAttributes)
 	{
@@ -384,7 +408,12 @@ void FDynamicMeshAttributeSet::CompactInPlace(const FCompactMaps& CompactMaps)
 	{
 		PolygroupLayers[GroupIdx].CompactInPlace(CompactMaps);
 	}
-
+	
+	for (auto& [AttrName, TriangleLabelAttr]: TriangleLabelAttributes)
+	{
+		TriangleLabelAttr->CompactInPlace(CompactMaps);
+	}
+	
 	for (int WeightLayerIdx = 0; WeightLayerIdx < NumWeightLayers(); ++WeightLayerIdx)
 	{
 		WeightLayers[WeightLayerIdx].CompactInPlace(CompactMaps);
@@ -510,12 +539,18 @@ void FDynamicMeshAttributeSet::EnableMatchingAttributes(const FDynamicMeshAttrib
 	if (bClearExisting)
 	{
 		// discard existing
+		TriangleLabelAttributes.Reset();
 		SkinWeightAttributes.Reset();
 		MorphTargetAttributes.Reset();
 		GenericAttributes.Reset();
 		ResetRegisteredAttributes();
 
 		// register new attributes by name
+		for (const TPair<FName, TUniquePtr<FDynamicMeshTriangleLabelAttribute>>& AttribPair: ToMatch.TriangleLabelAttributes)
+		{
+			AttachTriangleLabelAttribute(AttribPair.Key, 
+				static_cast<FDynamicMeshTriangleLabelAttribute *>(AttribPair.Value->MakeNew(ParentMesh)));
+		}
 		for (const TPair<FName, TUniquePtr<FDynamicMeshVertexSkinWeightsAttribute>>& AttribPair : ToMatch.SkinWeightAttributes)
 		{
 			AttachSkinWeightsAttribute(AttribPair.Key,
@@ -536,11 +571,21 @@ void FDynamicMeshAttributeSet::EnableMatchingAttributes(const FDynamicMeshAttrib
 		// get rid of any attributes in current SkinWeights and Generic sets that are not in ToMatch
 		if (bDiscardExtraAttributes)
 		{
+			TriangleLabelAttributeMap ExistingTriangleLabels = MoveTemp(TriangleLabelAttributes);
 			SkinWeightAttributesMap ExistingSkinWeights = MoveTemp(SkinWeightAttributes);
 			GenericAttributesMap ExistingGenericAttributes = MoveTemp(GenericAttributes);
+			TriangleLabelAttributes.Reset();
 			SkinWeightAttributes.Reset();
 			GenericAttributes.Reset();
+			
 			ResetRegisteredAttributes();
+			for (TPair<FName, TUniquePtr<FDynamicMeshTriangleLabelAttribute>>& AttribPair : ExistingTriangleLabels)
+			{
+				if ( ToMatch.TriangleLabelAttributes.Contains(AttribPair.Key) )
+				{
+					AttachTriangleLabelAttribute(AttribPair.Key, AttribPair.Value.Release());
+				}
+			}
 			for (TPair<FName, TUniquePtr<FDynamicMeshVertexSkinWeightsAttribute>>& AttribPair : ExistingSkinWeights)
 			{
 				if ( ToMatch.SkinWeightAttributes.Contains(AttribPair.Key) )
@@ -557,6 +602,16 @@ void FDynamicMeshAttributeSet::EnableMatchingAttributes(const FDynamicMeshAttrib
 			}
 		}
 
+		// add any new triangle label attribute that did not previously exist
+		for (const TPair<FName, TUniquePtr<FDynamicMeshTriangleLabelAttribute>>& AttribPair : ToMatch.TriangleLabelAttributes)
+		{
+			if (!TriangleLabelAttributes.Contains(AttribPair.Key))
+			{
+				AttachTriangleLabelAttribute(AttribPair.Key,
+					static_cast<FDynamicMeshTriangleLabelAttribute *>(AttribPair.Value->MakeNew(ParentMesh)));	
+			}
+		}
+		
 		// add any new SkinWeight attributes that did not previously exist
 		for (const TPair<FName, TUniquePtr<FDynamicMeshVertexSkinWeightsAttribute>>& AttribPair : ToMatch.SkinWeightAttributes)
 		{
@@ -617,6 +672,11 @@ void FDynamicMeshAttributeSet::Reparent(FDynamicMesh3* NewParent)
 	{
 		WeightLayers[WeightLayerIdx].Reparent(NewParent);
 	}
+	
+	for (const TPair<FName, TUniquePtr<FDynamicMeshTriangleLabelAttribute>>& AttribPair : TriangleLabelAttributes)
+	{
+		AttribPair.Value.Get()->Reparent(NewParent);
+	}
 
 	for (const TPair<FName, TUniquePtr<FDynamicMeshVertexSkinWeightsAttribute>>& AttribPair : SkinWeightAttributes)
 	{
@@ -875,6 +935,52 @@ void FDynamicMeshAttributeSet::DisableMaterialID()
 	MaterialIDAttrib.Reset();
 }
 
+void FDynamicMeshAttributeSet::AttachTriangleLabelAttribute(FName InName, FDynamicMeshTriangleLabelAttribute* InAttribute)
+{
+	RemoveTriangleLabelAttribute(InName);
+
+	// Ensure proper ownership.
+	InAttribute->Reparent(ParentMesh);
+	InAttribute->SetName(InName);
+	TriangleLabelAttributes.Add(InName, TUniquePtr<FDynamicMeshTriangleLabelAttribute>(InAttribute));
+
+	RegisterExternalAttribute(InAttribute);
+}
+
+void FDynamicMeshAttributeSet::RemoveTriangleLabelAttribute(FName InName)
+{
+	if (TriangleLabelAttributes.Contains(InName))
+	{
+		UnregisterExternalAttribute(TriangleLabelAttributes[InName].Get());
+		TriangleLabelAttributes.Remove(InName);
+	}
+}
+
+void FDynamicMeshAttributeSet::RemoveAllTriangleLabelAttributes()
+{
+	for (const auto&[Name, AttrUniquePtr]: TriangleLabelAttributes)
+	{
+		UnregisterExternalAttribute(AttrUniquePtr.Get());
+	}
+	TriangleLabelAttributes.Reset();
+}
+
+bool FDynamicMeshAttributeSet::HasTriangleLabelAttribute(FName InName) const
+{
+	return TriangleLabelAttributes.Contains(InName);
+}
+
+FDynamicMeshTriangleLabelAttribute* FDynamicMeshAttributeSet::FindTriangleLabelAttribute(FName InName) const
+{
+	const TUniquePtr<FDynamicMeshTriangleLabelAttribute>* Found = TriangleLabelAttributes.Find(InName);
+	return Found && Found->IsValid() ? Found->Get() : nullptr;
+}
+	
+const FDynamicMeshAttributeSet::TriangleLabelAttributeMap& FDynamicMeshAttributeSet::GetTriangleLabelAttributes() const
+{
+	return TriangleLabelAttributes;
+}
+
 void FDynamicMeshAttributeSet::AttachSkinWeightsAttribute(FName InProfileName, FDynamicMeshVertexSkinWeightsAttribute* InAttribute)
 {
 	RemoveSkinWeightsAttribute(InProfileName);
@@ -1161,6 +1267,12 @@ void FDynamicMeshAttributeSet::OnNewTriangle(int TriangleID, bool bInserted)
 		int32 NewGroup = 0;
 		PolygroupLayer.SetNewValue(TriangleID, &NewGroup);
 	}
+
+	const FName NewLabel = NAME_None;
+	for (auto &[Name, TriangleLabelAttr]: TriangleLabelAttributes)
+	{
+		TriangleLabelAttr->SetNewValue(TriangleID, &NewLabel);
+	}
 }
 
 
@@ -1180,7 +1292,11 @@ void FDynamicMeshAttributeSet::OnRemoveTriangle(int TriangleID)
 	{
 		ColorLayer->OnRemoveTriangle(TriangleID);
 	}
-
+	
+	for (auto &[Name, TriangleLabelAttr]: TriangleLabelAttributes)
+	{
+		TriangleLabelAttr->OnRemoveTriangle(TriangleID);
+	}
 }
 
 void FDynamicMeshAttributeSet::OnReverseTriOrientation(int TriangleID)
@@ -1231,6 +1347,11 @@ void FDynamicMeshAttributeSet::OnSplitEdge(const FDynamicMesh3::FEdgeSplitInfo&
 	{
 		WeightLayer.OnSplitEdge(SplitInfo);
 	}
+	
+	for (auto& [Name, TriangleLabelAttr] : TriangleLabelAttributes)
+	{
+		TriangleLabelAttr->OnSplitEdge(SplitInfo);
+	}
 }
 
 void FDynamicMeshAttributeSet::OnFlipEdge(const FDynamicMesh3::FEdgeFlipInfo & flipInfo)
@@ -1263,6 +1384,11 @@ void FDynamicMeshAttributeSet::OnFlipEdge(const FDynamicMesh3::FEdgeFlipInfo & f
 	{
 		WeightLayer.OnFlipEdge(flipInfo);
 	}
+	
+	for (auto& [Name, TriangleLabelAttr] : TriangleLabelAttributes)
+	{
+		TriangleLabelAttr->OnFlipEdge(flipInfo);
+	}
 }
 
 
@@ -1296,6 +1422,11 @@ void FDynamicMeshAttributeSet::OnCollapseEdge(const FDynamicMesh3::FEdgeCollapse
 	{
 		WeightLayer.OnCollapseEdge(collapseInfo);
 	}
+	
+	for (auto& [Name, TriangleLabelAttr] : TriangleLabelAttributes)
+	{
+		TriangleLabelAttr->OnCollapseEdge(collapseInfo);
+	}
 }
 
 void FDynamicMeshAttributeSet::OnPokeTriangle(const FDynamicMesh3::FPokeTriangleInfo & pokeInfo)
@@ -1328,6 +1459,11 @@ void FDynamicMeshAttributeSet::OnPokeTriangle(const FDynamicMesh3::FPokeTriangle
 	{
 		WeightLayer.OnPokeTriangle(pokeInfo);
 	}
+	
+	for (auto& [Name, TriangleLabelAttr] : TriangleLabelAttributes)
+	{
+		TriangleLabelAttr->OnPokeTriangle(pokeInfo);
+	}
 }
 
 void FDynamicMeshAttributeSet::OnMergeEdges(const FDynamicMesh3::FMergeEdgesInfo & mergeInfo)
@@ -1360,6 +1496,11 @@ void FDynamicMeshAttributeSet::OnMergeEdges(const FDynamicMesh3::FMergeEdgesInfo
 	{
 		WeightLayer.OnMergeEdges(mergeInfo);
 	}
+	
+	for (auto& [Name, TriangleLabelAttr] : TriangleLabelAttributes)
+	{
+		TriangleLabelAttr->OnMergeEdges(mergeInfo);
+	}
 }
 
 void FDynamicMeshAttributeSet::OnMergeVertices(const DynamicMeshInfo::FMergeVerticesInfo& MergeInfo)
@@ -1392,6 +1533,11 @@ void FDynamicMeshAttributeSet::OnMergeVertices(const DynamicMeshInfo::FMergeVert
 	{
 		WeightLayer.OnMergeVertices(MergeInfo);
 	}
+	
+	for (auto& [Name, TriangleLabelAttr] : TriangleLabelAttributes)
+	{
+		TriangleLabelAttr->OnMergeVertices(MergeInfo);
+	}
 }
 
 void FDynamicMeshAttributeSet::OnSplitVertex(const DynamicMeshInfo::FVertexSplitInfo& SplitInfo, const TArrayView<const int>& TrianglesToUpdate)
@@ -1424,6 +1570,11 @@ void FDynamicMeshAttributeSet::OnSplitVertex(const DynamicMeshInfo::FVertexSplit
 	{
 		WeightLayer.OnSplitVertex(SplitInfo, TrianglesToUpdate);
 	}
+	
+	for (auto& [Name, TriangleLabelAttr] : TriangleLabelAttributes)
+	{
+		TriangleLabelAttr->OnSplitVertex(SplitInfo, TrianglesToUpdate);
+	}
 }
 
 bool FDynamicMeshAttributeSet::IsSameAs(const FDynamicMeshAttributeSet& Other, bool bIgnoreDataLayout) const
@@ -1431,7 +1582,8 @@ bool FDynamicMeshAttributeSet::IsSameAs(const FDynamicMeshAttributeSet& Other, b
 	if (UVLayers.Num() != Other.UVLayers.Num() ||
 		NormalLayers.Num() != Other.NormalLayers.Num() ||
 		PolygroupLayers.Num() != Other.PolygroupLayers.Num() ||
-		WeightLayers.Num() != Other.WeightLayers.Num())
+		WeightLayers.Num() != Other.WeightLayers.Num() ||
+		TriangleLabelAttributes.Num() != Other.TriangleLabelAttributes.Num())
 	{
 		return false;
 	}
@@ -1492,6 +1644,23 @@ bool FDynamicMeshAttributeSet::IsSameAs(const FDynamicMeshAttributeSet& Other, b
 		}
 	}
 
+	if (!TriangleLabelAttributes.IsEmpty())
+	{
+		for (const auto&[Name, LabelAttribute]: TriangleLabelAttributes)
+		{
+			const TUniquePtr<FDynamicMeshTriangleLabelAttribute>* Found = Other.TriangleLabelAttributes.Find(Name);
+			const FDynamicMeshTriangleLabelAttribute* OtherLabelAttribute = Found ? Found->Get() : nullptr;
+			if (!OtherLabelAttribute)
+			{
+				return false;
+			}
+			if (!LabelAttribute->IsSameAs(*OtherLabelAttribute, bIgnoreDataLayout))
+			{
+				return false;
+			}
+		}
+	}
+	
 	if (SkinWeightAttributes.Num() != Other.SkinWeightAttributes.Num())
 	{
 		return false;
@@ -1635,6 +1804,163 @@ void SerializeLayers(TIndirectArray<LayerType>& Layers, FArchive& Ar, const FCom
 		Layers[Index].Serialize(Ar, CompactMaps, bUseCompression);
 	}
 }
+	
+void SerializeLabels(FArchive& Ar, FDynamicMeshTriangleLabelAttribute& Labels)
+{
+	enum class EIndexSize : uint8 { U8 = 0, U16 = 1, U32 = 2 };
+	
+	if (Ar.IsLoading())
+	{
+		// read names count
+		uint32 NumNames = 0;
+		Ar << NumNames;
+		
+		constexpr uint32 ArrayMax = static_cast<uint32>(std::numeric_limits<TArray<FName>::SizeType>::max());
+		if (!ensure(NumNames <= ArrayMax))
+		{
+			Ar.SetError();
+			return;
+		}
+		
+		if (NumNames == 0)
+		{
+			return;
+		}
+		
+		// ret unique names
+		uint32 NumUniqueNames = 0;
+		Ar << NumUniqueNames;
+		
+		if (!ensure(NumUniqueNames <= ArrayMax))
+		{
+			Ar.SetError();
+			return;
+		}
+		
+		TArray<FName> UniqueNames;
+		UniqueNames.Reserve(NumUniqueNames);
+		
+		for (uint32 Index = 0; Index < NumUniqueNames; ++Index)
+		{
+			FString NameStr;
+			Ar << NameStr;
+			UniqueNames.Add(FName(NameStr));
+		}
+
+		// single FName so assign the same for all triangles
+		// or all names are unique so just assign them 
+		if (NumUniqueNames == 1 || NumUniqueNames == NumNames)
+		{
+			const bool bSingle = NumUniqueNames == 1;	
+			for (uint32 Index = 0; Index < NumNames; ++Index)
+			{
+				Labels.SetValue(Index, bSingle ? UniqueNames[0] : UniqueNames[Index]);
+			}
+			return;
+		}
+
+		// read index size
+		EIndexSize IndexSize = EIndexSize::U8;
+		Ar << IndexSize;
+
+		// read indices relative to UniqueNames 
+		for (uint32 Index = 0; Index < NumNames; ++Index)
+		{
+			uint32 NameIndex = 0;
+			if (IndexSize == EIndexSize::U8)
+			{
+				uint8 v = 0;
+				Ar << v;
+				NameIndex = v;
+			}
+			else if (IndexSize == EIndexSize::U16)
+			{
+				uint16 v = 0;
+				Ar << v;
+				NameIndex = v;
+			}
+			else
+			{
+				Ar << NameIndex;
+			}
+
+			if (ensure(UniqueNames.IsValidIndex(NameIndex)))
+			{
+				Labels.SetValue(Index, UniqueNames[NameIndex]);
+			}
+		}
+	}
+	else
+	{
+		// write names count
+		uint32 NumNames = Labels.GetParentMesh()->MaxTriangleID();
+		Ar << NumNames;
+		
+		if (NumNames == 0)
+		{
+			return;
+		}
+		
+		// get unique names
+		TMap<FName, uint32> NameToIndex;
+		TArray<FString> UniqueNames;
+		TArray<uint32> NameIndices;
+		NameIndices.Init(0, NumNames);
+
+		for (uint32 Index = 0; Index < NumNames; ++Index)
+		{
+			const FName Name = Labels.GetValue(Index);
+			if (const uint32* Found = NameToIndex.Find(Name))
+			{
+				NameIndices[Index] = *Found;
+			}
+			else
+			{
+				NameIndices[Index] = UniqueNames.Num();
+				NameToIndex.Add(Name, UniqueNames.Num());
+				UniqueNames.Emplace(Name.ToString());
+			}
+		}
+		
+		// write unique count and names
+		uint32 NumUniqueNames = UniqueNames.Num();
+		Ar << NumUniqueNames;
+		for (uint32 Index = 0; Index < NumUniqueNames; ++Index)
+		{
+			Ar << UniqueNames[Index];
+		}
+		
+		// single FName or all names are unique so nothing else to do
+		if (NumUniqueNames == 1 || NumUniqueNames == NumNames)
+		{
+			return;
+		}
+
+		// choose a convenient index size and write it
+		EIndexSize IndexSize = (NumUniqueNames <= MAX_uint8) ? EIndexSize::U8 : (NumUniqueNames <= MAX_uint16) ? EIndexSize::U16 : EIndexSize::U32;
+		Ar << IndexSize;
+
+		// write indices relative to UniqueNames 
+		for (uint32 Index = 0; Index < NumNames; ++Index)
+		{
+			if (IndexSize == EIndexSize::U8)
+			{
+				uint8 v = static_cast<uint8>(NameIndices[Index]);
+				Ar << v;
+			}
+			else if (IndexSize == EIndexSize::U16)
+			{
+				uint16 v = static_cast<uint16>(NameIndices[Index]);
+				Ar << v;
+			}
+			else
+			{
+				Ar << NameIndices[Index];
+			}
+		}
+	}
+}
+
 } // namespace FDynamicMeshAttributeSet_Local
 
 void FDynamicMeshAttributeSet::Serialize(FArchive& Ar, const FCompactMaps* CompactMaps, bool bUseCompression)
@@ -1849,6 +2175,54 @@ void FDynamicMeshAttributeSet::Serialize(FArchive& Ar, const FCompactMaps* Compa
 		}
 	}
 	
+	const bool bSerializeTriangleLabels = !Ar.IsLoading() ||
+			Ar.CustomVer(FFortniteMainBranchObjectVersion::GUID) >= FFortniteMainBranchObjectVersion::DynamicMeshAttributesTriangleLabels;
+	if (bSerializeTriangleLabels)
+	{		
+		int32 NumTriangleLabelAttributes = TriangleLabelAttributes.Num();
+		Ar << NumTriangleLabelAttributes;
+		
+		FNameAsStringProxyArchive ProxyArchive(Ar);
+		
+		if (Ar.IsLoading())
+		{
+			RemoveAllTriangleLabelAttributes();
+
+			for (int32 Index = 0; Index < NumTriangleLabelAttributes; ++Index)
+			{
+				FName Key;
+				ProxyArchive << Key;
+
+				bool IsValid;
+				Ar << IsValid;
+
+				if (IsValid)
+				{
+					TUniquePtr<FDynamicMeshTriangleLabelAttribute>& Value =
+						TriangleLabelAttributes.Emplace(Key, new FDynamicMeshTriangleLabelAttribute(ParentMesh));
+					SerializeLabels(Ar, *Value.Get());
+					Value->SetName(Key);
+					RegisterExternalAttribute(Value.Get());
+				}
+			}
+		}
+		else
+		{
+			for (TriangleLabelAttributeMap::TIterator It(TriangleLabelAttributes); It; ++It)
+			{
+				ProxyArchive << It.Key();
+
+				const TUniquePtr<FDynamicMeshTriangleLabelAttribute>& Value = It.Value();
+				bool bIsValid = Value.IsValid();
+				Ar << bIsValid;
+				if (bIsValid)
+				{
+					SerializeLabels(Ar, *Value.Get());
+				}
+			}
+		}	
+	}
+	
 	//Ar << GenericAttributes; // TODO
 }
 
@@ -1885,6 +2259,11 @@ SIZE_T FDynamicMeshAttributeSet::GetByteCount() const
 	{
 		ByteCount += WeightLayer.GetByteCount();
 	}
+	
+	for (const TPair<FName, TUniquePtr<FDynamicMeshTriangleLabelAttribute>>& TriangleLabelAttribute: TriangleLabelAttributes)
+	{
+		ByteCount += TriangleLabelAttribute.Value->GetByteCount();
+	}
 
 	for (const TPair<FName, TUniquePtr<FDynamicMeshVertexSkinWeightsAttribute>>& SkinWeightAttribute : SkinWeightAttributes)
 	{
@@ -1974,7 +2353,20 @@ bool FDynamicMeshAttributeSet::CheckValidity(bool bAllowNonmanifold, EValidityCh
 	{
 		bValid = GetWeightLayer(WeightLayerIndex)->CheckValidity(bAllowNonmanifold, FailMode) && bValid;
 	}
+	
+	for (const TPair<FName, TUniquePtr<FDynamicMeshTriangleLabelAttribute>>& Pair : TriangleLabelAttributes)
+	{
+		if (Pair.Value)
+		{
+			bValid = Pair.Value->CheckValidity(bAllowNonmanifold, FailMode) && bValid;
 
+			// If non-null, triangle label attributes are in the RegisteredAttributes array
+			int32 FoundIndex;
+			const bool bFound = RegisteredAttributes.Find(Pair.Value.Get(), FoundIndex);
+			bValid = bFound && bValid;
+		}
+	}
+	
 	for (const TPair<FName, TUniquePtr<FDynamicMeshVertexSkinWeightsAttribute>>& Pair : SkinWeightAttributes)
 	{
 		if (Pair.Value)
diff --git a/Engine/Source/Runtime/GeometryCore/Public/DynamicMesh/DynamicMeshAttributeSet.h b/Engine/Source/Runtime/GeometryCore/Public/DynamicMesh/DynamicMeshAttributeSet.h
--- a/Engine/Source/Runtime/GeometryCore/Public/DynamicMesh/DynamicMeshAttributeSet.h
+++ b/Engine/Source/Runtime/GeometryCore/Public/DynamicMesh/DynamicMeshAttributeSet.h
@@ -371,6 +371,29 @@ public:
 		return MaterialIDAttrib.Get();
 	}
 
+	/** Triangle Label attributes */
+	
+	/*
+	 * Attaches and registers the triangle label attribute with a given name.
+	 * If an attribute with this name already exists, the existing attribute will be deleted.
+	*/
+	UE_INTERNAL GEOMETRYCORE_API void AttachTriangleLabelAttribute(FName InName, FDynamicMeshTriangleLabelAttribute* InAttribute);
+
+	/* Remove a triangle label attribute matching the given name. */
+	UE_INTERNAL GEOMETRYCORE_API void RemoveTriangleLabelAttribute(FName InName);
+
+	/* Remove all triangle label attributes */
+	UE_INTERNAL GEOMETRYCORE_API void RemoveAllTriangleLabelAttributes();
+
+	/* Returns true if the list of triangle label attributes includes the given name. */
+	UE_INTERNAL bool HasTriangleLabelAttribute(FName InName) const;
+	
+	/* Returns a pointer to the triangle label attribute of the given name or null ptr if it does not exist */
+	UE_INTERNAL GEOMETRYCORE_API FDynamicMeshTriangleLabelAttribute* FindTriangleLabelAttribute(FName InName) const;
+	
+	/* Returns the existing triangle label attributes. */
+	UE_INTERNAL GEOMETRYCORE_API const TMap<FName, TUniquePtr<FDynamicMeshTriangleLabelAttribute>>& GetTriangleLabelAttributes() const;
+	
 	/** Skin weights */
 
 	/// Create a new skin weights attribute with a given skin weights profile name. If an attribute already exists with
@@ -651,6 +674,9 @@ protected:
 
 	TIndirectArray<FDynamicMeshWeightAttribute> WeightLayers;
 	TIndirectArray<FDynamicMeshPolygroupAttribute> PolygroupLayers;
+	
+	using TriangleLabelAttributeMap = TMap<FName, TUniquePtr<FDynamicMeshTriangleLabelAttribute>>;
+	TriangleLabelAttributeMap TriangleLabelAttributes;
 
 	using SkinWeightAttributesMap = TMap<FName, TUniquePtr<FDynamicMeshVertexSkinWeightsAttribute>>;
 	SkinWeightAttributesMap SkinWeightAttributes;
diff --git a/Engine/Source/Runtime/GeometryCore/Public/DynamicMesh/DynamicMeshTriangleAttribute.h b/Engine/Source/Runtime/GeometryCore/Public/DynamicMesh/DynamicMeshTriangleAttribute.h
--- a/Engine/Source/Runtime/GeometryCore/Public/DynamicMesh/DynamicMeshTriangleAttribute.h
+++ b/Engine/Source/Runtime/GeometryCore/Public/DynamicMesh/DynamicMeshTriangleAttribute.h
@@ -159,13 +159,18 @@ public:
 
 
 	/** Initialize the attribute values with InitialValue, and resize to the parent mesh's max triangle ID */
-	void Initialize(AttribValueType InitialValue = (AttribValueType)0)
+	void Initialize(AttribValueType InitialValue)
 	{
 		check(ParentMesh != nullptr);
 		AttribValues.Resize(ParentMesh->MaxTriangleID() * AttribDimension);
 		AttribValues.Fill(InitialValue);
 	}
 
+	void Initialize()
+	{
+		Initialize(GetDefaultAttributeValue());
+	}
+
 	void SetNewValue(int NewTriangleID, const AttribValueType* Data)
 	{
 		int k = NewTriangleID * AttribDimension;
@@ -397,7 +402,7 @@ public:
 
 	virtual AttribValueType GetDefaultAttributeValue()
 	{
-		return (AttribValueType)0;
+		return AttribValueType();
 	}
 
 	inline void ResizeAttribStoreIfNeeded(int TriangleID)
@@ -473,6 +478,33 @@ public:
 		return true;
 	}
 
+	/**
+	* Check validity of attribute
+	* 
+	* @param bAllowNonmanifold Accept non-manifold topology as valid. Note that this should almost always be true for attributes; non-manifold overlays are generally valid.
+	* @param FailMode Desired behavior if mesh is found invalid
+	*/
+	virtual bool CheckValidity(bool bAllowNonmanifold, EValidityCheckFailMode FailMode) const override
+	{
+		// just check that the values buffer is big enough
+		if (!ParentMesh || ParentMesh->MaxTriangleID() < 0 || static_cast<size_t>(ParentMesh->MaxTriangleID())*AttribDimension > AttribValues.Num())
+		{
+			switch (FailMode)
+			{
+			case EValidityCheckFailMode::Check:
+				check(false);
+				return false;
+			case EValidityCheckFailMode::Ensure:
+				ensure(false);
+				return false;
+			default:
+				return false;
+			}
+		}
+
+		return true;
+	}
+	
 	/**
 	 * Serialization operator for TDynamicMeshTriangleAttribute.
 	 *
@@ -631,6 +663,30 @@ public:
 	}
 };
 
+template<typename AttribValueType>
+class TDynamicMeshSingleTriangleAttribute : public TDynamicMeshTriangleAttribute<AttribValueType, 1>
+{
+	using BaseType = TDynamicMeshTriangleAttribute<AttribValueType, 1>;
+	
+public:
+	TDynamicMeshSingleTriangleAttribute() = default;
+
+	TDynamicMeshSingleTriangleAttribute(FDynamicMesh3* ParentMeshIn, bool bAutoInit = true)
+		: BaseType(ParentMeshIn, bAutoInit)
+	{}
+
+	inline AttribValueType GetValue(int TriangleID) const
+	{
+		return this->AttribValues[TriangleID];
+	}
+	
+	inline void SetValue(int TriangleID, const AttribValueType& Value)
+	{
+		this->AttribValues[TriangleID] = Value;
+	}
+};
+
+using FDynamicMeshTriangleLabelAttribute = TDynamicMeshSingleTriangleAttribute<FName>;
 
 } // end namespace UE::Geometry
 } // end namespace UE
diff --git a/Engine/Source/Runtime/MeshConversion/Private/DynamicMeshToMeshDescription.cpp b/Engine/Source/Runtime/MeshConversion/Private/DynamicMeshToMeshDescription.cpp
--- a/Engine/Source/Runtime/MeshConversion/Private/DynamicMeshToMeshDescription.cpp
+++ b/Engine/Source/Runtime/MeshConversion/Private/DynamicMeshToMeshDescription.cpp
@@ -1157,6 +1157,9 @@ void FDynamicMeshToMeshDescription::Convert_NoSharedInstances(const FDynamicMesh
 
 	// convert weight map layers
 	ConvertWeightLayers(MeshIn, MeshOut, MapV);
+	
+	// convert triangle attributes
+	ConvertTriangleAttributes(MeshIn, MeshOut, IndexToTriangleIDMap);
 
 	// convert normal seams to hard edges
 	if (ConversionOptions.bSetHardEdgesFromNormalSeams && NormalOverlay)
@@ -1342,6 +1345,68 @@ void FDynamicMeshToMeshDescription::ConvertWeightLayers(const FDynamicMesh3* Mes
 	}
 }
 
+namespace DynamicMeshToMeshDescriptionConversionHelper
+{
+	template <typename AttributeType>
+	TTriangleAttributesRef<AttributeType> GetTriangleAttribute(TAttributesSet<FTriangleID>& TriAttribsSet, const FName AttrName, const AttributeType Default)
+	{
+		TTriangleAttributesRef<AttributeType> InvalidAttribute;
+		
+		if (TriAttribsSet.HasAttributeOfType<AttributeType>(AttrName))
+		{
+			TTriangleAttributesRef<AttributeType> ExistingAttribute = TriAttribsSet.GetAttributesRef<AttributeType>(AttrName);
+			if (ExistingAttribute.GetNumChannels() == 1)
+			{
+				return ExistingAttribute;
+			}
+		}
+		else if (TriAttribsSet.HasAttribute(AttrName))
+		{ 
+			return InvalidAttribute;	
+		}
+		else
+		{
+			TriAttribsSet.RegisterAttribute<AttributeType>(AttrName, 1, Default, EMeshAttributeFlags::AutoGenerated);
+			return TriAttribsSet.GetAttributesRef<AttributeType>(AttrName);
+		}
+
+		return InvalidAttribute;
+	}
+}
+
+void FDynamicMeshToMeshDescription::ConvertTriangleAttributes(const FDynamicMesh3* MeshIn, FMeshDescription& MeshOut, const TArray<FTriangleID>& IndexToTriangleIDMap)
+{
+	using namespace DynamicMeshToMeshDescriptionConversionHelper;
+	
+	const FDynamicMeshAttributeSet* DynaMeshAttributes = MeshIn->Attributes();
+	if (!DynaMeshAttributes)
+	{
+		return;
+	}
+
+	TAttributesSet<FTriangleID>& MeshDescAttributes = MeshOut.TriangleAttributes();
+	
+	for (const auto&[AttrName, TriangleLabelAttr]: MeshIn->Attributes()->GetTriangleLabelAttributes())
+	{
+		if (TriangleLabelAttr.IsValid())
+		{
+			TTriangleAttributesRef<FName> Attribute = GetTriangleAttribute<FName>(MeshDescAttributes, AttrName, NAME_None);
+			if (ensure(Attribute.IsValid()))
+			{
+				for (int32 tid : MeshIn->TriangleIndicesItr())
+				{
+					const FTriangleID TriangleID = IndexToTriangleIDMap[tid];
+					Attribute.Set(TriangleID, TriangleLabelAttr->GetValue(tid));
+				}
+			}
+			else
+			{
+				UE_LOG(LogTemp, Warning, TEXT("FDynamicMeshToMeshDescription::ConvertTriangleAttributes - could not create attribute named %s"), *AttrName.ToString());
+			}
+		}
+	}
+}
+
 void FDynamicMeshToMeshDescription::ApplyVertexColorTransform(FVector4f& Color) const
 {
 	// There is inconsistency in how vertex colors are intended to be consumed in
diff --git a/Engine/Source/Runtime/MeshConversion/Private/MeshDescriptionToDynamicMesh.cpp b/Engine/Source/Runtime/MeshConversion/Private/MeshDescriptionToDynamicMesh.cpp
--- a/Engine/Source/Runtime/MeshConversion/Private/MeshDescriptionToDynamicMesh.cpp
+++ b/Engine/Source/Runtime/MeshConversion/Private/MeshDescriptionToDynamicMesh.cpp
@@ -177,6 +177,11 @@ struct FMorphTargetAttribCopyInfo
 	FDynamicMeshMorphTargetAttribute* DynaMeshAttribRef = nullptr;
 };
 
+struct FTriangleLabelAttribCopyInfo
+{
+	TTriangleAttributesConstRef<FName> MeshDescAttribRef;
+	FDynamicMeshTriangleLabelAttribute* DynaMeshAttribRef = nullptr;
+};
 
 void FMeshDescriptionToDynamicMesh::Convert(const FMeshDescription* MeshIn, FDynamicMesh3& MeshOut, bool bCopyTangents )
 {
@@ -432,6 +437,8 @@ void FMeshDescriptionToDynamicMesh::Convert(const FMeshDescription* MeshIn, FDyn
 
 	TArray<FSkinWeightsAttribCopyInfo> SkinWeightAttribs;
 	TArray<FMorphTargetAttribCopyInfo> MorphTargetAttribs;
+	TArray<FTriangleLabelAttribCopyInfo> TriangleLabelAttributes;
+	
 	if (!bDisableAttributes)
 	{
 		MeshOut.EnableAttributes(); // by default 1-UV layer and 1-normal layer
@@ -532,6 +539,20 @@ void FMeshDescriptionToDynamicMesh::Convert(const FMeshDescription* MeshIn, FDyn
 				BoneColorAttrib = MeshOut.Attributes()->GetBoneColors();
 			}
 		}
+		
+		// find triangle label attribs
+		const TAttributesSet<FTriangleID>& TriAttribsSet = MeshIn->TriangleAttributes();
+		TriAttribsSet.ForEach([&TriAttribsSet, &TriangleLabelAttributes, &MeshOut](const FName AttributeName, auto AttributesRef)
+		{
+			if (FSkeletalMeshAttributes::IsReservedAttributeName(AttributeName)) return;
+
+			if (TriAttribsSet.template HasAttributeOfType<FName>(AttributeName))
+			{
+				FDynamicMeshTriangleLabelAttribute* NewTriangleLabelAttrib = new FDynamicMeshTriangleLabelAttribute(&MeshOut);
+				MeshOut.Attributes()->AttachTriangleLabelAttribute(AttributeName, NewTriangleLabelAttrib);
+				TriangleLabelAttributes.Add({TriAttribsSet.GetAttributesRef<FName>(AttributeName), NewTriangleLabelAttrib});
+			}
+		});
 	}
 
 
@@ -950,6 +971,23 @@ void FMeshDescriptionToDynamicMesh::Convert(const FMeshDescription* MeshIn, FDyn
 			Pending.Add(MoveTemp(PolygroupTask));
 		}
 
+		// initialize triangle labels
+		for (const FTriangleLabelAttribCopyInfo& TriLabelAttrInfo: TriangleLabelAttributes)
+		{
+			auto TriangleLabelTask = UE::Tasks::Launch(UE_SOURCE_LOCATION, [&, TriLabelAttrInfo]()
+			{
+				if (ensure(TriLabelAttrInfo.DynaMeshAttribRef))
+				{
+					for (int32 tid : MeshOut.TriangleIndicesItr())
+					{
+						FTriangleID SourceTriangleID = TriIDMap[tid];
+						TriLabelAttrInfo.DynaMeshAttribRef->SetValue(tid, TriLabelAttrInfo.MeshDescAttribRef.Get(SourceTriangleID));
+					}
+				}
+			});
+			Pending.Add(MoveTemp(TriangleLabelTask));
+		}
+
 		// initialize weight layers
 		for (int32 WeightLayerIdx = 0; WeightLayerIdx < WeightAttribs.Num(); WeightLayerIdx++)
 		{
diff --git a/Engine/Source/Runtime/MeshConversion/Public/DynamicMeshToMeshDescription.h b/Engine/Source/Runtime/MeshConversion/Public/DynamicMeshToMeshDescription.h
--- a/Engine/Source/Runtime/MeshConversion/Public/DynamicMeshToMeshDescription.h
+++ b/Engine/Source/Runtime/MeshConversion/Public/DynamicMeshToMeshDescription.h
@@ -175,6 +175,13 @@ protected:
 	*/
 	MESHCONVERSION_API void ConvertWeightLayers(const FDynamicMesh3* MeshIn, FMeshDescription& MeshOut, const TArray<FVertexID>& IndexToVertexIDMap);
 
+	/**
+	 * Transfer triangle label attributes from DynamicMesh AttributeSet to MeshDescription.
+	 * Will copy to existing MeshDescription TriangleAttribute<FName> if one with the same name exists.
+	 * Otherwise will register a new one.
+	 */
+	UE_INTERNAL void ConvertTriangleAttributes(const FDynamicMesh3* MeshIn, FMeshDescription& MeshOut, const TArray<FTriangleID>& IndexToTriangleIDMap);
+	
 	/**
 	 * Applies an optional sRGB-to-Linear color transform on the input. The color transform
 	 * is controlled by ConversionOptions.bTransformVtxColorsSRGBToLinear.
diff --git a/Engine/Source/Runtime/MeshDescription/Private/MeshDescription.cpp b/Engine/Source/Runtime/MeshDescription/Private/MeshDescription.cpp
--- a/Engine/Source/Runtime/MeshDescription/Private/MeshDescription.cpp
+++ b/Engine/Source/Runtime/MeshDescription/Private/MeshDescription.cpp
@@ -175,6 +175,7 @@ void FMeshDescription::Serialize(FArchive& BaseAr)
 	Ar.UsingCustomVersion(FReleaseObjectVersion::GUID);
 	Ar.UsingCustomVersion(FEditorObjectVersion::GUID);
 	Ar.UsingCustomVersion(FUE5MainStreamObjectVersion::GUID);
+	Ar.UsingCustomVersion(FFortniteMainBranchObjectVersion::GUID);
 
 	if (Ar.IsLoading() && Ar.CustomVer(FReleaseObjectVersion::GUID) < FReleaseObjectVersion::MeshDescriptionNewSerialization)
 	{
@@ -2046,7 +2047,8 @@ TConstArrayView<FGuid> FMeshDescriptionBulkData::GetMeshDescriptionCustomVersion
 		FEditorObjectVersion::GUID,
 		FReleaseObjectVersion::GUID,
 		FEnterpriseObjectVersion::GUID,
-		FUE5MainStreamObjectVersion::GUID
+		FUE5MainStreamObjectVersion::GUID,
+		FFortniteMainBranchObjectVersion::GUID
 	};
 	return MakeArrayView(Versions);
 }
@@ -2057,6 +2059,7 @@ void FMeshDescriptionBulkData::Serialize( FArchive& Ar, UObject* Owner )
 	Ar.UsingCustomVersion( FReleaseObjectVersion::GUID );
 	Ar.UsingCustomVersion( FEnterpriseObjectVersion::GUID );
 	Ar.UsingCustomVersion( FUE5MainStreamObjectVersion::GUID );
+	Ar.UsingCustomVersion( FFortniteMainBranchObjectVersion::GUID );
 
 	// Make sure to serialize only actual data
 	if (Ar.ShouldSkipBulkData() || Ar.IsObjectReferenceCollector())
diff --git a/Engine/Source/Runtime/MeshDescription/Public/MeshAttributeArray.h b/Engine/Source/Runtime/MeshDescription/Public/MeshAttributeArray.h
--- a/Engine/Source/Runtime/MeshDescription/Public/MeshAttributeArray.h
+++ b/Engine/Source/Runtime/MeshDescription/Public/MeshAttributeArray.h
@@ -256,6 +256,195 @@ inline typename TEnableIf<!TIsBulkSerializable<T>::Value, FArchive>::Type& opera
 	return Ar;
 }
 
+template <>
+inline TEnableIf<!TIsBulkSerializable<FName>::Value, FArchive>::Type& operator<<( FArchive& Ar, TMeshAttributeArrayBase<FName>& Array )
+{
+	if (Ar.CustomVer(FFortniteMainBranchObjectVersion::GUID) < FFortniteMainBranchObjectVersion::DynamicMeshAttributesTriangleLabels)
+	{		
+		// see non-specialized implementation
+		if (Ar.IsLoading() &&
+		   Ar.CustomVer(FReleaseObjectVersion::GUID) != FReleaseObjectVersion::MeshDescriptionNewFormat &&
+		   Ar.CustomVer(FUE5MainStreamObjectVersion::GUID) < FUE5MainStreamObjectVersion::MeshDescriptionNewFormat)
+		{
+			Array.Extent = 1;
+		}
+		else
+		{
+			Ar << Array.Extent;
+		}
+		
+		Ar << Array.Container;
+		
+		return Ar;
+	}
+	
+	Ar << Array.Extent;
+	
+	enum class EIndexSize : uint8 { U8 = 0, U16 = 1, U32 = 2 };
+	
+	TArray<FName>& Names = Array.Container;
+	
+	if (Ar.IsLoading())
+	{
+		// read names count
+		uint32 NumNames = 0;
+		Ar << NumNames;
+		
+		constexpr uint32 ArrayMax = static_cast<uint32>(std::numeric_limits<TArray<FName>::SizeType>::max());
+		if (!ensure(NumNames <= ArrayMax))
+		{
+			Ar.SetError();
+			return Ar;
+		}
+		
+		if (NumNames == 0)
+		{
+			Names.Reset();
+			return Ar;
+		}
+		
+		// ret unique names
+		uint32 NumUniqueNames = 0;
+		Ar << NumUniqueNames;
+		
+		if (!ensure(NumUniqueNames <= ArrayMax))
+		{
+			Ar.SetError();
+			return Ar;
+		}
+		
+		TArray<FName> UniqueNames;
+		UniqueNames.Reserve(NumUniqueNames);
+		
+		for (uint32 Index = 0; Index < NumUniqueNames; ++Index)
+		{
+			FString NameStr;
+			Ar << NameStr;
+			UniqueNames.Add(FName(NameStr));
+		}
+	
+		// single FName so assign the same for all triangles
+		if (NumUniqueNames == 1)
+		{
+			Names.Init(UniqueNames[0], NumNames);
+			return Ar;
+		}
+		
+		// all names are unique 
+		if (NumUniqueNames == NumNames)
+		{
+			Names = MoveTemp(UniqueNames);
+			return Ar;
+		}
+		
+		Names.Init(NAME_None, NumNames);
+	
+		// read index size
+		EIndexSize IndexSize = EIndexSize::U8;
+		Ar << IndexSize;
+	
+		// read indices relative to UniqueNames 
+		for (uint32 Index = 0; Index < NumNames; ++Index)
+		{
+			uint32 NameIndex = 0;
+			if (IndexSize == EIndexSize::U8)
+			{
+				uint8 v = 0;
+				Ar << v;
+				NameIndex = v;
+			}
+			else if (IndexSize == EIndexSize::U16)
+			{
+				uint16 v = 0;
+				Ar << v;
+				NameIndex = v;
+			}
+			else
+			{
+				Ar << NameIndex;
+			}
+	
+			if (ensure(UniqueNames.IsValidIndex(NameIndex)))
+			{
+				Names[Index] = UniqueNames[NameIndex];
+			}
+		}
+	}
+	else
+	{
+		// write names count
+		uint32 NumNames = Names.Num();
+		Ar << NumNames;
+		
+		if (NumNames == 0)
+		{
+			return Ar;
+		}
+		
+		// get unique names
+		TMap<FName, uint32> NameToIndex;
+		TArray<FString> UniqueNames;
+		TArray<uint32> NameIndices;
+		NameIndices.Init(0, NumNames);
+	
+		for (uint32 Index = 0; Index < NumNames; ++Index)
+		{
+			const FName Name = Names[Index];
+			if (const uint32* Found = NameToIndex.Find(Name))
+			{
+				NameIndices[Index] = *Found;
+			}
+			else
+			{
+				NameIndices[Index] = UniqueNames.Num();
+				NameToIndex.Add(Name, UniqueNames.Num());
+				UniqueNames.Emplace(Name.ToString());
+			}
+		}
+		
+		// write unique count and names
+		uint32 NumUniqueNames = UniqueNames.Num();
+		Ar << NumUniqueNames;
+		
+		for (uint32 Index = 0; Index < NumUniqueNames; ++Index)
+		{
+			Ar << UniqueNames[Index];
+		}
+		
+		// single FName or all names are unique so nothing else to do
+		if (NumUniqueNames == 1 || NumUniqueNames == NumNames)
+		{
+			return Ar;
+		}
+	
+		// choose a convenient index size and write it
+		EIndexSize IndexSize = (NumUniqueNames <= MAX_uint8) ? EIndexSize::U8 : (NumUniqueNames <= MAX_uint16) ? EIndexSize::U16 : EIndexSize::U32;
+		Ar << IndexSize;
+	
+		// write indices relative to UniqueNames 
+		for (uint32 Index = 0; Index < NumNames; ++Index)
+		{
+			if (IndexSize == EIndexSize::U8)
+			{
+				uint8 v = static_cast<uint8>(NameIndices[Index]);
+				Ar << v;
+			}
+			else if (IndexSize == EIndexSize::U16)
+			{
+				uint16 v = static_cast<uint16>(NameIndices[Index]);
+				Ar << v;
+			}
+			else
+			{
+				Ar << NameIndices[Index];
+			}
+		}
+	}
+	
+	return Ar;
+}
+
+
 template <typename T>
 inline typename TEnableIf<TIsBulkSerializable<T>::Value, FArchive>::Type& operator<<( FArchive& Ar, TMeshAttributeArrayBase<T>& Array )
 {
