make the dynamic mesh's attribute set generic attribute accessor return nullptr if the attribute is not found, rather than failing a check(), to match how most of the other accessors work when requesting an attrib that doesn't exist

#Synced-CL 51695666

diff --git a/Engine/Source/Runtime/GeometryCore/Public/DynamicMesh/DynamicMeshAttributeSet.h b/Engine/Source/Runtime/GeometryCore/Public/DynamicMesh/DynamicMeshAttributeSet.h
--- a/Engine/Source/Runtime/GeometryCore/Public/DynamicMesh/DynamicMeshAttributeSet.h
+++ b/Engine/Source/Runtime/GeometryCore/Public/DynamicMesh/DynamicMeshAttributeSet.h
@@ -610,12 +610,20 @@ public:
 
 	FDynamicMeshAttributeBase* GetAttachedAttribute(FName AttribName)
 	{
-		return GenericAttributes[AttribName].Get();
+		if (TUniquePtr<FDynamicMeshAttributeBase>* AttribPtrPtr = GenericAttributes.Find(AttribName))
+		{
+			return AttribPtrPtr->Get();
+		}
+		return nullptr;
 	}
 
 	const FDynamicMeshAttributeBase* GetAttachedAttribute(FName AttribName) const
 	{
-		return GenericAttributes[AttribName].Get();
+		if (const TUniquePtr<FDynamicMeshAttributeBase>* AttribPtrPtr = GenericAttributes.Find(AttribName))
+		{
+			return AttribPtrPtr->Get();
+		}
+		return nullptr;
 	}
 	int NumAttachedAttributes() const
 	{
