clean up OpeningAngleD and make an equivalent OpeningAngleR in GeometryCore::VectorUtil

#Synced-CL 51925774

diff --git a/Engine/Source/Runtime/GeometryCore/Public/VectorUtil.h b/Engine/Source/Runtime/GeometryCore/Public/VectorUtil.h
--- a/Engine/Source/Runtime/GeometryCore/Public/VectorUtil.h
+++ b/Engine/Source/Runtime/GeometryCore/Public/VectorUtil.h
@@ -508,16 +508,21 @@ namespace VectorUtil
 
 
 	/**
-	 * @return angle between vectors (A-CornerPt) and (B-CornerPt)
+	 * @return angle in degrees between vectors (A-CornerPt) and (B-CornerPt)
 	 */
 	template<typename RealType>
-	inline RealType OpeningAngleD(TVector<RealType> A, TVector<RealType> B, const TVector<RealType>& P)
+	inline RealType OpeningAngleD(const TVector<RealType>& A, const TVector<RealType>& B, const TVector<RealType>& P)
 	{
-		A -= P; 
-		Normalize(A);
-		B -= P;
-		Normalize(B);
-		return AngleD(A, B);
+		return AngleD(Normalized(A - P), Normalized(B - P));
+	}
+
+	/**
+	 * @return angle in radians between vectors (A-CornerPt) and (B-CornerPt)
+	 */
+	template<typename RealType>
+	inline RealType OpeningAngleR(const TVector<RealType>& A, const TVector<RealType>& B, const TVector<RealType>& P)
+	{
+		return AngleR(Normalized(A - P), Normalized(B - P));
 	}
 
 	/**
