Added a conversion operator to geometry core oriented box types so it can convert to FBox.

#Synced-CL 52063422

#Synced-CL 52063422

diff --git a/Engine/Source/Runtime/GeometryCore/Public/OrientedBoxTypes.h b/Engine/Source/Runtime/GeometryCore/Public/OrientedBoxTypes.h
--- a/Engine/Source/Runtime/GeometryCore/Public/OrientedBoxTypes.h
+++ b/Engine/Source/Runtime/GeometryCore/Public/OrientedBoxTypes.h
@@ -70,6 +70,13 @@ struct TOrientedBox3
 		ToRet.ExtentZ = (double)Extents.Z;
 		return ToRet;
 	}
+	
+	explicit operator TBox<RealType>() const
+	{
+		TBox<RealType> ToRet;
+		EnumerateCorners([&](const TVector<RealType>& Corner) { ToRet += Corner; });
+		return ToRet;
+	}
 
 	// Convert from Core Math's FOrientedBox
 	TOrientedBox3(const FOrientedBox& Box) :
