add option for dynamic mesh overlay's per-vertex build to make the elements exactly 1:1 with vertices, so caller doesn't need to remap

#Synced-CL 52050116

diff --git a/Engine/Source/Runtime/GeometryCore/Private/DynamicMesh/DynamicMeshOverlay.cpp b/Engine/Source/Runtime/GeometryCore/Private/DynamicMesh/DynamicMeshOverlay.cpp
--- a/Engine/Source/Runtime/GeometryCore/Private/DynamicMesh/DynamicMeshOverlay.cpp
+++ b/Engine/Source/Runtime/GeometryCore/Private/DynamicMesh/DynamicMeshOverlay.cpp
@@ -164,20 +164,35 @@ void TDynamicMeshOverlay<RealType, ElementSize>::CreateFromPredicate(TFunctionRe
 }
 
 template<typename RealType, int ElementSize>
-void TDynamicMeshOverlay<RealType, ElementSize>::CreatePerVertex(RealType InitElementValue)
+void TDynamicMeshOverlay<RealType, ElementSize>::CreatePerVertex(RealType InitElementValue, bool bExactMapIDs)
 {
 	ClearElements(); // deletes all elements and initializes triangles to be 1:1 w/ parentmesh IDs
+	TStaticArray<RealType, ElementSize> DefaultElement;
+	for (int32 Idx = 0; Idx < ElementSize; ++Idx)
+	{
+		DefaultElement[Idx] = InitElementValue;
+	}
 	TArray<int32> VIDtoEID;
-	bool bIsCompactV = ParentMesh->IsCompactV();
-	if (!bIsCompactV)
+	bool bNeedsRemap = !bExactMapIDs && !ParentMesh->IsCompactV();
+	if (bNeedsRemap)
 	{
 		VIDtoEID.SetNumUninitialized(ParentMesh->MaxVertexID());
 	}
 	for (int32 VertexID : ParentMesh->VertexIndicesItr())
 	{
-		int32 ElementID = AppendElement(InitElementValue);
+		int32 ElementID;
+		if (bExactMapIDs)
+		{
+			ElementID = VertexID;
+			EMeshResult Result = InsertElement(VertexID, DefaultElement.GetData());
+			checkSlow(Result == EMeshResult::Ok); // because we allocate in increasing sequential order, shouldn't be possible InsertElement to return failure
+		}
+		else
+		{
+			ElementID = AppendElement(InitElementValue);
+		}
 		ParentVertices.InsertAt(VertexID, ElementID, FDynamicMesh3::InvalidID);
-		if (!bIsCompactV)
+		if (bNeedsRemap)
 		{
 			VIDtoEID[VertexID] = ElementID;
 		}
@@ -194,7 +209,7 @@ void TDynamicMeshOverlay<RealType, ElementSize>::CreatePerVertex(RealType InitEl
 		for (int32 SubIdx = 0; SubIdx < 3; ++SubIdx)
 		{
 			int32 VID = Tri[SubIdx];
-			int32 EID = bIsCompactV ? VID : VIDtoEID[VID];
+			int32 EID = bNeedsRemap ? VIDtoEID[VID] : VID;
 			ElementTriangles.InsertAt(EID, Start + SubIdx, FDynamicMesh3::InvalidID);
 			ElementsRefCounts.Increment(EID);
 		}
diff --git a/Engine/Source/Runtime/GeometryCore/Public/DynamicMesh/DynamicMeshOverlay.h b/Engine/Source/Runtime/GeometryCore/Public/DynamicMesh/DynamicMeshOverlay.h
--- a/Engine/Source/Runtime/GeometryCore/Public/DynamicMesh/DynamicMeshOverlay.h
+++ b/Engine/Source/Runtime/GeometryCore/Public/DynamicMesh/DynamicMeshOverlay.h
@@ -382,8 +382,9 @@ public:
 	/**
 	 * Build overlay topology with one element per vertex.
 	 * Note: Faster-but-equivalent-to calling CreateFromPredicate() with a CanShareVertex predicate that always returns true
+	 * @param bExactMapIDs If true, ElementIDs will be constructed to exactly match the corresponding vertex IDs, even if the mesh is non-compact.
 	 */
-	void CreatePerVertex(RealType InitElementValue);
+	void CreatePerVertex(RealType InitElementValue, bool bExactMapIDs = false);
 
 	/**
 	 * Refine an existing overlay topology.  For any element on a given triangle, if the predicate returns true, it gets topologically split out so it isn't shared by any other triangle.
