Geometry: Added parallel argument for tangent computation options. This is so to avoid nested ParallelFor calls.

#Synced-CL 51764840

diff --git a/Engine/Source/Runtime/GeometryCore/Private/DynamicMesh/MeshTangents.cpp b/Engine/Source/Runtime/GeometryCore/Private/DynamicMesh/MeshTangents.cpp
--- a/Engine/Source/Runtime/GeometryCore/Private/DynamicMesh/MeshTangents.cpp
+++ b/Engine/Source/Runtime/GeometryCore/Private/DynamicMesh/MeshTangents.cpp
@@ -351,6 +351,8 @@ void TMeshTangents<RealType>::ComputeSeparatePerTriangleTangents(const FDynamicM
 	int32 MaxTriangleID = Mesh->MaxTriangleID();
 	InitializeTriVertexTangents(false);
 
+	const EParallelForFlags ParallelForFlags = Options.bParallel ? EParallelForFlags::None : EParallelForFlags::ForceSingleThread;
+
 	// compute per-triangle tangent and bitangent
 	ParallelFor(MaxTriangleID, [&](int32 TriangleID)
 	{
@@ -385,8 +387,7 @@ void TMeshTangents<RealType>::ComputeSeparatePerTriangleTangents(const FDynamicM
 				Normalized((TVector<RealType>)ReconsBitangent) );
 		}
 
-	} /*, EParallelForFlags::ForceSingleThread */);
-
+	}, ParallelForFlags);
 }
 
 
@@ -412,6 +413,8 @@ void TMeshTangents<RealType>::ComputeMikkTStyleTangents(const FDynamicMeshNormal
 	TArray<int32> DegenerateTris;
 	FCriticalSection DegenerateTrisLock;
 
+	const EParallelForFlags ParallelForFlags = Options.bParallel ? EParallelForFlags::None : EParallelForFlags::ForceSingleThread;
+
 	// similar process to MikkTSpace:
 	// http://image.diku.dk/projects/media/morten.mikkelsen.08.pdf
 	//  - todo: currently splitting any triangle with large-angle from accumulated group. Should combine these split triangles into a new group.
@@ -448,7 +451,7 @@ void TMeshTangents<RealType>::ComputeMikkTStyleTangents(const FDynamicMeshNormal
 			return;
 		}
 
-	} /*, EParallelForFlags::ForceSingleThread */ );
+	}, ParallelForFlags);
 
 	// Average triangles tangents across vertex "wedges", ie combine one-ring
 	// triangles that share same UV and Normal, for each triangle-vertex
@@ -594,7 +597,7 @@ void TMeshTangents<RealType>::ComputeMikkTStyleTangents(const FDynamicMeshNormal
 			//	UE_LOG(LogTemp, Warning, TEXT("Recons Bitangent is %f %f %f"), BTResult.X, BTResult.Y, BTResult.Z);
 			//}
 		}
-	} /*, EParallelForFlags::ForceSingleThread */ );
+	}, ParallelForFlags);
 
 
 	AllDegenerateTris = DegenerateTris;
diff --git a/Engine/Source/Runtime/GeometryCore/Public/DynamicMesh/MeshTangents.h b/Engine/Source/Runtime/GeometryCore/Public/DynamicMesh/MeshTangents.h
--- a/Engine/Source/Runtime/GeometryCore/Public/DynamicMesh/MeshTangents.h
+++ b/Engine/Source/Runtime/GeometryCore/Public/DynamicMesh/MeshTangents.h
@@ -64,6 +64,8 @@ struct FComputeTangentsOptions
 	bool bAveraged = true;		
 	/** if true, per-face tangents are weighted by opening angle at triangle-vertices when blending */
 	bool bAngleWeighted = true;
+	/** if true, the computation will be done in parallel (via ParallelFor) */
+	bool bParallel = true;
 };
 
 
@@ -142,7 +144,7 @@ public:
 	/**
 	 * Calculate per-triangle tangent spaces based on the given per-triangle normal and UV overlays.
 	 * In this mode there is no averaging of tangents across triangles. So if we have N triangles
-	 * in the mesh, then 3*N tangents are generated. These tangents are computed in parallel.
+	 * in the mesh, then 3*N tangents are generated. These tangents are optionally computed in parallel.
 	 */
 	GEOMETRYCORE_API void ComputeTriVertexTangents(const FDynamicMeshNormalOverlay* NormalOverlay, const FDynamicMeshUVOverlay* UVOverlay, const FComputeTangentsOptions& Options);
 
@@ -243,7 +245,7 @@ public:
 	}
 
 	/**
-	 * Compute per-triangle tangents for the given UV Overlay 
+	 * Compute per-triangle tangents for the given UV Overlay. Computation is done in parallel.
 	 * @param bOrthogonalize if true, we orthogonalize the tangents by computing the Bitangent sign and regenerating it. Otherwise the tangents are left non-orthogonal.
 	 */
 	GEOMETRYCORE_API void ComputeTriangleTangents(const FDynamicMeshUVOverlay* UVOverlay, bool bOrthogonalize = true);
