Add vertex and triangle selection version of FDynamicMesh3's bounding box compute method (with similar logic to compute in parallel chunks)

#Synced-CL 52079281

diff --git a/Engine/Plugins/Runtime/MeshModelingToolset/Source/MeshModelingTools/Private/Sculpting/MeshSculptUtil.cpp b/Engine/Plugins/Runtime/MeshModelingToolset/Source/MeshModelingTools/Private/Sculpting/MeshSculptUtil.cpp
--- a/Engine/Plugins/Runtime/MeshModelingToolset/Source/MeshModelingTools/Private/Sculpting/MeshSculptUtil.cpp
+++ b/Engine/Plugins/Runtime/MeshModelingToolset/Source/MeshModelingTools/Private/Sculpting/MeshSculptUtil.cpp
@@ -175,31 +175,7 @@ void UE::SculptUtil::RecalculateROINormals(
 
 static FAxisAlignedBox3d ParallelComputeROIBounds(const FDynamicMesh3& Mesh, const TArray<int32>& Triangles)
 {
-	FAxisAlignedBox3d FinalBounds = FAxisAlignedBox3d::Empty();
-	FCriticalSection FinalBoundsLock;
-	int32 N = Triangles.Num();
-	constexpr int32 BlockSize = 4096;
-	int32 Blocks = (N / BlockSize) + 1;
-	ParallelFor(Blocks, [&](int bi)
-	{
-		FAxisAlignedBox3d BlockBounds = FAxisAlignedBox3d::Empty();
-		for (int32 k = 0; k < BlockSize; ++k)
-		{
-			int32 i = bi * BlockSize + k;
-			if (i < N)
-			{
-				int32 tid = Triangles[i];
-				const FIndex3i& TriV = Mesh.GetTriangleRef(tid);
-				BlockBounds.Contain(Mesh.GetVertexRef(TriV.A));
-				BlockBounds.Contain(Mesh.GetVertexRef(TriV.B));
-				BlockBounds.Contain(Mesh.GetVertexRef(TriV.C));
-			}
-		}
-		FinalBoundsLock.Lock();
-		FinalBounds.Contain(BlockBounds);
-		FinalBoundsLock.Unlock();
-	});
-	return FinalBounds;
+	return Mesh.GetBoundsForTriangleSelection(Triangles);
 }
 
 
diff --git a/Engine/Source/Runtime/GeometryCore/Private/DynamicMesh/DynamicMesh3_Queries.cpp b/Engine/Source/Runtime/GeometryCore/Private/DynamicMesh/DynamicMesh3_Queries.cpp
--- a/Engine/Source/Runtime/GeometryCore/Private/DynamicMesh/DynamicMesh3_Queries.cpp
+++ b/Engine/Source/Runtime/GeometryCore/Private/DynamicMesh/DynamicMesh3_Queries.cpp
@@ -649,6 +649,105 @@ FAxisAlignedBox3d FDynamicMesh3::GetBounds(bool bParallel) const
 	return ResultBounds;
 }
 
+/**
+ * Computes bounding box of selected vertices.
+ */
+FAxisAlignedBox3d FDynamicMesh3::GetBoundsForVertexSelection(TConstArrayView<int32> VertexIDs) const
+{
+	int32 NumVertices = VertexIDs.Num();
+	if (NumVertices == 0)
+	{
+		return FAxisAlignedBox3d::Empty();
+	}
+
+	FAxisAlignedBox3d ResultBounds = FAxisAlignedBox3d::Empty();
+
+	constexpr int32 ChunkSize = 4096*2;
+	constexpr int32 MinCountToThread = ChunkSize * 3;
+	
+	constexpr bool bParallel = true;
+	if (!bParallel || NumVertices < MinCountToThread)
+	{
+		FVector3d MinVec = Vertices[VertexIDs[0]];
+		FVector3d MaxVec = MinVec;
+		for (int Idx = 1; Idx < VertexIDs.Num(); ++Idx)
+		{
+			int32 VID = VertexIDs[Idx];
+			MinVec = Min(MinVec, Vertices[VID]);
+			MaxVec = Max(MaxVec, Vertices[VID]);
+		}
+		return FAxisAlignedBox3d(MinVec, MaxVec);
+	}
+	else
+	{
+		int32 NumChunks = (NumVertices / ChunkSize) + 1;
+		FCriticalSection BoundsLock;
+		ParallelFor(NumChunks, [&](int ci)
+		{
+			FAxisAlignedBox3d ChunkBounds = FAxisAlignedBox3d::Empty();
+			int32 Start = ci * ChunkSize;
+			int32 End = FMath::Min(Start + ChunkSize, NumVertices);
+			for (int Idx = Start; Idx < End; ++Idx)
+			{
+				int32 VID = VertexIDs[Idx];
+				ChunkBounds.Contain(Vertices[VID]);
+			}
+			BoundsLock.Lock();
+			ResultBounds.Contain(ChunkBounds);
+			BoundsLock.Unlock();
+		});
+	}
+	return ResultBounds;
+}
+
+FAxisAlignedBox3d FDynamicMesh3::GetBoundsForTriangleSelection(TConstArrayView<int32> TriangleIDs) const
+{
+	int32 NumTriangles = TriangleIDs.Num();
+	if (NumTriangles == 0)
+	{
+		return FAxisAlignedBox3d::Empty();
+	}
+
+	auto UpdateBoundsWithTriangles = [this](FAxisAlignedBox3d& Bounds, TConstArrayView<int32> IndexArray)
+	{
+		for (int32 TID : IndexArray)
+		{
+			const FIndex3i& Tri = Triangles[TID];
+			Bounds.Contain(Vertices[Tri.A]);
+			Bounds.Contain(Vertices[Tri.B]);
+			Bounds.Contain(Vertices[Tri.C]);
+		}
+	};
+
+	FAxisAlignedBox3d ResultBounds = FAxisAlignedBox3d::Empty();
+
+	constexpr int32 ChunkSize = 4096;
+	constexpr int32 MinCountToThread = ChunkSize * 3;
+
+	constexpr bool bParallel = true;
+	if (!bParallel || NumTriangles < MinCountToThread)
+	{
+		UpdateBoundsWithTriangles(ResultBounds, TriangleIDs);
+	}
+	else
+	{
+		int32 NumChunks = (NumTriangles / ChunkSize) + 1;
+		FCriticalSection BoundsLock;
+		ParallelFor(NumChunks, [&](int ci)
+		{
+			FAxisAlignedBox3d ChunkBounds = FAxisAlignedBox3d::Empty();
+			int32 Start = ci * ChunkSize;
+			int32 End = FMath::Min(TriangleIDs.Num(), Start + ChunkSize);
+			TConstArrayView<int32> IndexArrayView(TriangleIDs.GetData() + Start, End - Start);
+			UpdateBoundsWithTriangles(ChunkBounds, IndexArrayView);
+			BoundsLock.Lock();
+			ResultBounds.Contain(ChunkBounds);
+			BoundsLock.Unlock();
+		});
+	}
+	return ResultBounds;
+}
+
 
 bool FDynamicMesh3::IsClosed() const
 {
diff --git a/Engine/Source/Runtime/GeometryCore/Public/DynamicMesh/DynamicMesh3.h b/Engine/Source/Runtime/GeometryCore/Public/DynamicMesh/DynamicMesh3.h
--- a/Engine/Source/Runtime/GeometryCore/Public/DynamicMesh/DynamicMesh3.h
+++ b/Engine/Source/Runtime/GeometryCore/Public/DynamicMesh/DynamicMesh3.h
@@ -1190,6 +1190,12 @@ public:
 	/** Returns bounding box of all mesh vertices (including unreferenced vertices) */
 	GEOMETRYCORE_API FAxisAlignedBox3d GetBounds(bool bParallel = false) const;
 
+	/** Returns bounding box of all selected mesh vertices. Will use a chunked parallel implementation for larger selections. */
+	GEOMETRYCORE_API FAxisAlignedBox3d GetBoundsForVertexSelection(TConstArrayView<int32> VertexIDs) const;
+
+	/** Returns bounding box of all selected mesh triangles. Will use a chunked parallel implementation for larger selections. */
+	GEOMETRYCORE_API FAxisAlignedBox3d GetBoundsForTriangleSelection(TConstArrayView<int32> TriangleIDs) const;
+
 	/**
 	 * Compute a normal/tangent frame at vertex that is "stable" as long as
 	 * the mesh topology doesn't change, meaning that one axis of the frame
