Fix issues with sparse dynamic point octree --
- fixed queries hanging on very large query bounds (by refactoring IterateRootsInRadius's logic to clamp query ranges and fall back to all-root iterations into a general root iteration helper)
- fixed points very near cell bounds crashing on insert due to floating point error making CellContains() not agreeing with the cell assignment. Now CellContains directly uses the cell assignment logic.

Also, fix issue w/ octree range iteration methods where we could potentially overflow the large-range test by casting to double too late

#Synced-CL 52243346

diff --git a/Engine/Source/Runtime/GeometryCore/Private/Spatial/SparseDynamicOctree3.cpp b/Engine/Source/Runtime/GeometryCore/Private/Spatial/SparseDynamicOctree3.cpp
--- a/Engine/Source/Runtime/GeometryCore/Private/Spatial/SparseDynamicOctree3.cpp
+++ b/Engine/Source/Runtime/GeometryCore/Private/Spatial/SparseDynamicOctree3.cpp
@@ -349,7 +349,7 @@ TArray<const FSparseOctreeCell*, TInlineAllocator<32>> FSparseDynamicOctree3::In
 		// Double-check that there are enough root cells that we are likely to save time by local iteration,
 		//  since we could have a huge iteration range but only a couple cells allocated in that range.
 		// Use double for the query size here to avoid overflow.
-		FVector3d QuerySize(RootMaxIndex - RootMinIndex + FVector3i(1, 1, 1));
+		FVector3d QuerySize = FVector3d(RootMaxIndex) - FVector3d(RootMinIndex) + FVector3d::One();
 		if (RootCells.GetCount() > QuerySize.X * QuerySize.Y * QuerySize.Z) // There are enough root cells that we could save time by local iteration
 		{
 			if (bCheckOverlapFnInRangeBounds)
diff --git a/Engine/Source/Runtime/GeometryCore/Private/Spatial/SparseDynamicPointOctree3.cpp b/Engine/Source/Runtime/GeometryCore/Private/Spatial/SparseDynamicPointOctree3.cpp
--- a/Engine/Source/Runtime/GeometryCore/Private/Spatial/SparseDynamicPointOctree3.cpp
+++ b/Engine/Source/Runtime/GeometryCore/Private/Spatial/SparseDynamicPointOctree3.cpp
@@ -9,58 +9,12 @@ namespace UE::Geometry
 
 void FSparseDynamicPointOctree3::IterateRootsInRadius(FVector3d QueryPt, double InitialRadius, TFunctionRef<double(double, const FSparsePointOctreeCell&)> RootFn) const
 {
-	FAxisAlignedBox3d Bounds(QueryPt, InitialRadius);
 	double ThresholdSq = InitialRadius * InitialRadius;
+	FAxisAlignedBox3d Bounds(QueryPt, InitialRadius);
 
-	FVector3i RootMinIndex = PointToIndex(0, Bounds.Min);
-	FVector3i RootMaxIndex = PointToIndex(0, Bounds.Max);
-	
-	
-	// Clamp our iteration to only the range of allocated cells. Do it only in the appropriate
-	//  direction for the max/min (e.g. don't clamp min to be <= allocated max, to make it easier
-	//  to detect when we were entirely out of allocated bounds).
-	FAxisAlignedBox3i IndexBounds = RootCells.GetBoundsInclusive();
-	RootMinIndex = Max(RootMinIndex, IndexBounds.Min);
-	RootMaxIndex = Min(RootMaxIndex, IndexBounds.Max);
-	
-	// See if the range bounds were outside of the allocated bounds entirely. This check isn't
-	//  strictly necessary since the RangeIteration call below shouldn't do anything if Max < Min,
-	//  but we'll be explicit about this.
-	for (int i = 0; i < 3; ++i)
-	{
-		if (RootMinIndex[i] > RootMaxIndex[i])
-		{
-			return;
-		}
-	}
-	
-	// Double-check that there are enough root cells that we are likely to save time by local iteration,
-	//  since we could have a huge iteration range but only a couple cells allocated in that range.
-	// Use double for the query size here to avoid overflow.
-	FVector3d QuerySize(RootMaxIndex - RootMinIndex + FVector3i(1, 1, 1));
-	if (RootCells.GetCount() > QuerySize.X * QuerySize.Y * QuerySize.Z) // There are enough root cells that we could save time by local iteration
-	{
-		RootCells.RangeIteration(RootMinIndex, RootMaxIndex, [this, QueryPt, &ThresholdSq, &RootFn](uint32 RootCellID)
-		{
-			const FSparsePointOctreeCell* RootCell = &Cells[RootCellID];
-			if (GetCellBox(*RootCell).DistanceSquared(QueryPt) < ThresholdSq)
-			{
-				ThresholdSq = RootFn(ThresholdSq, *RootCell);
-			}
-		});
-	}
-	else
-	{
-		RootCells.AllocatedIteration([this, QueryPt, &ThresholdSq, &RootFn](const uint32* RootCellID)
-		{
-			const FSparsePointOctreeCell* RootCell = &Cells[*RootCellID];
-			if (GetCellBox(*RootCell).DistanceSquared(QueryPt) < ThresholdSq)
-			{
-				ThresholdSq = RootFn(ThresholdSq, *RootCell);
-			}
-		});
-	}
-
+	ForEachRootCellInBounds(Bounds,
+		[this, &QueryPt, &ThresholdSq](const FSparsePointOctreeCell& Cell) { return GetCellBox(Cell).DistanceSquared(QueryPt) < ThresholdSq; },
+		[&ThresholdSq, &RootFn](const FSparsePointOctreeCell& Cell) { ThresholdSq = RootFn(ThresholdSq, Cell); });
 }
 
 void FSparseDynamicPointOctree3::FindKClosestPoints(FVector3d QueryPt, double DistanceThreshold, int32 NumToFind,
diff --git a/Engine/Source/Runtime/GeometryCore/Public/Spatial/SparseDynamicPointOctree3.h b/Engine/Source/Runtime/GeometryCore/Public/Spatial/SparseDynamicPointOctree3.h
--- a/Engine/Source/Runtime/GeometryCore/Public/Spatial/SparseDynamicPointOctree3.h
+++ b/Engine/Source/Runtime/GeometryCore/Public/Spatial/SparseDynamicPointOctree3.h
@@ -397,8 +397,8 @@ protected:
 
 	bool CellContains(const FSparsePointOctreeCell& Cell, const FVector3d& Position) const
 	{
-		FAxisAlignedBox3d CellBox = GetCellBox(Cell);
-		return CellBox.Contains(Position);
+		// Note: Uses PointToIndex for consistency; GetCellBox(Cell).Contains(Position) may disagree due to floating point error
+		return PointToIndex(Cell.Level, Position) == Cell.Index;
 	}
 
 	uint32 GetCellForPoint(int32 PointID) const
@@ -420,9 +420,19 @@ protected:
 	inline double FindNearestRayCellIntersection(const FSparsePointOctreeCell& Cell, const FRay3d& Ray) const;
 
 private:
+	// Helper to iterate over root cells in a query range. 
+	// Calls IsInRangeFn on each candidate cell, and then ProcessCellFn on the cell if IsInRangeFn returns true.
+	template<typename IsInRangeFn, typename ProcessCellFn>
+	void ForEachRootCellInBounds(const FAxisAlignedBox3d& Bounds, const IsInRangeFn& IsInRange, const ProcessCellFn& ProcessCell) const;
+
+
 	// Helper to iterate over the root cells within a radius. Note the RootFn can return a refined search radius squared, which will let the iteration visit fewer root cells.
 	// RootFn is passed the current squared search radius, and returns the (possibly smaller) squared search radius.
 	void IterateRootsInRadius(FVector3d QueryPt, double InitialRadius, TFunctionRef<double(double RadSq, const FSparsePointOctreeCell& Root)> RootFn) const;
+
+	// Helper to gather relevant root cells for a range query
+	template<typename ArrayType>
+	inline void GatherRootCellsForRangeQuery(const FAxisAlignedBox3d& Bounds, ArrayType& OutRootCells) const;
 };
 
 
@@ -459,6 +469,7 @@ void FSparseDynamicPointOctree3::ConfigureFromPointCountEstimate(double MaxBound
 
 bool FSparseDynamicPointOctree3::ContainsPoint(int32 PointID) const
 {
+	checkSlow(PointID >= 0);
 	return (PointID < PointIDToCellMap.Num() && PointIDToCellMap[PointID] != InvalidCellID);
 }
 
@@ -525,6 +536,7 @@ void FSparseDynamicPointOctree3::Insert_NewChildCell(int32 PointID, const FVecto
 
 void FSparseDynamicPointOctree3::Insert_ToCell(int32 PointID, const FVector3d& Position, const FSparsePointOctreeCell& ExistingCell)
 {
+	check(ExistingCell.CellID != InvalidCellID);
 	checkSlow(CellRefCounts.IsValid(ExistingCell.CellID));
 
 	PointIDToCellMap.InsertAt(ExistingCell.CellID, PointID, InvalidCellID);
@@ -532,7 +544,6 @@ void FSparseDynamicPointOctree3::Insert_ToCell(int32 PointID, const FVector3d& P
 	CellPointLists.Insert(ExistingCell.CellID, PointID);
 
 	checkSlow(CellContains(ExistingCell, Position));
-	checkSlow(PointToIndex(ExistingCell.Level, Position) == ExistingCell.Index);
 }
 
 
@@ -973,6 +984,68 @@ int32 FSparseDynamicPointOctree3::FindNearestHitPoint(const FRay3d& Ray,
 
 
 
+template<typename IsInRangeFn, typename ProcessCellFn>
+void FSparseDynamicPointOctree3::ForEachRootCellInBounds(
+	const FAxisAlignedBox3d& Bounds,
+	const IsInRangeFn& IsInRange,
+	const ProcessCellFn& ProcessCell) const
+{
+	FVector3i RootMinIndex = PointToIndex(0, Bounds.Min);
+	FVector3i RootMaxIndex = PointToIndex(0, Bounds.Max);
+
+	// Clamp to the range of allocated root cells
+	FAxisAlignedBox3i IndexBounds = RootCells.GetBoundsInclusive();
+	RootMinIndex = Max(RootMinIndex, IndexBounds.Min);
+	RootMaxIndex = Min(RootMaxIndex, IndexBounds.Max);
+
+	// Early out if query is outside allocated cell bounds
+	for (int i = 0; i < 3; ++i)
+	{
+		if (RootMinIndex[i] > RootMaxIndex[i])
+		{
+			return;
+		}
+	}
+
+	// If the clamped range would iterate more coordinates than there are allocated cells,
+	// fall back to iterating only allocated cells. Use double to avoid int32 overflow.
+	FVector3d QuerySize = FVector3d(RootMaxIndex) - FVector3d(RootMinIndex) + FVector3d::One();
+	if (static_cast<double>(RootCells.GetCount()) < QuerySize.X * QuerySize.Y * QuerySize.Z)
+	{
+		RootCells.AllocatedIteration([this, &IsInRange, &ProcessCell](const uint32* RootCellID)
+		{
+			const FSparsePointOctreeCell& RootCell = Cells[*RootCellID];
+			if (IsInRange(RootCell))
+			{
+				ProcessCell(RootCell);
+			}
+		});
+	}
+	else
+	{
+		RootCells.RangeIteration(RootMinIndex, RootMaxIndex, [this, &IsInRange, &ProcessCell](uint32 RootCellID)
+		{
+			const FSparsePointOctreeCell& RootCell = Cells[RootCellID];
+			if (IsInRange(RootCell))
+			{
+				ProcessCell(RootCell);
+			}
+		});
+	}
+}
+
+
+template<typename ArrayType>
+void FSparseDynamicPointOctree3::GatherRootCellsForRangeQuery(
+	const FAxisAlignedBox3d& Bounds,
+	ArrayType& OutRootCells) const
+{
+	ForEachRootCellInBounds(Bounds,
+		[this, &Bounds](const FSparsePointOctreeCell& Cell) { return GetCellBox(Cell).Intersects(Bounds); },
+		[&OutRootCells](const FSparsePointOctreeCell& Cell) { OutRootCells.Add(&Cell); });
+}
+
+
 void FSparseDynamicPointOctree3::RangeQuery(
 	const FAxisAlignedBox3d& Bounds,
 	TFunctionRef<bool(int)> PredicateFunc,
@@ -985,13 +1058,7 @@ void FSparseDynamicPointOctree3::RangeQuery(
 	Queue.Reset();
 	Queue.Reserve(128);
 
-	FVector3i RootMinIndex = PointToIndex(0, Bounds.Min);
-	FVector3i RootMaxIndex = PointToIndex(0, Bounds.Max);
-	RootCells.RangeIteration(RootMinIndex, RootMaxIndex, [&](uint32 RootCellID)
-	{
-		const FSparsePointOctreeCell* RootCell = &Cells[RootCellID];
-		Queue.Add(RootCell);
-	});
+	GatherRootCellsForRangeQuery(Bounds, Queue);
 
 	while (Queue.Num() > 0)
 	{
@@ -1042,13 +1109,7 @@ void FSparseDynamicPointOctree3::ParallelRangeQuery(
 
 	FCriticalSection OutputLock;
 
-	FVector3i RootMinIndex = PointToIndex(0, Bounds.Min);
-	FVector3i RootMaxIndex = PointToIndex(0, Bounds.Max);
-	RootCells.RangeIteration(RootMinIndex, RootMaxIndex, [&](uint32 RootCellID)
-	{
-		const FSparsePointOctreeCell* RootCell = &Cells[RootCellID];
-		ProcessRootCells.Add(RootCell);
-	});
+	GatherRootCellsForRangeQuery(Bounds, ProcessRootCells);
 
 	ParallelFor(ProcessRootCells.Num(), [&](int idx)
 	{
@@ -1216,8 +1277,7 @@ void FSparseDynamicPointOctree3::CheckValidity(
 			{
 				CheckOrFailF(CellRefCounts.IsValid(CellID));
 				FSparsePointOctreeCell Cell = Cells[CellID];
-				FAxisAlignedBox3d CellBounds = GetCellBox(Cell);
-				CheckOrFailF(CellBounds.Contains(Position));
+				CheckOrFailF(CellContains(Cell, Position));
 				CheckOrFailF(CellPointLists.Contains(CellID, PointID));
 
 				PointsAtLevel[Cell.Level]++;
