add option for sparse point octree ParallelRangeQuery to also use a custom TFunctionRef to filter which bounds are relevant to the query

#Synced-CL 52243346

diff --git a/Engine/Source/Runtime/GeometryCore/Public/Spatial/SparseDynamicPointOctree3.h b/Engine/Source/Runtime/GeometryCore/Public/Spatial/SparseDynamicPointOctree3.h
--- a/Engine/Source/Runtime/GeometryCore/Public/Spatial/SparseDynamicPointOctree3.h
+++ b/Engine/Source/Runtime/GeometryCore/Public/Spatial/SparseDynamicPointOctree3.h
@@ -267,6 +267,12 @@ public:
 		TArray<int>& PointIDsOut,
 		TArray<const FSparsePointOctreeCell*>* TempBuffer = nullptr) const;
 
+	inline void ParallelRangeQuery(const FAxisAlignedBox3d& Bounds,
+		TFunctionRef<bool(const FAxisAlignedBox3d& Bounds)> BoundsPredicateFunc,
+		TFunctionRef<bool(int)> PredicateFunc,
+		TArray<int>& PointIDsOut,
+		TArray<const FSparsePointOctreeCell*>* TempBuffer = nullptr) const;
+
 
 
 	/**
@@ -1100,6 +1106,17 @@ void FSparseDynamicPointOctree3::ParallelRangeQuery(
 	TFunctionRef<bool(int)> PredicateFunc,
 	TArray<int>& PointIDs,
 	TArray<const FSparsePointOctreeCell*>* TempBuffer) const
+{
+	ParallelRangeQuery(Bounds, [](const FAxisAlignedBox3d&) {return true;},
+		PredicateFunc, PointIDs, TempBuffer);
+}
+
+void FSparseDynamicPointOctree3::ParallelRangeQuery(
+	const FAxisAlignedBox3d& Bounds,
+	TFunctionRef<bool(const FAxisAlignedBox3d&)> BoundsPredicateFunc,
+	TFunctionRef<bool(int)> PredicateFunc,
+	TArray<int>& PointIDs,
+	TArray<const FSparsePointOctreeCell*>* TempBuffer) const
 {
 	TArray<const FSparsePointOctreeCell*> InternalBuffer;
 	TArray<const FSparsePointOctreeCell*>& ProcessRootCells =
@@ -1109,7 +1126,13 @@ void FSparseDynamicPointOctree3::ParallelRangeQuery(
 
 	FCriticalSection OutputLock;
 
-	GatherRootCellsForRangeQuery(Bounds, ProcessRootCells);
+	ForEachRootCellInBounds(Bounds,
+		[this, &Bounds, BoundsPredicateFunc](const FSparsePointOctreeCell& Cell) 
+		{
+			FAxisAlignedBox3d CellBox = GetCellBox(Cell);
+			return CellBox.Intersects(Bounds) && BoundsPredicateFunc(CellBox);
+		},
+		[&ProcessRootCells](const FSparsePointOctreeCell& Cell) { ProcessRootCells.Add(&Cell); });
 
 	ParallelFor(ProcessRootCells.Num(), [&](int idx)
 	{
@@ -1137,7 +1160,8 @@ void FSparseDynamicPointOctree3::ParallelRangeQuery(
 				if (CurCell->HasChild(k))
 				{
 					const FSparsePointOctreeCell* ChildCell = &Cells[CurCell->GetChildCellID(k)];
-					if (GetCellBox(*ChildCell).Intersects(Bounds))
+					FAxisAlignedBox3d ChildCellBox = GetCellBox(*ChildCell);
+					if (ChildCellBox.Intersects(Bounds) && BoundsPredicateFunc(ChildCellBox))
 					{
 						Queue.Add(ChildCell);
 					}
