fixed performance issue with MeshSelfUnion on meshes with many open boundaries
by using pointhashgrid to filter potential edge matches in self union by midpoint.
also, now check for edge vertex position matches in either ordering

#Synced-CL 52243365

diff --git a/Engine/Source/Runtime/GeometryCore/Private/Operations/MeshSelfUnion.cpp b/Engine/Source/Runtime/GeometryCore/Private/Operations/MeshSelfUnion.cpp
--- a/Engine/Source/Runtime/GeometryCore/Private/Operations/MeshSelfUnion.cpp
+++ b/Engine/Source/Runtime/GeometryCore/Private/Operations/MeshSelfUnion.cpp
@@ -13,6 +13,7 @@
 #include "Async/ParallelFor.h"
 #include "DynamicMesh/MeshTransforms.h"
 #include "Spatial/SparseDynamicOctree3.h"
+#include "Spatial/PointHashGrid3.h"
 
 #include "Algo/RemoveIf.h"
 
@@ -709,59 +710,74 @@ bool FMeshSelfUnion::MergeEdges(const TArray<int>& CutBoundaryEdges, const TMap<
 		}
 	}
 
-	// collect remaining unmatched edges
-	TArray<int> UnmatchedEdges;
-	for (int EID : CutBoundaryEdges)
-	{
-		if (Mesh->IsEdge(EID) && Mesh->IsBoundaryEdge(EID))
-		{
-			UnmatchedEdges.Add(EID);
-		}
-	}
-
-	// try to greedily match remaining edges within snap tolerance
+	// try to greedily match remaining unmatched edges within snap tolerance
 	double SnapToleranceSq = SnapTolerance * SnapTolerance;
-	for (int Idx = 0; Idx + 1 < UnmatchedEdges.Num(); Idx++)
+	
+	TPointHashGrid3<int32, double> MidpointsHash(FMath::Max(SnapTolerance * 2, FMathd::ZeroTolerance), -1);
+	int32 UnmatchedInHash = 0;
+	for (int32 EID : CutBoundaryEdges)
 	{
-		int EID = UnmatchedEdges[Idx];
 		if (!Mesh->IsEdge(EID) || !Mesh->IsBoundaryEdge(EID))
 		{
 			continue;
 		}
 		FVector3d A, B;
 		Mesh->GetEdgeV(EID, A, B);
-		for (int OtherIdx = Idx + 1; OtherIdx < UnmatchedEdges.Num(); OtherIdx++)
+		FVector3d Midpoint = (A + B) * .5;
+		int32 EIDMatched = INDEX_NONE;
+		FVector3d MatchedMidpoint;
+		if (UnmatchedInHash > 0)
 		{
-			int OtherEID = UnmatchedEdges[OtherIdx];
-			if (!Mesh->IsEdge(OtherEID) || !Mesh->IsBoundaryEdge(OtherEID))
-			{
-				UnmatchedEdges.RemoveAtSwap(OtherIdx, EAllowShrinking::No);
-				OtherIdx--;
-				continue;
-			}
-			FVector3d OA, OB;
-			Mesh->GetEdgeV(OtherEID, OA, OB);
-
-			if (DistanceSquared(OA, A) < SnapToleranceSq && DistanceSquared(OB, B) < SnapToleranceSq)
-			{
-				FDynamicMesh3::FMergeEdgesInfo MergeInfo;
-				EMeshResult EdgeMergeResult = Mesh->MergeEdges(EID, OtherEID, MergeInfo);
-				if (EdgeMergeResult == EMeshResult::Ok)
+			MidpointsHash.EnumeratePointsInBall(Midpoint, SnapTolerance,
+				[this, &Midpoint](const int32& OtherEID) -> double
+				{
+					return FVector3d::DistSquared(Midpoint, Mesh->GetEdgePoint(OtherEID, .5));
+				},
+				[this, &A, &B, EID, SnapToleranceSq, &EIDMatched, &MatchedMidpoint](int32 OtherEID, double MidSqDist) -> bool
 				{
-					UnmatchedEdges.RemoveAtSwap(OtherIdx, EAllowShrinking::No);
-					if (bTrackAllNewEdges)
+					FVector3d OA, OB;
+					Mesh->GetEdgeV(OtherEID, OA, OB);
+
+					if ((DistanceSquared(OA, A) < SnapToleranceSq && DistanceSquared(OB, B) < SnapToleranceSq) ||
+						(DistanceSquared(OB, A) < SnapToleranceSq && DistanceSquared(OA, B) < SnapToleranceSq))
 					{
-						AllNewEdges.Add(EID);
+						FDynamicMesh3::FMergeEdgesInfo MergeInfo;
+						EMeshResult EdgeMergeResult = Mesh->MergeEdges(EID, OtherEID, MergeInfo);
+						if (EdgeMergeResult == EMeshResult::Ok)
+						{
+							EIDMatched = OtherEID;
+							MatchedMidpoint = (OA + OB) * .5;
+							if (bTrackAllNewEdges)
+							{
+								AllNewEdges.Add(EID);
+							}
+							return false; // stop iteration on successful merge
+						}
 					}
-					break;
+					return true; // continue iterating
+				},
+				[this](int32 OtherEID) -> bool
+				{
+					return !Mesh->IsEdge(OtherEID) || !Mesh->IsBoundaryEdge(OtherEID);
 				}
-			}
+			);
+		}
+
+		if (EIDMatched == INDEX_NONE)
+		{
+			MidpointsHash.InsertPointUnsafe(EID, Midpoint);
+			UnmatchedInHash++;
+		}
+		else
+		{
+			MidpointsHash.RemovePointUnsafe(EIDMatched, MatchedMidpoint);
+			UnmatchedInHash--;
 		}
 	}
 
 	// store the failure cases
 	bool bAllMatched = true;
-	for (int EID : UnmatchedEdges)
+	for (int EID : CutBoundaryEdges)
 	{
 		if (Mesh->IsEdge(EID) && Mesh->IsBoundaryEdge(EID))
 		{
