add point-triangle distances method to vertex connected components

#Synced-CL 52268326

diff --git a/Engine/Source/Runtime/GeometryCore/Public/VertexConnectedComponents.h b/Engine/Source/Runtime/GeometryCore/Public/VertexConnectedComponents.h
--- a/Engine/Source/Runtime/GeometryCore/Public/VertexConnectedComponents.h
+++ b/Engine/Source/Runtime/GeometryCore/Public/VertexConnectedComponents.h
@@ -4,6 +4,7 @@
 
 #include "BoxTypes.h"
 #include "IndexTypes.h"
+#include "Spatial/MeshAABBTree3.h"
 #include "Spatial/PointHashGrid3.h"
 #include "Util/SizedDisjointSet.h"
 
@@ -66,6 +67,55 @@ public:
 			DisjointSet.Union(Triangle[1], Triangle[2]);
 		}
 	}
+	
+	/**
+	 * Union (filtered) vertices to nearest different component w/in threshold distance --
+	 * Intended to be called after ConnectTriangles and other local Connect methods, to merge nearby-but-disconnected components.
+	 * 
+	 * @param MeshSpatial AABB tree for the mesh, used for spatial triangle queries
+	 * @param DistanceThreshold Maximum point-to-triangle distance for bridging components
+	 * @param VertexIDFilter Return true if the vertex should be considered for connection
+	 * @param KeepSizeThreshold Skip considering vertices that are in a smaller cluster size than this. Note: A value > 1 after calling ConnectTriangles will ensure the unreferenced vertices are skipped.
+	 */
+	template<typename TriangleMeshType>
+	void ConnectVerticesToNearestDifferentComponent(const TMeshAABBTree3<TriangleMeshType>& MeshSpatial, double DistanceThreshold, TFunctionRef<bool(int32)> VertexIDFilter = [](int32) {return true;}, int32 KeepSizeThreshold = 0)
+	{
+		const TriangleMeshType* Mesh = MeshSpatial.GetMesh();
+
+		for (int32 VID = 0; VID < Mesh->MaxVertexID(); VID++)
+		{
+			if (!Mesh->IsVertex(VID) || !VertexIDFilter(VID))
+			{
+				continue;
+			}
+			int32 SetID = DisjointSet.Find(VID);
+			if (KeepSizeThreshold > 0 && DisjointSet.Sizes[SetID] < KeepSizeThreshold)
+			{
+				continue;
+			}
+
+			FVector3d Pt = Mesh->GetVertex(VID);
+			double NearestDistSqr;
+			IMeshSpatial::FQueryOptions Options;
+			Options.MaxDistance = DistanceThreshold;
+			Options.TriangleFilterF = [this, Mesh, SetID](int32 TriID) -> bool
+			{
+				FIndex3i Triangle = Mesh->GetTriangle(TriID);
+				// Accept only triangles where at least one vertex is in a different component
+				return DisjointSet.Find(Triangle[0]) != SetID
+					|| DisjointSet.Find(Triangle[1]) != SetID
+					|| DisjointSet.Find(Triangle[2]) != SetID;
+			};
+			int32 NearestTriID = MeshSpatial.FindNearestTriangle(Pt, NearestDistSqr, Options);
+			if (NearestTriID >= 0)
+			{
+				FIndex3i Triangle = Mesh->GetTriangle(NearestTriID);
+				DisjointSet.Union(VID, Triangle[0]);
+				DisjointSet.Union(VID, Triangle[1]);
+				DisjointSet.Union(VID, Triangle[2]);
+			}
+		}
+	}
 
 	/**
 	 * Connect a filtered subset of vertices that are closer than a threshold distance in space
