minor improvement to mesh connected components methods -- For set versions, don't redundantly add elements that we already added to the output set. For array result version, do reserve/append at the end instead of during.

#Synced-CL 52609914

diff --git a/Engine/Source/Runtime/GeometryCore/Private/Selections/MeshConnectedComponents.cpp b/Engine/Source/Runtime/GeometryCore/Private/Selections/MeshConnectedComponents.cpp
--- a/Engine/Source/Runtime/GeometryCore/Private/Selections/MeshConnectedComponents.cpp
+++ b/Engine/Source/Runtime/GeometryCore/Private/Selections/MeshConnectedComponents.cpp
@@ -680,7 +680,7 @@ void FMeshConnectedComponents::GrowToConnectedTriangles(const FDynamicMesh3* Mes
 	TArray<int32> LocalQueue;
 	QueueBuffer = (QueueBuffer == nullptr) ? &LocalQueue : QueueBuffer;
 	QueueBuffer->Reset(); 
-	QueueBuffer->Insert(InputROI, 0);
+	QueueBuffer->Append(InputROI);
 
 	TSet<int32> LocalDone;
 	DoneBuffer = (DoneBuffer == nullptr) ? &LocalDone : DoneBuffer;
@@ -690,7 +690,6 @@ void FMeshConnectedComponents::GrowToConnectedTriangles(const FDynamicMesh3* Mes
 	while (QueueBuffer->Num() > 0)
 	{
 		int32 CurTri = QueueBuffer->Pop(EAllowShrinking::No);
-		ResultROI.Add(CurTri);
 
 		FIndex3i NbrTris = Mesh->GetTriNeighbourTris(CurTri);
 		for (int j = 0; j < 3; ++j)
@@ -703,6 +702,11 @@ void FMeshConnectedComponents::GrowToConnectedTriangles(const FDynamicMesh3* Mes
 			}
 		}
 	}
+	ResultROI.Reserve(ResultROI.Num() + DoneBuffer->Num());
+	for (int32 TID : *DoneBuffer)
+	{
+		ResultROI.Add(TID);
+	}
 }
 
 
@@ -716,7 +720,7 @@ void FMeshConnectedComponents::GrowToConnectedTriangles(const FDynamicMesh3* Mes
 	TArray<int32> LocalQueue;
 	QueueBuffer = (QueueBuffer == nullptr) ? &LocalQueue : QueueBuffer;
 	QueueBuffer->Reset();
-	QueueBuffer->Insert(InputROI, 0);
+	QueueBuffer->Append(InputROI);
 
 	ResultROI.Reset();
 	ResultROI.Append(InputROI);
@@ -724,7 +728,6 @@ void FMeshConnectedComponents::GrowToConnectedTriangles(const FDynamicMesh3* Mes
 	while (QueueBuffer->Num() > 0)
 	{
 		int32 CurTri = QueueBuffer->Pop(EAllowShrinking::No);
-		ResultROI.Add(CurTri);
 
 		FIndex3i NbrTris = Mesh->GetTriNeighbourTris(CurTri);
 		for (int j = 0; j < 3; ++j)
@@ -749,7 +752,7 @@ void FMeshConnectedComponents::GrowToConnectedVertices(const FDynamicMesh3& Mesh
 	TArray<int32> LocalQueue;
 	QueueBuffer = (QueueBuffer == nullptr) ? &LocalQueue : QueueBuffer;
 	QueueBuffer->Reset();
-	QueueBuffer->Insert(InputROI, 0);
+	QueueBuffer->Append(InputROI);
 
 	ResultROI.Reset();
 	ResultROI.Append(InputROI);
@@ -757,7 +760,6 @@ void FMeshConnectedComponents::GrowToConnectedVertices(const FDynamicMesh3& Mesh
 	while (QueueBuffer->Num() > 0)
 	{
 		int32 CurVertexID = QueueBuffer->Pop(EAllowShrinking::No);
-		ResultROI.Add(CurVertexID);
 
 		Mesh.EnumerateVertexVertices(CurVertexID, [&](int32 NbrVertexID)
 		{
@@ -781,7 +783,7 @@ void FMeshConnectedComponents::GrowToConnectedEdges(const FDynamicMesh3& Mesh,
 	TArray<int32> LocalQueue;
 	QueueBuffer = (QueueBuffer == nullptr) ? &LocalQueue : QueueBuffer;
 	QueueBuffer->Reset();
-	QueueBuffer->Insert(InputROI, 0);
+	QueueBuffer->Append(InputROI);
 
 	ResultROI.Reset();
 	ResultROI.Append(InputROI);
@@ -789,7 +791,6 @@ void FMeshConnectedComponents::GrowToConnectedEdges(const FDynamicMesh3& Mesh,
 	while (QueueBuffer->Num() > 0)
 	{
 		int32 CurEdgeID = QueueBuffer->Pop(EAllowShrinking::No);
-		ResultROI.Add(CurEdgeID);
 
 		FIndex2i EdgeV = Mesh.GetEdgeV(CurEdgeID);
 		for (int32 k = 0; k < 2; ++k)
